[Teds] Rename TedsAssetData files and classes to remove TypedElement strings

#rnx

diff --git a/Engine/Plugins/Experimental/EditorDataStorageFeatures/Source/TedsAssetData/Private/Debug/AssetProcessors.cpp b/Engine/Plugins/Experimental/EditorDataStorageFeatures/Source/TedsAssetData/Private/Debug/AssetProcessors.cpp
deleted file mode 100644
diff --git a/Engine/Plugins/Experimental/EditorDataStorageFeatures/Source/TedsAssetData/Private/Debug/AssetProcessors.h b/Engine/Plugins/Experimental/EditorDataStorageFeatures/Source/TedsAssetData/Private/Debug/AssetProcessors.h
deleted file mode 100644
diff --git a/Engine/Plugins/Experimental/EditorDataStorageFeatures/Source/TedsAssetData/Private/Debug/TedsAssetDataDebugCommands.cpp b/Engine/Plugins/Experimental/EditorDataStorageFeatures/Source/TedsAssetData/Private/Debug/TedsAssetDataDebugCommands.cpp
--- a/Engine/Plugins/Experimental/EditorDataStorageFeatures/Source/TedsAssetData/Private/Debug/TedsAssetDataDebugCommands.cpp
+++ b/Engine/Plugins/Experimental/EditorDataStorageFeatures/Source/TedsAssetData/Private/Debug/TedsAssetDataDebugCommands.cpp
@@ -23,7 +23,7 @@
 
 DEFINE_LOG_CATEGORY_STATIC(LogTEDSAssetRegistry, Log, All)
 
-namespace UE::TEDSAssetDataDebugCommands::Local
+namespace UE::EditorDataStorage::Debug::Private
 {
 	static FAutoConsoleCommand CCMDTestFolderRowData(
 		TEXT("TEDS.Debug.ShowDataOfAssetFolder"),
diff --git a/Engine/Plugins/Experimental/EditorDataStorageFeatures/Source/TedsAssetData/Private/Processors/AssetProcessors.cpp b/Engine/Plugins/Experimental/EditorDataStorageFeatures/Source/TedsAssetData/Private/Processors/AssetProcessors.cpp
new file mode 100644
--- /dev/null
+++ b/Engine/Plugins/Experimental/EditorDataStorageFeatures/Source/TedsAssetData/Private/Processors/AssetProcessors.cpp
@@ -0,0 +1,17 @@
+// Copyright Epic Games, Inc. All Rights Reserved.
+
+#include "Processors/AssetProcessors.h"
+
+#include "Elements/Columns/TypedElementLabelColumns.h"
+#include "Elements/Columns/TypedElementTypeInfoColumns.h"
+#include "Elements/Interfaces/TypedElementQueryStorageInterfaces.h"
+#include "TedsAssetDataColumns.h"
+
+void UTedsAssetDataFactory::RegisterTables(ITypedElementDataStorageInterface& DataStorage) 
+{
+	DataStorage.RegisterTable(
+		TTypedElementColumnTypeList<
+			FTypedElementLabelColumn, FTypedElementClassTypeInfoColumn, FAssetPathColumn_Experimental,
+			FAssetTag, FDiskSizeColumn, FVersePathColumn>(),
+		FName("Editor_PlaceholderAssetTable"));
+}
diff --git a/Engine/Plugins/Experimental/EditorDataStorageFeatures/Source/TedsAssetData/Private/Processors/AssetProcessors.h b/Engine/Plugins/Experimental/EditorDataStorageFeatures/Source/TedsAssetData/Private/Processors/AssetProcessors.h
new file mode 100644
--- /dev/null
+++ b/Engine/Plugins/Experimental/EditorDataStorageFeatures/Source/TedsAssetData/Private/Processors/AssetProcessors.h
@@ -0,0 +1,22 @@
+// Copyright Epic Games, Inc. All Rights Reserved.
+
+#pragma once
+
+#include "Elements/Interfaces/TypedElementDataStorageFactory.h"
+
+#include "AssetProcessors.generated.h"
+
+class ITypedElementDataStorageInterface;
+
+UCLASS()
+class UTedsAssetDataFactory : public UTypedElementDataStorageFactory
+{
+	GENERATED_BODY()
+
+public:
+	~UTedsAssetDataFactory() override = default;
+
+	void RegisterTables(ITypedElementDataStorageInterface& DataStorage) override;
+
+private:
+};
diff --git a/Engine/Plugins/Experimental/EditorDataStorageFeatures/Source/TedsAssetData/Private/TedsAssetData.cpp b/Engine/Plugins/Experimental/EditorDataStorageFeatures/Source/TedsAssetData/Private/TedsAssetData.cpp
--- a/Engine/Plugins/Experimental/EditorDataStorageFeatures/Source/TedsAssetData/Private/TedsAssetData.cpp
+++ b/Engine/Plugins/Experimental/EditorDataStorageFeatures/Source/TedsAssetData/Private/TedsAssetData.cpp
@@ -16,7 +16,7 @@
 
 #define TRACK_TEDSASSETDATA_MEMORY 0
 
-namespace UE::TypedElementsDataStorageAssetData::Private
+namespace UE::EditorDataStorage::AssetData::Private
 {
 	
 constexpr int32 ParallelForMinBatchSize = 1024 * 4;
diff --git a/Engine/Plugins/Experimental/EditorDataStorageFeatures/Source/TedsAssetData/Private/TedsAssetData.h b/Engine/Plugins/Experimental/EditorDataStorageFeatures/Source/TedsAssetData/Private/TedsAssetData.h
--- a/Engine/Plugins/Experimental/EditorDataStorageFeatures/Source/TedsAssetData/Private/TedsAssetData.h
+++ b/Engine/Plugins/Experimental/EditorDataStorageFeatures/Source/TedsAssetData/Private/TedsAssetData.h
@@ -10,7 +10,7 @@ struct FAssetData;
 
 class ITypedElementDataStorageInterface;
 
-namespace UE::TypedElementsDataStorageAssetData::Private
+namespace UE::EditorDataStorage::AssetData::Private
 {
 
 /**
diff --git a/Engine/Plugins/Experimental/EditorDataStorageFeatures/Source/TedsAssetData/Private/TedsAssetDataModule.cpp b/Engine/Plugins/Experimental/EditorDataStorageFeatures/Source/TedsAssetData/Private/TedsAssetDataModule.cpp
new file mode 100644
--- /dev/null
+++ b/Engine/Plugins/Experimental/EditorDataStorageFeatures/Source/TedsAssetData/Private/TedsAssetDataModule.cpp
@@ -0,0 +1,113 @@
+// Copyright Epic Games, Inc. All Rights Reserved.
+
+#include "TedsAssetDataModule.h"
+
+#include "Elements/Framework/TypedElementRegistry.h"
+#include "HAL/IConsoleManager.h"
+#include "Modules/ModuleManager.h"
+#include "TedsAssetData.h"
+
+IMPLEMENT_MODULE(UE::EditorDataStorage::AssetData::FTedsAssetDataModule, TedsAssetData);
+
+namespace UE::EditorDataStorage::AssetData
+{
+
+namespace Private
+{
+TAutoConsoleVariable<bool> CVarTEDSAssetDataStorage(TEXT("TEDS.AssetDataStorage"), false, TEXT("When true we will activate a wrapper that store the a copy of the asset data including the in memory change from the asset registry into TEDS.")
+	, FConsoleVariableDelegate::CreateLambda([](IConsoleVariable* Variable)
+	{
+		const bool bIsEnabled = Variable->GetBool();
+		FTedsAssetDataModule& Module = FTedsAssetDataModule::GetChecked();
+
+		if (bIsEnabled)
+		{
+			Module.EnableTedsAssetRegistryStorage();
+		}
+		else
+		{
+			Module.DisableTedsAssetRegistryStorage();
+		}
+	}));
+}
+
+void FTedsAssetDataModule::StartupModule()
+{
+	if (Private::CVarTEDSAssetDataStorage.GetValueOnGameThread())
+	{
+		EnableTedsAssetRegistryStorage();
+	}
+}
+
+void FTedsAssetDataModule::ShutdownModule()
+{
+	if (UTypedElementRegistry* TypedElementRegistry = UTypedElementRegistry::GetInstance())
+	{
+		TypedElementRegistry->OnDataStorageInterfacesSet().RemoveAll(this);
+	}
+}
+
+FTedsAssetDataModule* FTedsAssetDataModule::Get()
+{
+	return FModuleManager::Get().LoadModulePtr<FTedsAssetDataModule>(TEXT("TypedElementsDataStorageAssetData"));
+}
+
+FTedsAssetDataModule& FTedsAssetDataModule::GetChecked()
+{
+	return FModuleManager::Get().LoadModuleChecked<FTedsAssetDataModule>(TEXT("TypedElementsDataStorageAssetData"));
+}
+
+void FTedsAssetDataModule::EnableTedsAssetRegistryStorage()
+{
+	if (!AssetRegistryStorage)
+	{
+		FModuleManager::Get().LoadModuleChecked(TEXT("TypedElementFramework"));
+		UTypedElementRegistry* TypedElementRegistry = UTypedElementRegistry::GetInstance();
+		if (TypedElementRegistry->GetMutableDataStorage())
+		{
+			InitAssetRegistryStorage();
+		}
+		else
+		{
+			TypedElementRegistry->OnDataStorageInterfacesSet().AddRaw(this, &FTedsAssetDataModule::InitAssetRegistryStorage);
+		}
+
+		if (!Private::CVarTEDSAssetDataStorage.GetValueOnGameThread())
+		{
+			Private::CVarTEDSAssetDataStorage.AsVariable()->Set(true);
+		}
+	}
+}
+
+void FTedsAssetDataModule::DisableTedsAssetRegistryStorage()
+{
+	if (AssetRegistryStorage)
+	{
+		AssetRegistryStorage.Reset();
+
+		if (Private::CVarTEDSAssetDataStorage.GetValueOnGameThread())
+		{
+			Private::CVarTEDSAssetDataStorage.AsVariable()->Set(false);
+		}
+	}
+}
+
+bool FTedsAssetDataModule::IsTedsAssetRegistryStorageEnabled() const
+{
+	return AssetRegistryStorage.IsValid();
+}
+
+void FTedsAssetDataModule::ProcessDependentEvents()
+{
+	if (Private::FTedsAssetData* Storage = AssetRegistryStorage.Get())
+	{
+		Storage->ProcessAllEvents();
+	}
+}
+
+void FTedsAssetDataModule::InitAssetRegistryStorage()
+{
+	AssetRegistryStorage = MakeUnique<Private::FTedsAssetData>(*UTypedElementRegistry::GetInstance()->GetMutableDataStorage());
+}
+
+}
\ No newline at end of file
diff --git a/Engine/Plugins/Experimental/EditorDataStorageFeatures/Source/TedsAssetData/Private/TedsAssetDataModule.h b/Engine/Plugins/Experimental/EditorDataStorageFeatures/Source/TedsAssetData/Private/TedsAssetDataModule.h
new file mode 100644
--- /dev/null
+++ b/Engine/Plugins/Experimental/EditorDataStorageFeatures/Source/TedsAssetData/Private/TedsAssetDataModule.h
@@ -0,0 +1,42 @@
+// Copyright Epic Games, Inc. All Rights Reserved.
+
+#pragma once
+
+#include "Modules/ModuleInterface.h"
+#include "Templates/UniquePtr.h"
+
+namespace UE::EditorDataStorage::AssetData
+{ 
+
+namespace Private
+{
+	class FTedsAssetData;
+}
+
+class FTedsAssetDataModule : public IModuleInterface
+{
+public:
+	virtual void StartupModule() override;
+	virtual void ShutdownModule() override;
+
+	static FTedsAssetDataModule* Get();
+	static FTedsAssetDataModule& GetChecked();
+
+	void EnableTedsAssetRegistryStorage();
+	void DisableTedsAssetRegistryStorage();
+	bool IsTedsAssetRegistryStorageEnabled() const;
+
+
+	/**
+	 * Process now any pending event that might make the Teds database out of sync with the asset registry.
+	 * Note: This isn't needed when using the editor so it should be only called by the automation scripts that need it to avoid creating some unneeded stalls.
+	 */
+	void ProcessDependentEvents();
+
+private:
+	void InitAssetRegistryStorage();
+
+	TUniquePtr<Private::FTedsAssetData> AssetRegistryStorage;
+};
+
+}
diff --git a/Engine/Plugins/Experimental/EditorDataStorageFeatures/Source/TedsAssetData/Private/Tests/TedsAssetDataTest.cpp b/Engine/Plugins/Experimental/EditorDataStorageFeatures/Source/TedsAssetData/Private/Tests/TedsAssetDataTest.cpp
--- a/Engine/Plugins/Experimental/EditorDataStorageFeatures/Source/TedsAssetData/Private/Tests/TedsAssetDataTest.cpp
+++ b/Engine/Plugins/Experimental/EditorDataStorageFeatures/Source/TedsAssetData/Private/Tests/TedsAssetDataTest.cpp
@@ -13,14 +13,14 @@
 #include "Elements/Interfaces/TypedElementDataStorageInterface.h"
 #include "Misc/ScopeExit.h"
 #include "TedsAssetDataColumns.h"
-#include "TypedElementsDataStorageAssetDataModule.h"
+#include "TedsAssetDataModule.h"
 #include "UObject/NameTypes.h"
 
 IMPLEMENT_SIMPLE_AUTOMATION_TEST(FTedsAssetDataTest, "TEDS.AssetRegistry.ValidateState", EAutomationTestFlags::EditorContext | EAutomationTestFlags::EngineFilter | EAutomationTestFlags::HighPriority)
  
 bool FTedsAssetDataTest::RunTest(const FString& Parameters)
 {
-	UE::TypedElementsDataStorageAssetData::FTypedElementsDataStorageAssetDataModule& TedsAssetRegistry = UE::TypedElementsDataStorageAssetData::FTypedElementsDataStorageAssetDataModule::GetChecked();
+	UE::EditorDataStorage::AssetData::FTedsAssetDataModule& TedsAssetRegistry = UE::EditorDataStorage::AssetData::FTedsAssetDataModule::GetChecked();
 
 	const bool bIsEnabled = TedsAssetRegistry.IsTedsAssetRegistryStorageEnabled();
 	ON_SCOPE_EXIT
diff --git a/Engine/Plugins/Experimental/EditorDataStorageFeatures/Source/TedsAssetData/Private/TypedElementsDataStorageAssetDataModule.cpp b/Engine/Plugins/Experimental/EditorDataStorageFeatures/Source/TedsAssetData/Private/TypedElementsDataStorageAssetDataModule.cpp
deleted file mode 100644
diff --git a/Engine/Plugins/Experimental/EditorDataStorageFeatures/Source/TedsAssetData/Private/TypedElementsDataStorageAssetDataModule.h b/Engine/Plugins/Experimental/EditorDataStorageFeatures/Source/TedsAssetData/Private/TypedElementsDataStorageAssetDataModule.h
deleted file mode 100644
