[Teds] Rename TedsPropertyEditor files and classes to remove TypedElement strings

#rnx

diff --git a/Engine/Plugins/Experimental/EditorDataStorageFeatures/Source/TedsPropertyEditor/Private/TedsRowPickingMode.cpp b/Engine/Plugins/Experimental/EditorDataStorageFeatures/Source/TedsPropertyEditor/Private/TedsRowPickingMode.cpp
new file mode 100644
--- /dev/null
+++ b/Engine/Plugins/Experimental/EditorDataStorageFeatures/Source/TedsPropertyEditor/Private/TedsRowPickingMode.cpp
@@ -0,0 +1,42 @@
+// Copyright Epic Games, Inc. All Rights Reserved.
+
+#include "TedsRowPickingMode.h"
+
+#include "TypedElementOutlinerItem.h"
+
+FTedsRowPickingMode::FTedsRowPickingMode(const FTypedElementOutlinerModeParams& Params, FOnSceneOutlinerItemPicked OnItemPickedDelegate)
+	: FTypedElementOutlinerMode(Params)
+	, OnItemPicked(OnItemPickedDelegate)
+{
+
+}
+
+void FTedsRowPickingMode::OnItemSelectionChanged(FSceneOutlinerTreeItemPtr Item, ESelectInfo::Type SelectionType, const FSceneOutlinerItemSelection& Selection)
+{
+	// In picking mode, we fire off the notification to whoever is listening.
+	// This may often cause the widget itself to be enqueued for destruction
+	auto SelectedItems = SceneOutliner->GetSelectedItems();
+	if (SelectedItems.Num() > 0)
+	{
+		auto FirstItem = SelectedItems[0];
+		if (FirstItem->CanInteract())
+		{
+			OnItemPicked.ExecuteIfBound(FirstItem.ToSharedRef());
+		}
+	}
+}
+
+/** Allow the user to commit their selection by pressing enter if it is valid */
+void FTedsRowPickingMode::OnFilterTextCommited(FSceneOutlinerItemSelection& Selection, ETextCommit::Type CommitType)
+{
+	// In picking mode, we check to see if we have any Typed Element items, and if so, fire
+	// off the notification to whoever is listening. This may often cause the widget itself
+	// to be enqueued for destruction
+	TArray<FTypedElementOutlinerTreeItem*> TypedElements;
+	Selection.Get(TypedElements);
+	if (TypedElements.Num() == 1 && TypedElements[0])
+	{
+		// Signal that a Typed Element was selected. We assume it is valid as it won't have been added to Selection if not.
+		SceneOutliner->SetItemSelection(TypedElements[0]->AsShared(), true, ESelectInfo::OnKeyPress);
+	}
+}
\ No newline at end of file
diff --git a/Engine/Plugins/Experimental/EditorDataStorageFeatures/Source/TedsPropertyEditor/Private/TypedElementPickingMode.cpp b/Engine/Plugins/Experimental/EditorDataStorageFeatures/Source/TedsPropertyEditor/Private/TypedElementPickingMode.cpp
deleted file mode 100644
diff --git a/Engine/Plugins/Experimental/EditorDataStorageFeatures/Source/TedsPropertyEditor/Private/Widgets/SPropertyMenuTedsRowPicker.cpp b/Engine/Plugins/Experimental/EditorDataStorageFeatures/Source/TedsPropertyEditor/Private/Widgets/SPropertyMenuTedsRowPicker.cpp
new file mode 100644
--- /dev/null
+++ b/Engine/Plugins/Experimental/EditorDataStorageFeatures/Source/TedsPropertyEditor/Private/Widgets/SPropertyMenuTedsRowPicker.cpp
@@ -0,0 +1,125 @@
+// Copyright Epic Games, Inc. All Rights Reserved.
+
+#include "Widgets/SPropertyMenuTedsRowPicker.h"
+
+#include "Elements/Framework/TypedElementRegistry.h"
+#include "TedsRowPickingMode.h"
+#include "TypedElementOutlinerItem.h"
+
+#define LOCTEXT_NAMESPACE "TedsPropertyEditor"
+void SPropertyMenuTedsRowPicker::Construct(const FArguments& InArgs)
+{
+	bAllowClear = InArgs._AllowClear;
+	TypedElementQueryFilter = InArgs._TypedElementQueryFilter;
+	ElementFilter = InArgs._ElementFilter;
+	OnSet = InArgs._OnSet;
+
+	FMenuBuilder MenuBuilder(true, nullptr);
+
+	MenuBuilder.BeginSection(NAME_None, LOCTEXT("CurrentTypedElementOperationsHeader", "Current Element"));
+	{
+		if (bAllowClear)
+		{
+			MenuBuilder.AddMenuEntry(
+				LOCTEXT("ClearElement", "Clear"),
+				LOCTEXT("ClearElement_Tooltip", "Clears the item set on this field"),
+				FSlateIcon(),
+				FUIAction(FExecuteAction::CreateSP(this, &SPropertyMenuTedsRowPicker::OnClear))
+			);
+		}
+	}
+	MenuBuilder.EndSection();
+
+	MenuBuilder.BeginSection(NAME_None, LOCTEXT("BrowseHeader", "Browse"));
+	{
+		TSharedPtr<SWidget> MenuContent;
+		{
+			// TEDS-Outliner TODO: Taken from private implementation of PropertyEditorAssetConstants.
+			//                     Should be centralized when TEDS is moved to core
+			static const FVector2D ContentBrowserWindowSize(300.0f, 300.0f);
+			static const FVector2D SceneOutlinerWindowSize(350.0f, 300.0f);
+
+			UTypedElementRegistry* Registry = UTypedElementRegistry::GetInstance();
+			checkf(Registry, TEXT("Unable to initialize the Typed Elements Outliner before TEDS is initialized."));
+
+			if (!Registry->AreDataStorageInterfacesSet())
+			{
+				MenuContent = SNew(STextBlock)
+					.Text(LOCTEXT("TEDSPluginNotEnabledText", "Typed Element Data Storage plugin required to use this property picker."));
+			}
+			else
+			{
+				auto OnItemPicked = FOnSceneOutlinerItemPicked::CreateLambda([&](TSharedRef<ISceneOutlinerTreeItem> Item)
+					{
+						if (FTypedElementOutlinerTreeItem* ElementItem = Item->CastTo<FTypedElementOutlinerTreeItem>())
+						{
+							if (ElementItem->IsValid())
+							{
+								OnSet.ExecuteIfBound(ElementItem->GetRowHandle());
+							}
+						}
+					});
+
+				FSceneOutlinerInitializationOptions InitOptions;
+				InitOptions.bShowHeaderRow = true;
+				InitOptions.bShowTransient = true;
+				InitOptions.bShowSearchBox = false; // Search not currently supported in TEDS Outliner
+
+				InitOptions.ModeFactory = FCreateSceneOutlinerMode::CreateLambda([&](SSceneOutliner* Outliner)
+					{
+						FTypedElementOutlinerModeParams Params(Outliner);
+						Params.QueryDescription = TypedElementQueryFilter;
+						return new FTedsRowPickingMode(Params, OnItemPicked);
+					});
+
+				TSharedPtr<SSceneOutliner> Outliner = SNew(SSceneOutliner, InitOptions);
+
+				Outliner->AddFilter(
+					MakeShared<TSceneOutlinerPredicateFilter<FTypedElementOutlinerTreeItem>>(
+						FTypedElementOutlinerTreeItem::FFilterPredicate::CreateLambda([this](const TypedElementDataStorage::RowHandle RowHandle) -> bool
+					{
+						if (ElementFilter.IsBound())
+						{
+							return ElementFilter.Execute(RowHandle);
+						}
+						return true;
+					}), FSceneOutlinerFilter::EDefaultBehaviour::Pass));
+
+				MenuContent =
+					SNew(SBox)
+					.WidthOverride(static_cast<float>(SceneOutlinerWindowSize.X))
+					.HeightOverride(static_cast<float>(SceneOutlinerWindowSize.Y))
+					[
+						Outliner.ToSharedRef()
+					];
+			}
+		}
+
+		MenuBuilder.AddWidget(MenuContent.ToSharedRef(), FText::GetEmpty(), true);
+	}
+	MenuBuilder.EndSection();
+
+	ChildSlot
+	[
+		MenuBuilder.MakeWidget()
+	];
+}
+
+void SPropertyMenuTedsRowPicker::OnClear()
+{
+	SetValue(TypedElementInvalidRowHandle);
+	OnClose.ExecuteIfBound();
+}
+
+void SPropertyMenuTedsRowPicker::OnElementSelected(TypedElementDataStorage::RowHandle RowHandle)
+{
+	SetValue(RowHandle);
+	OnClose.ExecuteIfBound();
+}
+
+void SPropertyMenuTedsRowPicker::SetValue(TypedElementDataStorage::RowHandle RowHandle)
+{
+	OnSet.ExecuteIfBound(RowHandle);
+}
+
+#undef LOCTEXT_NAMESPACE // "TedsPropertyEditor"
diff --git a/Engine/Plugins/Experimental/EditorDataStorageFeatures/Source/TedsPropertyEditor/Private/Widgets/SPropertyMenuTypedElementPicker.cpp b/Engine/Plugins/Experimental/EditorDataStorageFeatures/Source/TedsPropertyEditor/Private/Widgets/SPropertyMenuTypedElementPicker.cpp
deleted file mode 100644
diff --git a/Engine/Plugins/Experimental/EditorDataStorageFeatures/Source/TedsPropertyEditor/Public/TedsRowPickingMode.h b/Engine/Plugins/Experimental/EditorDataStorageFeatures/Source/TedsPropertyEditor/Public/TedsRowPickingMode.h
new file mode 100644
--- /dev/null
+++ b/Engine/Plugins/Experimental/EditorDataStorageFeatures/Source/TedsPropertyEditor/Public/TedsRowPickingMode.h
@@ -0,0 +1,28 @@
+// Copyright Epic Games, Inc. All Rights Reserved.
+
+#pragma once
+
+#include "TypedElementOutlinerMode.h"
+
+DECLARE_DELEGATE_OneParam(FOnTedsRowSelected, TypedElementDataStorage::RowHandle RowHandle);
+DECLARE_DELEGATE_RetVal_OneParam(bool, FOnShouldFilterTedsRow, const TypedElementDataStorage::RowHandle RowHandle);
+
+/*
+* Picking mode for TEDs Scene Outliner Widgets. Based off of FActorPickingMode
+*/
+class TEDSPROPERTYEDITOR_API FTedsRowPickingMode : public FTypedElementOutlinerMode
+{
+public:
+	FTedsRowPickingMode(const FTypedElementOutlinerModeParams& Params, FOnSceneOutlinerItemPicked OnItemPickedDelegate);
+
+	virtual ~FTedsRowPickingMode() = default;
+public:
+	virtual void OnItemSelectionChanged(FSceneOutlinerTreeItemPtr Item, ESelectInfo::Type SelectionType, const FSceneOutlinerItemSelection& Selection) override;
+
+	/** Allow the user to commit their selection by pressing enter if it is valid */
+	virtual void OnFilterTextCommited(FSceneOutlinerItemSelection& Selection, ETextCommit::Type CommitType) override;
+
+	virtual bool ShowViewButton() const override { return false; }
+private:
+	FOnSceneOutlinerItemPicked OnItemPicked;
+};
\ No newline at end of file
diff --git a/Engine/Plugins/Experimental/EditorDataStorageFeatures/Source/TedsPropertyEditor/Public/TypedElementPickingMode.h b/Engine/Plugins/Experimental/EditorDataStorageFeatures/Source/TedsPropertyEditor/Public/TypedElementPickingMode.h
deleted file mode 100644
diff --git a/Engine/Plugins/Experimental/EditorDataStorageFeatures/Source/TedsPropertyEditor/Public/Widgets/SPropertyMenuTedsRowPicker.h b/Engine/Plugins/Experimental/EditorDataStorageFeatures/Source/TedsPropertyEditor/Public/Widgets/SPropertyMenuTedsRowPicker.h
new file mode 100644
--- /dev/null
+++ b/Engine/Plugins/Experimental/EditorDataStorageFeatures/Source/TedsPropertyEditor/Public/Widgets/SPropertyMenuTedsRowPicker.h
@@ -0,0 +1,44 @@
+// Copyright Epic Games, Inc. All Rights Reserved.
+
+#pragma once
+
+#include "TedsRowPickingMode.h"
+#include "Widgets/SCompoundWidget.h"
+
+class TEDSPROPERTYEDITOR_API SPropertyMenuTedsRowPicker : public SCompoundWidget
+{
+public:
+	SLATE_BEGIN_ARGS(SPropertyMenuTedsRowPicker)
+		: _AllowClear(true)
+		, _ElementFilter()
+	{}
+		SLATE_ARGUMENT(bool, AllowClear)
+		SLATE_ARGUMENT(TypedElementDataStorage::FQueryDescription, TypedElementQueryFilter)
+		SLATE_ARGUMENT(FOnShouldFilterTedsRow, ElementFilter)
+		SLATE_EVENT(FOnTedsRowSelected, OnSet)
+		SLATE_EVENT(FSimpleDelegate, OnClose)
+	SLATE_END_ARGS()
+
+	void Construct(const FArguments& InArgs);
+
+private:
+
+	void OnClear();
+
+	void OnElementSelected(TypedElementDataStorage::RowHandle RowHandle);
+
+	void SetValue(TypedElementDataStorage::RowHandle RowHandle);
+
+private:
+	bool bAllowClear;
+
+	TypedElementDataStorage::FQueryDescription TypedElementQueryFilter;
+
+	FOnShouldFilterTedsRow ElementFilter;
+
+	FOnTedsRowSelected OnSet;
+
+	FSimpleDelegate OnClose;
+
+	FSimpleDelegate OnUseSelected;
+};
\ No newline at end of file
diff --git a/Engine/Plugins/Experimental/EditorDataStorageFeatures/Source/TedsPropertyEditor/Public/Widgets/SPropertyMenuTypedElementPicker.h b/Engine/Plugins/Experimental/EditorDataStorageFeatures/Source/TedsPropertyEditor/Public/Widgets/SPropertyMenuTypedElementPicker.h
deleted file mode 100644
diff --git a/Engine/Plugins/Experimental/EditorDataStorageFeatures/Source/TedsPropertyEditor/TedsPropertyEditor.Build.cs b/Engine/Plugins/Experimental/EditorDataStorageFeatures/Source/TedsPropertyEditor/TedsPropertyEditor.Build.cs
--- a/Engine/Plugins/Experimental/EditorDataStorageFeatures/Source/TedsPropertyEditor/TedsPropertyEditor.Build.cs
+++ b/Engine/Plugins/Experimental/EditorDataStorageFeatures/Source/TedsPropertyEditor/TedsPropertyEditor.Build.cs
@@ -18,18 +18,13 @@ public class TedsPropertyEditor : ModuleRules
 				{
 					"Core",
 					"CoreUObject",
-					"EditorFramework",
 					"Engine",
 					"SceneOutliner",
 					"Slate",
 					"SlateCore",
 					"TypedElementFramework",
-					"TedsOutliner"
-				});
-			PrivateDependencyModuleNames.AddRange(
-				new string[]
-				{
-					"UnrealEd" // For EditorUndoClient
+					"TedsOutliner",
+					"UnrealEd", // For EditorUndoClient
 				});
 
 			DynamicallyLoadedModuleNames.AddRange(new string[] {});
