[Teds] Rename TedsDebugger files and classes to remove TypedElement strings

#rnx

diff --git a/Engine/Plugins/Experimental/EditorDataStorageFeatures/Source/TedsDebugger/Private/QueryEditor/TedsQueryEditor.cpp b/Engine/Plugins/Experimental/EditorDataStorageFeatures/Source/TedsDebugger/Private/QueryEditor/TedsQueryEditor.cpp
--- a/Engine/Plugins/Experimental/EditorDataStorageFeatures/Source/TedsDebugger/Private/QueryEditor/TedsQueryEditor.cpp
+++ b/Engine/Plugins/Experimental/EditorDataStorageFeatures/Source/TedsDebugger/Private/QueryEditor/TedsQueryEditor.cpp
@@ -6,13 +6,13 @@
 #include "Components/VerticalBox.h"
 #include "Widgets/Text/STextBlock.h"
 #include "Widgets/Input/SEditableComboBox.h"
-#include "QueryEditor/Widgets/TedsConditionSelectionComboWidget.h"
-#include "Widgets/TedsConditionCollectionViewWidget.h"
-#include "Widgets/TedsQueryEditorResultsView.h"
+#include "Widgets/QueryEditor/TedsConditionSelectionComboWidget.h"
+#include "Widgets/QueryEditor/TedsConditionCollectionViewWidget.h"
+#include "Widgets/QueryEditor/TedsQueryEditorResultsView.h"
 
 #define LOCTEXT_NAMESPACE "TedsQueryEditor"
 
-namespace UE::Teds::Debug::QueryEditor
+namespace UE::EditorDataStorage::Debug::QueryEditor
 {
 	struct SQueryEditorWidget::ColumnComboItem
 	{
@@ -26,7 +26,7 @@ namespace UE::Teds::Debug::QueryEditor
 
 	void SQueryEditorWidget::Construct(const FArguments& InArgs, FTedsQueryEditorModel& QueryEditorModel)
 	{
-		using namespace UE::Teds::Debug;
+		using namespace UE::EditorDataStorage::Debug;
 	
 		ComboItems.Reset();
 		Model = &QueryEditorModel;
@@ -48,7 +48,7 @@ namespace UE::Teds::Debug::QueryEditor
 					+SHorizontalBox::Slot()
 					.AutoWidth()
 					[
-						SNew(QueryEditor::SConditionComboWidget, *Model, QueryEditor::EOperatorType::Select)
+						SNew(SConditionComboWidget, *Model, QueryEditor::EOperatorType::Select)
 					]
 				]
 				+SVerticalBox::Slot()
@@ -57,12 +57,12 @@ namespace UE::Teds::Debug::QueryEditor
 					SNew(SHorizontalBox)
 					+SHorizontalBox::Slot()
 					[
-						SNew(UE::Teds::Debug::QueryEditor::SConditionCollectionViewWidget, *Model, QueryEditor::EOperatorType::All)
+						SNew(SConditionCollectionViewWidget, *Model, QueryEditor::EOperatorType::All)
 					]
 					+SHorizontalBox::Slot()
 					.AutoWidth()
 					[
-						SNew(QueryEditor::SConditionComboWidget, *Model, QueryEditor::EOperatorType::All)
+						SNew(SConditionComboWidget, *Model, QueryEditor::EOperatorType::All)
 					]
 				]
 				+SVerticalBox::Slot()
@@ -71,12 +71,12 @@ namespace UE::Teds::Debug::QueryEditor
 					SNew(SHorizontalBox)
 					+SHorizontalBox::Slot()
 					[
-					SNew(UE::Teds::Debug::QueryEditor::SConditionCollectionViewWidget, *Model, QueryEditor::EOperatorType::Any)
+					SNew(SConditionCollectionViewWidget, *Model, QueryEditor::EOperatorType::Any)
 					]
 					+SHorizontalBox::Slot()
 					.AutoWidth()
 					[
-					SNew(QueryEditor::SConditionComboWidget, *Model, QueryEditor::EOperatorType::Any)
+					SNew(SConditionComboWidget, *Model, QueryEditor::EOperatorType::Any)
 					]
 				]
 				+SVerticalBox::Slot()
@@ -85,17 +85,17 @@ namespace UE::Teds::Debug::QueryEditor
 					SNew(SHorizontalBox)
 					+SHorizontalBox::Slot()
 					[
-					SNew(UE::Teds::Debug::QueryEditor::SConditionCollectionViewWidget, *Model, QueryEditor::EOperatorType::None)
+					SNew(SConditionCollectionViewWidget, *Model, QueryEditor::EOperatorType::None)
 					]
 					+SHorizontalBox::Slot()
 					.AutoWidth()
 					[
-					SNew(QueryEditor::SConditionComboWidget, *Model, QueryEditor::EOperatorType::None)
+					SNew(SConditionComboWidget, *Model, QueryEditor::EOperatorType::None)
 					]
 				]
 				+SVerticalBox::Slot()
 				[
-					SNew(QueryEditor::SResultsView, *Model)
+					SNew(SResultsView, *Model)
 				]
 			]
 			
diff --git a/Engine/Plugins/Experimental/EditorDataStorageFeatures/Source/TedsDebugger/Private/QueryEditor/TedsQueryEditor.h b/Engine/Plugins/Experimental/EditorDataStorageFeatures/Source/TedsDebugger/Private/QueryEditor/TedsQueryEditor.h
--- a/Engine/Plugins/Experimental/EditorDataStorageFeatures/Source/TedsDebugger/Private/QueryEditor/TedsQueryEditor.h
+++ b/Engine/Plugins/Experimental/EditorDataStorageFeatures/Source/TedsDebugger/Private/QueryEditor/TedsQueryEditor.h
@@ -6,7 +6,7 @@
 class SDockTab;
 class SWindow;
 
-namespace UE::Teds::Debug::QueryEditor
+namespace UE::EditorDataStorage::Debug::QueryEditor
 {
 	class FTedsQueryEditorModel;
 	class SQueryEditorWidget : public SCompoundWidget
diff --git a/Engine/Plugins/Experimental/EditorDataStorageFeatures/Source/TedsDebugger/Private/QueryEditor/TedsQueryEditorModel.cpp b/Engine/Plugins/Experimental/EditorDataStorageFeatures/Source/TedsDebugger/Private/QueryEditor/TedsQueryEditorModel.cpp
--- a/Engine/Plugins/Experimental/EditorDataStorageFeatures/Source/TedsDebugger/Private/QueryEditor/TedsQueryEditorModel.cpp
+++ b/Engine/Plugins/Experimental/EditorDataStorageFeatures/Source/TedsDebugger/Private/QueryEditor/TedsQueryEditorModel.cpp
@@ -1,4 +1,4 @@
-﻿// Copyright Epic Games, Inc. All Rights Reserved.
+// Copyright Epic Games, Inc. All Rights Reserved.
 #include "QueryEditor/TedsQueryEditorModel.h"
 
 #include "Elements/Common/TypedElementCommonTypes.h"
@@ -6,7 +6,7 @@
 #include "Elements/Interfaces/TypedElementDataStorageInterface.h"
 #include "UObject/UObjectIterator.h"
 
-namespace UE::Teds::Debug::QueryEditor
+namespace UE::EditorDataStorage::Debug::QueryEditor
 {
 	bool FConditionEntryHandle::operator==(const FConditionEntryHandle& Rhs) const
 	{
diff --git a/Engine/Plugins/Experimental/EditorDataStorageFeatures/Source/TedsDebugger/Private/QueryEditor/TedsQueryEditorModel.h b/Engine/Plugins/Experimental/EditorDataStorageFeatures/Source/TedsDebugger/Private/QueryEditor/TedsQueryEditorModel.h
--- a/Engine/Plugins/Experimental/EditorDataStorageFeatures/Source/TedsDebugger/Private/QueryEditor/TedsQueryEditorModel.h
+++ b/Engine/Plugins/Experimental/EditorDataStorageFeatures/Source/TedsDebugger/Private/QueryEditor/TedsQueryEditorModel.h
@@ -1,4 +1,4 @@
-﻿// Copyright Epic Games, Inc. All Rights Reserved.
+// Copyright Epic Games, Inc. All Rights Reserved.
 #pragma once
 #include "Elements/Common/TypedElementHandles.h"
 #include "Elements/Framework/TypedElementMetaData.h"
@@ -6,7 +6,7 @@
 
 class ITypedElementDataStorageInterface;
 
-namespace UE::Teds::Debug::QueryEditor
+namespace UE::EditorDataStorage::Debug::QueryEditor
 {
 	class FTedsQueryEditorModel;
 	
diff --git a/Engine/Plugins/Experimental/EditorDataStorageFeatures/Source/TedsDebugger/Private/QueryEditor/Widgets/TedsConditionCollectionViewWidget.cpp b/Engine/Plugins/Experimental/EditorDataStorageFeatures/Source/TedsDebugger/Private/QueryEditor/Widgets/TedsConditionCollectionViewWidget.cpp
deleted file mode 100644
diff --git a/Engine/Plugins/Experimental/EditorDataStorageFeatures/Source/TedsDebugger/Private/QueryEditor/Widgets/TedsConditionCollectionViewWidget.h b/Engine/Plugins/Experimental/EditorDataStorageFeatures/Source/TedsDebugger/Private/QueryEditor/Widgets/TedsConditionCollectionViewWidget.h
deleted file mode 100644
diff --git a/Engine/Plugins/Experimental/EditorDataStorageFeatures/Source/TedsDebugger/Private/QueryEditor/Widgets/TedsConditionSelectionCollectionWidget.cpp b/Engine/Plugins/Experimental/EditorDataStorageFeatures/Source/TedsDebugger/Private/QueryEditor/Widgets/TedsConditionSelectionCollectionWidget.cpp
deleted file mode 100644
diff --git a/Engine/Plugins/Experimental/EditorDataStorageFeatures/Source/TedsDebugger/Private/QueryEditor/Widgets/TedsConditionSelectionCollectionWidget.h b/Engine/Plugins/Experimental/EditorDataStorageFeatures/Source/TedsDebugger/Private/QueryEditor/Widgets/TedsConditionSelectionCollectionWidget.h
deleted file mode 100644
diff --git a/Engine/Plugins/Experimental/EditorDataStorageFeatures/Source/TedsDebugger/Private/QueryEditor/Widgets/TedsConditionSelectionComboWidget.cpp b/Engine/Plugins/Experimental/EditorDataStorageFeatures/Source/TedsDebugger/Private/QueryEditor/Widgets/TedsConditionSelectionComboWidget.cpp
deleted file mode 100644
diff --git a/Engine/Plugins/Experimental/EditorDataStorageFeatures/Source/TedsDebugger/Private/QueryEditor/Widgets/TedsConditionSelectionComboWidget.h b/Engine/Plugins/Experimental/EditorDataStorageFeatures/Source/TedsDebugger/Private/QueryEditor/Widgets/TedsConditionSelectionComboWidget.h
deleted file mode 100644
diff --git a/Engine/Plugins/Experimental/EditorDataStorageFeatures/Source/TedsDebugger/Private/QueryEditor/Widgets/TedsQueryEditorResultsView.cpp b/Engine/Plugins/Experimental/EditorDataStorageFeatures/Source/TedsDebugger/Private/QueryEditor/Widgets/TedsQueryEditorResultsView.cpp
deleted file mode 100644
diff --git a/Engine/Plugins/Experimental/EditorDataStorageFeatures/Source/TedsDebugger/Private/QueryEditor/Widgets/TedsQueryEditorResultsView.h b/Engine/Plugins/Experimental/EditorDataStorageFeatures/Source/TedsDebugger/Private/QueryEditor/Widgets/TedsQueryEditorResultsView.h
deleted file mode 100644
diff --git a/Engine/Plugins/Experimental/EditorDataStorageFeatures/Source/TedsDebugger/Private/STedsDebugger.cpp b/Engine/Plugins/Experimental/EditorDataStorageFeatures/Source/TedsDebugger/Private/STedsDebugger.cpp
--- a/Engine/Plugins/Experimental/EditorDataStorageFeatures/Source/TedsDebugger/Private/STedsDebugger.cpp
+++ b/Engine/Plugins/Experimental/EditorDataStorageFeatures/Source/TedsDebugger/Private/STedsDebugger.cpp
@@ -16,7 +16,7 @@
 
 #define LOCTEXT_NAMESPACE "STedsDebugger"
 
-namespace UE::Teds::Debugger::Private
+namespace UE::EditorDataStorage::Debugger::Private
 {
 	FName QueryEditorToolTabName = TEXT("TEDS Query Editor");
 	FName TableViewerToolTabName = TEXT("TEDS Table Viewer");
@@ -43,7 +43,7 @@ void STedsDebugger::Construct(const FArguments& InArgs, const TSharedRef<SDockTa
 	// Register Tab Spawners
 	RegisterTabSpawners();
 
-	using namespace UE::Teds::Debugger::Private;
+	using namespace UE::EditorDataStorage::Debugger::Private;
 
 	// Setup the default layout
 	TSharedRef<FTabManager::FLayout> Layout = FTabManager::NewLayout("TedsDebuggerLayout_v0")
@@ -136,15 +136,15 @@ TSharedRef<SDockTab> STedsDebugger::SpawnQueryEditorTab(const FSpawnTabArgs& Arg
 		if (Registry && Registry->AreDataStorageInterfacesSet())
 		{
 			ITypedElementDataStorageInterface* DataStorageInterface = Registry->GetMutableDataStorage();
-			QueryEditorModel = MakeUnique<UE::Teds::Debug::QueryEditor::FTedsQueryEditorModel>(*DataStorageInterface);
+			QueryEditorModel = MakeUnique<UE::EditorDataStorage::Debug::QueryEditor::FTedsQueryEditorModel>(*DataStorageInterface);
 		}
 	}
 	if (QueryEditorModel)
 	{
 		QueryEditorModel->Reset();	
 
-		TSharedRef<UE::Teds::Debug::QueryEditor::SQueryEditorWidget> QueryEditor =
-			SNew(UE::Teds::Debug::QueryEditor::SQueryEditorWidget, *QueryEditorModel);
+		TSharedRef<UE::EditorDataStorage::Debug::QueryEditor::SQueryEditorWidget> QueryEditor =
+			SNew(UE::EditorDataStorage::Debug::QueryEditor::SQueryEditorWidget, *QueryEditorModel);
 		DockTab->SetContent(QueryEditor);
 	}
 	else
@@ -207,7 +207,7 @@ TSharedRef<SDockTab> STedsDebugger::SpawnTableViewerTab(const FSpawnTabArgs& Arg
 
 void STedsDebugger::RegisterTabSpawners()
 {
-	using namespace UE::Teds::Debugger::Private;
+	using namespace UE::EditorDataStorage::Debugger::Private;
 
 	const TSharedRef<FWorkspaceItem> AppMenuGroup =
 		TabManager->AddLocalWorkspaceMenuCategory(LOCTEXT("TedsDebuggerGroupName", "Teds Debugger"));
@@ -248,7 +248,7 @@ void STedsDebugger::NavigateToRow(TypedElementDataStorage::RowHandle InRow) cons
 	// If the debugger isn't already open, open it
 	if (!TableViewerInstance.IsValid())
 	{
-		TabManager->TryInvokeTab(UE::Teds::Debugger::Private::TableViewerToolTabName);
+		TabManager->TryInvokeTab(UE::EditorDataStorage::Debugger::Private::TableViewerToolTabName);
 	}
 
 	TSharedPtr<ISceneOutliner> TableViewerPinned = TableViewerInstance.Pin();
diff --git a/Engine/Plugins/Experimental/EditorDataStorageFeatures/Source/TedsDebugger/Private/STedsDebugger.h b/Engine/Plugins/Experimental/EditorDataStorageFeatures/Source/TedsDebugger/Private/STedsDebugger.h
--- a/Engine/Plugins/Experimental/EditorDataStorageFeatures/Source/TedsDebugger/Private/STedsDebugger.h
+++ b/Engine/Plugins/Experimental/EditorDataStorageFeatures/Source/TedsDebugger/Private/STedsDebugger.h
@@ -50,5 +50,5 @@ private:
 	TWeakPtr<ISceneOutliner> TableViewerInstance;
 
 	// Query Editor
-	TUniquePtr<UE::Teds::Debug::QueryEditor::FTedsQueryEditorModel> QueryEditorModel;
+	TUniquePtr<UE::EditorDataStorage::Debug::QueryEditor::FTedsQueryEditorModel> QueryEditorModel;
 };
diff --git a/Engine/Plugins/Experimental/EditorDataStorageFeatures/Source/TedsDebugger/Private/TedsDebuggerModule.cpp b/Engine/Plugins/Experimental/EditorDataStorageFeatures/Source/TedsDebugger/Private/TedsDebuggerModule.cpp
--- a/Engine/Plugins/Experimental/EditorDataStorageFeatures/Source/TedsDebugger/Private/TedsDebuggerModule.cpp
+++ b/Engine/Plugins/Experimental/EditorDataStorageFeatures/Source/TedsDebugger/Private/TedsDebuggerModule.cpp
@@ -12,9 +12,11 @@
 
 #define LOCTEXT_NAMESPACE "TedsDebuggerModule"
 
-namespace UE::Teds::Debugger::Private
+namespace UE::EditorDataStorage::Debug
 {
-	FName TedsDebuggerTablName = TEXT("TEDS Debugger");
+namespace Private
+{
+	FName TedsDebuggerTableName = TEXT("TEDS Debugger");
 }
 
 void FTedsDebuggerModule::StartupModule()
@@ -35,10 +37,10 @@ void FTedsDebuggerModule::ShutdownModule()
 
 void FTedsDebuggerModule::RegisterTabSpawners()
 {
-	using namespace UE::Teds::Debugger::Private;
+	using namespace UE::EditorDataStorage::Debug::Private;
 
 	FGlobalTabmanager::Get()->RegisterNomadTabSpawner(
-		TedsDebuggerTablName,
+		TedsDebuggerTableName,
 		FOnSpawnTab::CreateRaw(this, &FTedsDebuggerModule::OpenTedsDebuggerTab))
 		.SetGroup(WorkspaceMenu::GetMenuStructure().GetDeveloperToolsDebugCategory())
 		.SetDisplayName(LOCTEXT("TedsDebugger_QueryEditorDisplayName", "TEDS Debugger"))
@@ -50,7 +52,7 @@ void FTedsDebuggerModule::UnregisterTabSpawners() const
 {
 	if (FSlateApplication::IsInitialized())
 	{
-		FGlobalTabmanager::Get()->UnregisterNomadTabSpawner(UE::Teds::Debugger::Private::TedsDebuggerTablName);
+		FGlobalTabmanager::Get()->UnregisterNomadTabSpawner(UE::EditorDataStorage::Debug::Private::TedsDebuggerTableName);
 	}
 }
 
@@ -80,7 +82,7 @@ void FTedsDebuggerModule::NavigateToRow(TypedElementDataStorage::RowHandle InRow
 	// If the debugger isn't already open, open it
 	if(!TedsDebuggerInstance.IsValid())
 	{
-		FGlobalTabmanager::Get()->TryInvokeTab(UE::Teds::Debugger::Private::TedsDebuggerTablName);
+		FGlobalTabmanager::Get()->TryInvokeTab(UE::EditorDataStorage::Debug::Private::TedsDebuggerTableName);
 	}
 
 	TSharedPtr<STedsDebugger> TedsDebuggerPinned = TedsDebuggerInstance.Pin();
@@ -91,7 +93,8 @@ void FTedsDebuggerModule::NavigateToRow(TypedElementDataStorage::RowHandle InRow
 
 	TedsDebuggerPinned->NavigateToRow(InRow);
 }
+}
 
-IMPLEMENT_MODULE(FTedsDebuggerModule, TedsDebugger);
+IMPLEMENT_MODULE(UE::EditorDataStorage::Debug::FTedsDebuggerModule, TedsDebugger);
 
 #undef LOCTEXT_NAMESPACE
\ No newline at end of file
diff --git a/Engine/Plugins/Experimental/EditorDataStorageFeatures/Source/TedsDebugger/Private/TedsDebuggerModule.h b/Engine/Plugins/Experimental/EditorDataStorageFeatures/Source/TedsDebugger/Private/TedsDebuggerModule.h
new file mode 100644
--- /dev/null
+++ b/Engine/Plugins/Experimental/EditorDataStorageFeatures/Source/TedsDebugger/Private/TedsDebuggerModule.h
@@ -0,0 +1,43 @@
+// Copyright Epic Games, Inc. All Rights Reserved.
+
+#pragma once
+
+#include "Elements/Common/TypedElementHandles.h"
+#include "Modules/ModuleInterface.h"
+#include "Templates/SharedPointer.h"
+
+class SDockTab;
+class SWidget;
+class FSpawnTabArgs;
+class ISceneOutliner;
+class STedsDebugger;
+
+namespace UE::EditorDataStorage::Debug
+{
+/**
+ * Implements the Teds Debugger module.
+ */
+class FTedsDebuggerModule
+	: public IModuleInterface
+{
+public:
+
+	FTedsDebuggerModule() = default;
+
+	// IModuleInterface interface
+	virtual void StartupModule() override;
+	virtual void ShutdownModule() override;
+
+	// Open the TEDS Debugger if not already open, and navigate to the given row in the table viewer tab
+	void NavigateToRow(TypedElementDataStorage::RowHandle InRow) const;
+
+private:
+	void RegisterTabSpawners();
+	void UnregisterTabSpawners() const;
+	
+	TSharedRef<SDockTab> OpenTedsDebuggerTab(const FSpawnTabArgs& SpawnTabArgs);
+
+private:
+	TWeakPtr<STedsDebugger> TedsDebuggerInstance;
+};
+}
\ No newline at end of file
diff --git a/Engine/Plugins/Experimental/EditorDataStorageFeatures/Source/TedsDebugger/Private/Widgets/QueryEditor/TedsConditionCollectionViewWidget.cpp b/Engine/Plugins/Experimental/EditorDataStorageFeatures/Source/TedsDebugger/Private/Widgets/QueryEditor/TedsConditionCollectionViewWidget.cpp
new file mode 100644
--- /dev/null
+++ b/Engine/Plugins/Experimental/EditorDataStorageFeatures/Source/TedsDebugger/Private/Widgets/QueryEditor/TedsConditionCollectionViewWidget.cpp
@@ -0,0 +1,180 @@
+// Copyright Epic Games, Inc. All Rights Reserved.
+#include "TedsConditionCollectionViewWidget.h"
+
+#include "Components/HorizontalBox.h"
+#include "QueryEditor/TedsQueryEditorModel.h"
+#include "Widgets/SOverlay.h"
+#include "Widgets/Input/SButton.h"
+#include "Widgets/Layout/SBox.h"
+#include "Widgets/Layout/SWrapBox.h"
+#include "Widgets/Text/STextBlock.h"
+
+using namespace UE::EditorDataStorage::Debug::QueryEditor;
+
+#define LOCTEXT_NAMESPACE "TedsDebuggerModule"
+
+namespace UE::EditorDataStorage::Debug::QueryEditor::Private
+{
+	FText GetOperatorTypeText(EOperatorType OperatorType)
+	{
+		switch(OperatorType)
+		{
+		case EOperatorType::Select:
+			return FText::FromStringView(TEXT("Select<"));
+		case EOperatorType::All:
+			return FText::FromStringView(TEXT("All<"));
+		case EOperatorType::Any:
+			return FText::FromStringView(TEXT("Any<"));
+		case EOperatorType::None:
+			return FText::FromStringView(TEXT("None<"));
+		case EOperatorType::Unset:
+			return FText::FromStringView(TEXT("Unset<"));
+		case EOperatorType::Invalid:
+		default:
+			return FText::FromStringView(TEXT("Invalid<"));
+		}
+	}
+}
+
+SConditionCollectionViewWidget::~SConditionCollectionViewWidget()
+{
+	Model->GetModelChangedDelegate().Remove(OnModelChangedDelegate);
+}
+
+void SConditionCollectionViewWidget::Construct(
+	const FArguments& InArgs,
+	FTedsQueryEditorModel& InModel, 
+	QueryEditor::EOperatorType InOperatorType)
+{
+	Model = &InModel;
+	OperatorType = InOperatorType;
+
+	OnModelChangedDelegate = Model->GetModelChangedDelegate().AddRaw(this, &SConditionCollectionViewWidget::OnModelChanged);
+
+	ChildSlot
+	[
+		SNew(SBorder)
+		.Padding(4.0, 3.0, 3.0, 4.0)
+		[
+		SNew(SHorizontalBox)
+		+SHorizontalBox::Slot()
+		.AutoWidth()
+		[
+			SNew(SVerticalBox)
+			+SVerticalBox::Slot()
+			.VAlign(EVerticalAlignment::VAlign_Center)
+			[
+				SNew(STextBlock)
+				.Text(Private::GetOperatorTypeText(OperatorType))
+			]
+		]
+		+SHorizontalBox::Slot()
+		[
+		SAssignNew(ColumnButtonWrap, SWrapBox)
+		.UseAllottedSize(true)
+		+SWrapBox::Slot()
+		.FillEmptySpace(false)
+		[
+			SNullWidget::NullWidget
+		]
+		]
+		+SHorizontalBox::Slot()
+		.AutoWidth()
+		[
+			SNew(SVerticalBox)
+			+SVerticalBox::Slot()
+			.VAlign(EVerticalAlignment::VAlign_Center)
+			[
+				SNew(STextBlock)
+				.Text(LOCTEXT("TedsQueryEditor_OpType_>", ">"))
+			]
+		]
+		]
+	];
+
+	SetCanTick(true);
+}
+
+void SConditionCollectionViewWidget::Tick(const FGeometry& AllottedGeometry, const double InCurrentTime, const float InDeltaTime)
+{
+	SCompoundWidget::Tick(AllottedGeometry, InCurrentTime, InDeltaTime);
+
+	ColumnButtonWrap->ClearChildren();
+	SetCanTick(false);
+
+	ButtonStyle = FButtonStyle(FCoreStyle::Get().GetWidgetStyle< FButtonStyle >( "Button" ))
+		.SetHoveredForeground(FSlateColor(FLinearColor(0.9f, 0.0f, 0.0f, 1.0f)));
+
+	if (Model->CountConditionsOfOperator(OperatorType) != 0)
+	{
+		Model->ForEachCondition([this](FTedsQueryEditorModel&, FConditionEntryHandle Handle)
+		{
+			if (Model->GetOperatorType(Handle) != OperatorType)
+			{
+				return;
+			}
+			// TODO Add styling to the button to go red and put an "X" to signal clicking
+			// will remove the column
+			const UScriptStruct* Type = Model->GetColumnScriptStruct(Handle);
+			FText TempName = FText::FromString(Type->GetName());
+			
+			TSharedRef<STextBlock> TextNameOfColumn = SNew(STextBlock)
+				.Text(TempName);
+
+			TSharedRef<SOverlay> ButtonContent = SNew(SOverlay);
+			
+			TSharedRef<SButton> Button = SNew(SButton)
+				.Text(TempName)
+				.OnClicked_Lambda( [this, Handle]() -> FReply 
+				{
+					Model->SetOperatorType(Handle, EOperatorType::Unset);
+					return FReply::Handled();
+				})
+				.IsEnabled_Lambda( [this, Handle]() -> bool
+				{
+					return Model->CanSetOperatorType(Handle, EOperatorType::Unset).Key;
+				})
+				.Content()
+				[
+					TextNameOfColumn
+				];
+			STextBlock* TextBlockPtr = &TextNameOfColumn.Get();
+			Button->SetOnHovered(FSimpleDelegate::CreateLambda([TextBlockPtr]()
+			{
+				TextBlockPtr->SetVisibility(EVisibility::Hidden);
+			}));
+			Button->SetOnUnhovered(FSimpleDelegate::CreateLambda([TextBlockPtr, TempName]()
+			{
+				TextBlockPtr->SetVisibility(EVisibility::Visible);
+			}));
+
+			
+			ButtonContent->AddSlot()
+			[
+				Button
+			];
+			
+			ColumnButtonWrap->AddSlot()
+			[
+				ButtonContent
+			];
+
+		});
+	}
+	else
+	{
+		ColumnButtonWrap->AddSlot()
+		[
+			SNullWidget::NullWidget
+		];
+	}
+	
+}
+
+void SConditionCollectionViewWidget::OnModelChanged()
+{
+	SetCanTick(true);
+	Invalidate(EInvalidateWidgetReason::Layout);
+}
+
+#undef LOCTEXT_NAMESPACE
\ No newline at end of file
diff --git a/Engine/Plugins/Experimental/EditorDataStorageFeatures/Source/TedsDebugger/Private/Widgets/QueryEditor/TedsConditionCollectionViewWidget.h b/Engine/Plugins/Experimental/EditorDataStorageFeatures/Source/TedsDebugger/Private/Widgets/QueryEditor/TedsConditionCollectionViewWidget.h
new file mode 100644
--- /dev/null
+++ b/Engine/Plugins/Experimental/EditorDataStorageFeatures/Source/TedsDebugger/Private/Widgets/QueryEditor/TedsConditionCollectionViewWidget.h
@@ -0,0 +1,37 @@
+// Copyright Epic Games, Inc. All Rights Reserved.
+#pragma once
+
+#include "QueryEditor/TedsQueryEditorModel.h"
+#include "Styling/SlateTypes.h"
+#include "Widgets/SCompoundWidget.h"
+
+class SWrapBox;
+class SHorizontalBox;
+
+namespace UE::EditorDataStorage::Debug::QueryEditor
+{
+	enum class EOperatorType : uint32;
+	class FTedsQueryEditorModel;
+
+	class SConditionCollectionViewWidget : public SCompoundWidget
+	{
+	public:
+		SLATE_BEGIN_ARGS( SConditionCollectionViewWidget ){}
+		SLATE_END_ARGS()
+
+		~SConditionCollectionViewWidget() override;
+		void Construct(const FArguments& InArgs, FTedsQueryEditorModel& InModel, QueryEditor::EOperatorType InOperatorType);
+		void Tick(const FGeometry& AllottedGeometry, const double InCurrentTime, const float InDeltaTime) override;
+		
+	private:
+
+		void OnModelChanged();
+
+		FTedsQueryEditorModel* Model = nullptr;
+		QueryEditor::EOperatorType OperatorType = EOperatorType::Unset;
+		FDelegateHandle OnModelChangedDelegate;
+
+		TSharedPtr<SWrapBox> ColumnButtonWrap;
+		FButtonStyle ButtonStyle;
+	};
+}
diff --git a/Engine/Plugins/Experimental/EditorDataStorageFeatures/Source/TedsDebugger/Private/Widgets/QueryEditor/TedsConditionSelectionCollectionWidget.cpp b/Engine/Plugins/Experimental/EditorDataStorageFeatures/Source/TedsDebugger/Private/Widgets/QueryEditor/TedsConditionSelectionCollectionWidget.cpp
new file mode 100644
--- /dev/null
+++ b/Engine/Plugins/Experimental/EditorDataStorageFeatures/Source/TedsDebugger/Private/Widgets/QueryEditor/TedsConditionSelectionCollectionWidget.cpp
@@ -0,0 +1,10 @@
+// Copyright Epic Games, Inc. All Rights Reserved.
+#include "Widgets/QueryEditor/TedsConditionSelectionCollectionWidget.h"
+
+namespace UE::EditorDataStorage::Debug::QueryEditor
+{
+	void SConditionSelectionCollectionWidget::Construct(const FArguments& InArgs, FTedsQueryEditorModel& InModel)
+	{
+		Model = &InModel;
+	}
+}
\ No newline at end of file
diff --git a/Engine/Plugins/Experimental/EditorDataStorageFeatures/Source/TedsDebugger/Private/Widgets/QueryEditor/TedsConditionSelectionCollectionWidget.h b/Engine/Plugins/Experimental/EditorDataStorageFeatures/Source/TedsDebugger/Private/Widgets/QueryEditor/TedsConditionSelectionCollectionWidget.h
new file mode 100644
--- /dev/null
+++ b/Engine/Plugins/Experimental/EditorDataStorageFeatures/Source/TedsDebugger/Private/Widgets/QueryEditor/TedsConditionSelectionCollectionWidget.h
@@ -0,0 +1,24 @@
+// Copyright Epic Games, Inc. All Rights Reserved.
+#pragma once
+#include "Widgets/SCompoundWidget.h"
+#include "Templates/SharedPointerFwd.h"
+
+namespace UE::EditorDataStorage::Debug::QueryEditor
+{
+	class FTedsQueryEditorModel;
+	
+	/**
+	 * Widget which shows and allows editing of columns/tabs within a given condition
+	 */
+	class SConditionSelectionCollectionWidget : public SCompoundWidget
+	{
+	public:
+		SLATE_BEGIN_ARGS( SConditionSelectionCollectionWidget ){}
+		SLATE_END_ARGS()
+
+		void Construct(const FArguments& InArgs, FTedsQueryEditorModel& InModel);
+
+	private:
+		FTedsQueryEditorModel* Model = nullptr;
+	};
+}
diff --git a/Engine/Plugins/Experimental/EditorDataStorageFeatures/Source/TedsDebugger/Private/Widgets/QueryEditor/TedsConditionSelectionComboWidget.cpp b/Engine/Plugins/Experimental/EditorDataStorageFeatures/Source/TedsDebugger/Private/Widgets/QueryEditor/TedsConditionSelectionComboWidget.cpp
new file mode 100644
--- /dev/null
+++ b/Engine/Plugins/Experimental/EditorDataStorageFeatures/Source/TedsDebugger/Private/Widgets/QueryEditor/TedsConditionSelectionComboWidget.cpp
@@ -0,0 +1,111 @@
+// Copyright Epic Games, Inc. All Rights Reserved.
+#include "TedsConditionSelectionComboWidget.h"
+
+#include "QueryEditor/TedsQueryEditorModel.h"
+
+#include "Algo/FindSortedStringCaseInsensitive.h"
+#include "Elements/Common/TypedElementDataStorageLog.h"
+#include "Filters/SBasicFilterBar.h"
+#include "Widgets/Text/STextBlock.h"
+#include "Widgets/Input/SComboBox.h"
+
+namespace UE::EditorDataStorage::Debug::QueryEditor
+{
+	struct SConditionComboWidget::FComboItem
+	{
+		FConditionEntryHandle Handle;
+		
+		bool IsValid() const { return Handle.IsValid(); }
+		void Reset() { Handle.Reset(); }
+	};
+
+	void SConditionComboWidget::OnConditionCollectionChanged()
+	{		
+		PopulateComboItems();
+	}
+
+	SConditionComboWidget::~SConditionComboWidget()
+	{
+		Model->GetModelChangedDelegate().Remove(ConditionCollectionChangedHandle);
+	}
+
+	void SConditionComboWidget::Construct(const FArguments& InArgs, FTedsQueryEditorModel& InModel, QueryEditor::EOperatorType InConditionType)
+	{
+		Model = &InModel;
+
+		check(InConditionType != EOperatorType::Invalid);
+		ConditionType = InConditionType;
+		
+		ConditionCollectionChangedHandle = Model->GetModelChangedDelegate().AddRaw(this, &SConditionComboWidget::OnConditionCollectionChanged);
+
+		
+		
+		ComboBox = SNew(SComboBox<TSharedPtr<FComboItem>>)
+			.OptionsSource(&ComboItems)
+			.OnSelectionChanged(this, &SConditionComboWidget::OnSelectionChanged)
+			.OnGenerateWidget(this, &SConditionComboWidget::OnGenerateWidget)
+			.IsEnabled_Lambda([this](){ return !ComboItems.IsEmpty(); })
+			.Content()
+			[
+				SNew(STextBlock)
+				.Text_Lambda([]() { return FText::FromString("+"); })
+			];
+
+		PopulateComboItems();
+
+		ChildSlot
+		[
+		SNew( SVerticalBox )
+		+ SVerticalBox::Slot()
+		.AutoHeight()
+		[
+			ComboBox.ToSharedRef()
+		]
+			
+		];
+	}
+
+	void SConditionComboWidget::PopulateComboItems()
+	{
+		ComboItems.Reset();
+		Model->GenerateValidOperatorChoices(ConditionType, [this](const FTedsQueryEditorModel&, const QueryEditor::FConditionEntryHandle& Handle)
+		{
+			FComboItem ComboItem;
+			ComboItem.Handle = Handle;
+			ComboItems.Emplace(MakeShared<FComboItem>(ComboItem));
+			
+		});
+
+		// Sort the ComboItems - this will make it only slightly easier to find them...
+		Algo::Sort(ComboItems, [this](TSharedPtr<FComboItem>& A, TSharedPtr<FComboItem>& B)->bool
+		{
+			const UScriptStruct* AStruct = Model->GetColumnScriptStruct(A->Handle);
+			const UScriptStruct* BStruct = Model->GetColumnScriptStruct(B->Handle);
+			return AStruct->GetFName().Compare(BStruct->GetFName()) < 0;
+		});
+		
+		ComboBox->RefreshOptions();
+	}
+
+	void SConditionComboWidget::OnSelectionChanged(TSharedPtr<FComboItem> NewSelection, ESelectInfo::Type SelectInfo)
+	{
+		if (NewSelection)
+		{
+			EErrorCode ErrorCode = Model->SetOperatorType(NewSelection->Handle, ConditionType);
+			if (ErrorCode != EErrorCode::Success)
+			{
+				UE_LOG(LogTypedElementDataStorage, Error, TEXT("Could not set model condition: [%d]"), static_cast<int>(ErrorCode));
+			}
+			ComboBox->ClearSelection();
+		}
+	}
+
+	TSharedRef<SWidget> SConditionComboWidget::OnGenerateWidget(TSharedPtr<FComboItem> Item)
+	{
+		const UScriptStruct* ColumnType = Model->GetColumnScriptStruct(Item->Handle);
+		const FText Label = ColumnType ? FText::FromString(ColumnType->GetName()) : FText::FromString(TEXT("Null Column Type"));
+
+		return SNew(STextBlock)
+			.Text(Label);
+	}
+}
diff --git a/Engine/Plugins/Experimental/EditorDataStorageFeatures/Source/TedsDebugger/Private/Widgets/QueryEditor/TedsConditionSelectionComboWidget.h b/Engine/Plugins/Experimental/EditorDataStorageFeatures/Source/TedsDebugger/Private/Widgets/QueryEditor/TedsConditionSelectionComboWidget.h
new file mode 100644
--- /dev/null
+++ b/Engine/Plugins/Experimental/EditorDataStorageFeatures/Source/TedsDebugger/Private/Widgets/QueryEditor/TedsConditionSelectionComboWidget.h
@@ -0,0 +1,43 @@
+// Copyright Epic Games, Inc. All Rights Reserved.
+#pragma once
+#include "Filters/GenericFilter.h"
+#include "Filters/SFilterBar.h"
+#include "QueryEditor/TedsQueryEditorModel.h"
+#include "Widgets/SCompoundWidget.h"
+
+template<typename>
+class SComboBox;
+
+namespace UE::EditorDataStorage::Debug::QueryEditor
+{
+	class FTedsQueryEditorModel;
+	struct FConditionEntry;
+
+	class SConditionComboWidget : public SCompoundWidget
+	{
+	public:
+		SLATE_BEGIN_ARGS( SConditionComboWidget ){}
+		SLATE_END_ARGS()
+
+		~SConditionComboWidget() override;
+		void OnConditionCollectionChanged();
+		void Construct(const FArguments& InArgs, FTedsQueryEditorModel& InModel, QueryEditor::EOperatorType InConditionType);
+	private:
+
+		struct FComboItem;
+		void PopulateComboItems();
+		void OnSelectionChanged(TSharedPtr<FComboItem> NewSelection, ESelectInfo::Type SelectInfo);
+		TSharedRef<SWidget> OnGenerateWidget(TSharedPtr<FComboItem> Item);
+
+		FTedsQueryEditorModel* Model = nullptr;
+		QueryEditor::EOperatorType ConditionType = QueryEditor::EOperatorType::Invalid;
+
+		FDelegateHandle ConditionCollectionChangedHandle;
+		bool bConditionCollectionDirty = true;
+
+		TArray<TSharedPtr<FComboItem>> ComboItems;
+		TSharedPtr<SComboBox<TSharedPtr<FComboItem>>> ComboBox;
+
+		TSharedPtr<SFilterBar<TSharedPtr<FComboItem>>> FilterThing;
+	};
+}
diff --git a/Engine/Plugins/Experimental/EditorDataStorageFeatures/Source/TedsDebugger/Private/Widgets/QueryEditor/TedsQueryEditorResultsView.cpp b/Engine/Plugins/Experimental/EditorDataStorageFeatures/Source/TedsDebugger/Private/Widgets/QueryEditor/TedsQueryEditorResultsView.cpp
new file mode 100644
--- /dev/null
+++ b/Engine/Plugins/Experimental/EditorDataStorageFeatures/Source/TedsDebugger/Private/Widgets/QueryEditor/TedsQueryEditorResultsView.cpp
@@ -0,0 +1,182 @@
+﻿// Copyright Epic Games, Inc. All Rights Reserved.
+
+#include "TedsQueryEditorResultsView.h"
+
+#include "SceneOutlinerPublicTypes.h"
+#include "SWarningOrErrorBox.h"
+#include "TedsOutlinerModule.h"
+#include "TypedElementOutlinerColumnIntegration.h"
+#include "TypedElementOutlinerMode.h"
+#include "Components/VerticalBox.h"
+#include "Elements/Columns/TypedElementLabelColumns.h"
+#include "Elements/Interfaces/TypedElementDataStorageInterface.h"
+#include "Elements/Framework/TypedElementQueryBuilder.h"
+#include "Modules/ModuleManager.h"
+#include "QueryEditor/TedsQueryEditorModel.h"
+#include "Widgets/Text/STextBlock.h"
+
+#define LOCTEXT_NAMESPACE "TedsDebuggerModule"
+
+using namespace UE::EditorDataStorage::Debug::QueryEditor;
+
+SResultsView::~SResultsView()
+{
+	Model->GetModelChangedDelegate().Remove(ModelChangedDelegateHandle);
+}
+
+void SResultsView::Construct(const FArguments& InArgs, FTedsQueryEditorModel& InModel)
+{
+	Model = &InModel;
+	ModelChangedDelegateHandle = Model->GetModelChangedDelegate().AddRaw(this, &SResultsView::OnModelChanged);
+	// RowQueryHandle = TypedElementDataStorage::InvalidQueryHandle;
+
+	ChildSlot
+	[
+		SNew(SVerticalBox)
+		+SVerticalBox::Slot()
+		.AutoHeight()
+		[
+			SNew(SWarningOrErrorBox)
+			.MessageStyle(EMessageStyle::Warning)
+			.Message(LOCTEXT("QueryEditor_TableViewUnreliable", "Integration of table view is WIP. Currently Unreliable"))
+		]
+		+SVerticalBox::Slot()
+		[
+			SAssignNew(TableViewHolder, SHorizontalBox)
+			+SHorizontalBox::Slot()
+			[
+				// Placeholder for Table View
+				SNullWidget::NullWidget
+			]
+		]
+		+SVerticalBox::Slot()
+		.AutoHeight()
+		[
+			SNew(STextBlock)
+			.Text_Lambda([this]()
+			{
+				auto& TedsInterface = Model->GetTedsInterface();
+				TypedElementDataStorage::FQueryResult QueryResult = TedsInterface.RunQuery(CountQueryHandle);
+				if (QueryResult.Completed == TypedElementDataStorage::FQueryResult::ECompletion::Fully)
+				{
+					FString String = FString::Printf(TEXT("Element Count: %u"), QueryResult.Count);
+					return FText::FromString(String);
+				}
+				else
+				{
+					return FText::FromString(TEXT("Invalid query"));
+				}
+			})
+		]
+	];
+}
+
+void SResultsView::Tick(const FGeometry& AllottedGeometry, const double InCurrentTime, const float InDeltaTime)
+{
+	using namespace TypedElementQueryBuilder;
+
+	ITypedElementDataStorageInterface& TedsInterface = Model->GetTedsInterface();
+
+	{
+		TypedElementDataStorage::FQueryDescription CountQueryDescription = Model->GenerateNoSelectQueryDescription();
+		
+		if (CountQueryHandle != TypedElementDataStorage::InvalidQueryHandle)
+		{
+			TedsInterface.UnregisterQuery(CountQueryHandle);
+			CountQueryHandle = TypedElementDataStorage::InvalidQueryHandle;
+		}
+		CountQueryHandle = TedsInterface.RegisterQuery(MoveTemp(CountQueryDescription));
+	}
+
+	{
+		// TODO: For now, destroy and re-create the table view
+		// This could be all that is required at some point
+		//
+		// FTypedElementSceneOutlinerQueryBinder::GetInstance().AssignQuery(SelectQueryHandle, TableView);
+		// TableView->FullRefresh();
+		
+		TableViewHolder->ClearChildren();
+
+		FSceneOutlinerInitializationOptions InitOptions;
+		InitOptions.bShowHeaderRow = true;
+		InitOptions.bShowSearchBox = false;
+		InitOptions.FilterBarOptions.bHasFilterBar = false;
+		InitOptions.OutlinerIdentifier = "QueryEditorResultsView";
+		
+		// if (RowQueryHandle != TypedElementDataStorage::InvalidQueryHandle)
+		// {
+		// 	TedsInterface.UnregisterQuery(RowQueryHandle);
+		// 	RowQueryHandle = TypedElementDataStorage::InvalidQueryHandle;
+		// } 
+		RowQueryDescription = Model->GenerateNoSelectQueryDescription();
+		if (!RowQueryDescription.ConditionTypes.IsEmpty())
+		{
+			// Work around since query registration takes by RRef but later the
+			// FTypedElementOutlinerModeParams.QueryDescription needs a reference to it.
+			// We are just going to copy it for now...
+			// TypedElementDataStorage::FQueryDescription Copy = SelectQueryDescription;
+			RowQueryDescription.Action = TypedElementDataStorage::FQueryDescription::EActionType::Select;
+			//RowQueryHandle = TedsInterface.RegisterQuery(MoveTemp(RowQueryDescription));
+		}
+	
+		FTypedElementOutlinerModeParams Params(nullptr);
+		Params.QueryDescription = TAttribute<TypedElementDataStorage::FQueryDescription>::CreateLambda([WeakResultsView = AsWeak()]()
+		{
+			if (TSharedPtr<SWidget> Shared = WeakResultsView.Pin())
+			{
+				SResultsView* This = static_cast<SResultsView*>(Shared.Get());
+				return This->RowQueryDescription;
+			}
+			return TypedElementDataStorage::FQueryDescription();
+		});
+		Params.HierarchyData = TOptional<FTypedElementOutlinerHierarchyData>(); // We don't want to show hierarchies
+		Params.CellWidgetPurposes = TArray<FName>{TEXT("General.Cell")};
+
+		TypedElementDataStorage::FQueryDescription ColumnQueryDescription = Model->GenerateQueryDescription();
+		if (ColumnQueryHandle != TypedElementDataStorage::InvalidQueryHandle)
+		{
+			TedsInterface.UnregisterQuery(ColumnQueryHandle);
+			ColumnQueryHandle = TypedElementDataStorage::InvalidQueryHandle;
+		}
+
+		if (!ColumnQueryDescription.ConditionTypes.IsEmpty() || !ColumnQueryDescription.SelectionTypes.IsEmpty())
+		{
+			// Work around since query registration takes by RRef but later the
+			// FTypedElementOutlinerModeParams.QueryDescription needs a reference to it.
+			// We are just going to copy it for now...
+			// TypedElementDataStorage::FQueryDescription Copy = SelectQueryDescription;
+			ColumnQueryDescription.Action = TypedElementDataStorage::FQueryDescription::EActionType::Select;
+			ColumnQueryHandle = TedsInterface.RegisterQuery(MoveTemp(ColumnQueryDescription));
+		}
+
+		FTedsOutlinerModule& TedsOutlinerModule = FModuleManager::GetModuleChecked<FTedsOutlinerModule>("TedsOutliner");
+		if (ColumnQueryHandle != TypedElementDataStorage::InvalidQueryHandle)
+		{
+			TSharedPtr<ISceneOutliner> TableView = TedsOutlinerModule.CreateTedsOutliner(InitOptions, Params, ColumnQueryHandle);
+
+			TableViewHolder->AddSlot()
+			[
+				TableView.ToSharedRef()
+			];
+		}
+		else
+		{
+			TableViewHolder->AddSlot()
+			[
+				// Placeholder for table view.  Mass doesn't like invalid/empty queries
+				SNullWidget::NullWidget
+			];
+		}
+	}
+
+	SetCanTick(false);
+}
+
+void SResultsView::OnModelChanged()
+{
+	SetCanTick(true);
+	Invalidate(EInvalidateWidgetReason::Layout);
+	
+}
+
+#undef LOCTEXT_NAMESPACE
\ No newline at end of file
diff --git a/Engine/Plugins/Experimental/EditorDataStorageFeatures/Source/TedsDebugger/Private/Widgets/QueryEditor/TedsQueryEditorResultsView.h b/Engine/Plugins/Experimental/EditorDataStorageFeatures/Source/TedsDebugger/Private/Widgets/QueryEditor/TedsQueryEditorResultsView.h
new file mode 100644
--- /dev/null
+++ b/Engine/Plugins/Experimental/EditorDataStorageFeatures/Source/TedsDebugger/Private/Widgets/QueryEditor/TedsQueryEditorResultsView.h
@@ -0,0 +1,42 @@
+// Copyright Epic Games, Inc. All Rights Reserved.
+#pragma once
+
+#include "Elements/Interfaces/TypedElementDataStorageInterface.h"
+#include "Widgets/SCompoundWidget.h"
+
+class ISceneOutliner;
+class SSceneOutliner;
+class SHorizontalBox;
+
+namespace UE::EditorDataStorage::Debug::QueryEditor
+{
+	class FTedsQueryEditorModel;
+
+	class SResultsView : public SCompoundWidget
+	{
+	public:
+		SLATE_BEGIN_ARGS( SResultsView ){}
+		SLATE_END_ARGS()
+
+		~SResultsView() override;
+		void Construct(const FArguments& InArgs, FTedsQueryEditorModel& InModel);
+		void Tick(const FGeometry& AllottedGeometry, const double InCurrentTime, const float InDeltaTime) override;
+
+	private:
+
+		void OnModelChanged();
+		
+		FTedsQueryEditorModel* Model = nullptr;
+		FDelegateHandle ModelChangedDelegateHandle;
+
+		TypedElementDataStorage::FQueryDescription RowQueryDescription;
+		TypedElementQueryHandle CountQueryHandle = TypedElementDataStorage::InvalidQueryHandle;
+		TypedElementQueryHandle ColumnQueryHandle = TypedElementDataStorage::InvalidQueryHandle;
+
+		TSharedPtr<SHorizontalBox> TableViewHolder;
+
+		
+	};
+
+	
+}
diff --git a/Engine/Plugins/Experimental/EditorDataStorageFeatures/Source/TedsDebugger/Private/Widgets/RowReferenceWidget.cpp b/Engine/Plugins/Experimental/EditorDataStorageFeatures/Source/TedsDebugger/Private/Widgets/RowReferenceWidget.cpp
new file mode 100644
--- /dev/null
+++ b/Engine/Plugins/Experimental/EditorDataStorageFeatures/Source/TedsDebugger/Private/Widgets/RowReferenceWidget.cpp
@@ -0,0 +1,179 @@
+// Copyright Epic Games, Inc. All Rights Reserved.
+
+#include "Widgets/RowReferenceWidget.h"
+
+#include "Columns/TedsOutlinerColumns.h"
+#include "Elements/Columns/TypedElementHiearchyColumns.h"
+#include "Elements/Columns/TypedElementLabelColumns.h"
+#include "Elements/Columns/TypedElementMiscColumns.h"
+#include "Elements/Interfaces/TypedElementDataStorageInterface.h"
+#include "ISceneOutliner.h"
+#include "Modules/ModuleManager.h"
+#include "TedsDebuggerModule.h"
+#include "Elements/Columns/TypedElementSlateWidgetColumns.h"
+#include "Elements/Framework/TypedElementRegistry.h"
+#include "Widgets/Input/SHyperlink.h"
+#include "Widgets/Layout/SBox.h"
+
+#define LOCTEXT_NAMESPACE "RowReferenceWidget"
+
+namespace UE::EditorDataStorage::Debug::Private
+{
+	void OnNavigateHyperlink(const ITypedElementDataStorageInterface* DataStorage, TypedElementDataStorage::RowHandle TargetRowHandle, TypedElementDataStorage::RowHandle UiRowHandle)
+	{
+		const FTableViewerColumn* TableViewerColumn = DataStorage->GetColumn<FTableViewerColumn>(UiRowHandle);
+
+		if(!TableViewerColumn)
+		{
+			return;
+		}
+		
+		TSharedPtr<ISceneOutliner> OwningTableViewer = TableViewerColumn->Outliner.Pin();
+		
+		if(!OwningTableViewer)
+		{
+			return;
+		}
+
+		// If the item was found in this table viewer, select it and navigate to it
+		if(FSceneOutlinerTreeItemPtr TreeItem = OwningTableViewer->GetTreeItem(TargetRowHandle))
+		{
+			OwningTableViewer->SetSelection([TreeItem](ISceneOutlinerTreeItem& Item)
+			{
+				return Item.GetID() == TreeItem->GetID();
+			});
+			
+			OwningTableViewer->FrameSelectedItems();
+
+			return;
+		}
+
+		// If it wasn't found in the table viewer owning this widget, navigate to it in the global TEDS debugger
+		UE::EditorDataStorage::Debug::FTedsDebuggerModule& TedsDebuggerModule = FModuleManager::GetModuleChecked<UE::EditorDataStorage::Debug::FTedsDebuggerModule>("TedsDebugger");
+		TedsDebuggerModule.NavigateToRow(TargetRowHandle);
+	}
+	
+	void CreateInternalWidget(const TWeakPtr<SWidget>& InWidget, TypedElementDataStorage::RowHandle UiRow, TypedElementDataStorage::RowHandle TargetRow)
+	{
+		const TSharedPtr<SWidget> Widget = InWidget.Pin();
+
+		if(!Widget)
+		{
+			return;
+		}
+		
+		checkf(Widget->GetType() == SBox::StaticWidgetClass().GetWidgetType(),
+			TEXT("Stored widget with FTypedElementLabelWidgetConstructor doesn't match type %s, but was a %s."),
+			*(SBox::StaticWidgetClass().GetWidgetType().ToString()),
+			*(Widget->GetTypeAsString()));
+
+		SBox* WidgetInstance = static_cast<SBox*>(Widget.Get());
+		WidgetInstance->SetContent(SNullWidget::NullWidget);
+
+		const ITypedElementDataStorageInterface* DataStorage = UTypedElementRegistry::GetInstance()->GetDataStorage();
+
+		// We only navigate to row references that have a label column
+		if (const FTypedElementLabelColumn* LabelColumn = DataStorage->GetColumn<FTypedElementLabelColumn>(TargetRow))
+		{
+			FNumberFormattingOptions NumberFormattingOptions;
+			NumberFormattingOptions.SetUseGrouping(false);
+
+			const FText Text = FText::AsNumber(TargetRow, &NumberFormattingOptions);
+			const FText TooltipText = FText::FromString(LabelColumn->Label);
+
+			TSharedRef<SWidget> HyperlinkWidget =
+				SNew(SHyperlink)
+					.Text(Text)
+					.Style(FAppStyle::Get(), "Common.GotoBlueprintHyperlink")
+					.ToolTipText(TooltipText)
+					.OnNavigate(FSimpleDelegate::CreateStatic(&UE::EditorDataStorage::Debug::Private::OnNavigateHyperlink, DataStorage, TargetRow, UiRow));
+			
+			WidgetInstance->SetContent(HyperlinkWidget);
+		}
+	}
+
+
+}
+
+URowReferenceWidgetFactory::~URowReferenceWidgetFactory()
+{
+}
+
+void URowReferenceWidgetFactory::RegisterWidgetConstructors(ITypedElementDataStorageInterface& DataStorage, ITypedElementDataStorageUiInterface& DataStorageUi) const
+{
+	// TEDS UI TODO: We can re-use this widget for FTypedElementParentColumn
+	DataStorageUi.RegisterWidgetFactory<FRowReferenceWidgetConstructor>(FName(TEXT("SceneOutliner.Cell")),
+	TypedElementDataStorage::FColumn<FTypedElementRowReferenceColumn>());
+}
+
+void URowReferenceWidgetFactory::RegisterQueries(ITypedElementDataStorageInterface& DataStorage)
+{
+	using namespace TypedElementQueryBuilder;
+	using DSI = ITypedElementDataStorageInterface;
+
+	const TypedElementQueryHandle UpdateRowReferenceWidget = DataStorage.RegisterQuery(
+		Select()
+			.ReadOnly<FTypedElementRowReferenceColumn>()
+		.Where()
+			.Any<FTypedElementSyncFromWorldTag, FTypedElementSyncBackToWorldTag>()
+		.Compile());
+
+	DataStorage.RegisterQuery(
+		Select(
+			TEXT("Sync row reference to widget"),
+			FProcessor(DSI::EQueryTickPhase::FrameEnd, DataStorage.GetQueryTickGroupName(DSI::EQueryTickGroups::SyncWidgets))
+				.ForceToGameThread(true),
+			[](
+				DSI::IQueryContext& Context,
+				TypedElementDataStorage::RowHandle UiRowHandle,
+				FTypedElementSlateWidgetReferenceColumn& Widget,
+				const FTypedElementRowReferenceColumn& Target)
+			{
+				Context.RunSubquery(0, Target.Row, CreateSubqueryCallbackBinding(
+					[&Widget, UiRowHandle](const FTypedElementRowReferenceColumn& Target)
+					{
+						UE::EditorDataStorage::Debug::Private::CreateInternalWidget(Widget.Widget, UiRowHandle, Target.Row);
+					}));
+			})
+		.DependsOn()
+			.SubQuery( UpdateRowReferenceWidget )
+		.Compile()
+	);
+}
+
+FRowReferenceWidgetConstructor::FRowReferenceWidgetConstructor()
+	: Super(FRowReferenceWidgetConstructor::StaticStruct())
+{
+}
+
+TSharedPtr<SWidget> FRowReferenceWidgetConstructor::CreateWidget(const TypedElementDataStorage::FMetaDataView& Arguments)
+{
+	return SNew(SBox)
+			.HAlign(HAlign_Left)
+			.VAlign(VAlign_Center);
+}
+
+bool FRowReferenceWidgetConstructor::FinalizeWidget(ITypedElementDataStorageInterface* DataStorage, ITypedElementDataStorageUiInterface* DataStorageUi, TypedElementDataStorage::RowHandle Row, const TSharedPtr<SWidget>& Widget)
+{
+	checkf(Widget, TEXT("Referenced widget is not valid. A constructed widget may not have been cleaned up. This can "
+		"also happen if this processor is running in the same phase as the processors responsible for cleaning up old "
+		"references."));
+
+	// The actual row we want to view in the widget
+	TypedElementRowHandle TargetRowReference = TypedElementInvalidRowHandle;
+
+	// The target row for which this widget was created
+	const TypedElementRowHandle TargetRow = DataStorage->GetColumn<FTypedElementRowReferenceColumn>(Row)->Row;
+
+	// Check if the target row has a row reference column, if so we want to view the row in there.
+	if(const FTypedElementRowReferenceColumn* RowReferenceColumn = DataStorage->GetColumn<FTypedElementRowReferenceColumn>(TargetRow))
+	{
+		TargetRowReference = RowReferenceColumn->Row;
+	}
+	
+	UE::EditorDataStorage::Debug::Private::CreateInternalWidget(Widget, Row, TargetRowReference);
+
+	return true;
+}
+
+#undef LOCTEXT_NAMESPACE
\ No newline at end of file
diff --git a/Engine/Plugins/Experimental/EditorDataStorageFeatures/Source/TedsDebugger/Private/Widgets/RowReferenceWidget.h b/Engine/Plugins/Experimental/EditorDataStorageFeatures/Source/TedsDebugger/Private/Widgets/RowReferenceWidget.h
new file mode 100644
--- /dev/null
+++ b/Engine/Plugins/Experimental/EditorDataStorageFeatures/Source/TedsDebugger/Private/Widgets/RowReferenceWidget.h
@@ -0,0 +1,43 @@
+// Copyright Epic Games, Inc. All Rights Reserved.
+
+#pragma once
+
+#include "Elements/Interfaces/TypedElementDataStorageFactory.h"
+#include "Elements/Interfaces/TypedElementDataStorageUiInterface.h"
+
+#include "RowReferenceWidget.generated.h"
+
+class ITypedElementDataStorageInterface;
+class SWidget;
+
+/*
+ * Widget for the TEDS Debugger that visualizes a reference to another row
+ */
+UCLASS()
+class URowReferenceWidgetFactory : public UTypedElementDataStorageFactory
+{
+	GENERATED_BODY()
+
+public:
+	~URowReferenceWidgetFactory() override;
+
+	void RegisterWidgetConstructors(ITypedElementDataStorageInterface& DataStorage,
+		ITypedElementDataStorageUiInterface& DataStorageUi) const override;
+	
+	void RegisterQueries(ITypedElementDataStorageInterface& DataStorage) override;
+};
+
+USTRUCT()
+struct FRowReferenceWidgetConstructor : public FTypedElementWidgetConstructor
+{
+	GENERATED_BODY()
+
+public:
+	FRowReferenceWidgetConstructor();
+	~FRowReferenceWidgetConstructor() override = default;
+
+protected:
+	TSharedPtr<SWidget> CreateWidget(const TypedElementDataStorage::FMetaDataView& Arguments) override;
+	bool FinalizeWidget(ITypedElementDataStorageInterface* DataStorage, ITypedElementDataStorageUiInterface* DataStorageUi,
+		TypedElementDataStorage::RowHandle Row, const TSharedPtr<SWidget>& Widget) override;
+};
\ No newline at end of file
diff --git a/Engine/Plugins/Experimental/EditorDataStorageFeatures/Source/TedsDebugger/Private/Widgets/SlateVisualizationWidget.cpp b/Engine/Plugins/Experimental/EditorDataStorageFeatures/Source/TedsDebugger/Private/Widgets/SlateVisualizationWidget.cpp
new file mode 100644
--- /dev/null
+++ b/Engine/Plugins/Experimental/EditorDataStorageFeatures/Source/TedsDebugger/Private/Widgets/SlateVisualizationWidget.cpp
@@ -0,0 +1,61 @@
+// Copyright Epic Games, Inc. All Rights Reserved.
+
+#include "Widgets/SlateVisualizationWidget.h"
+
+#include "Elements/Columns/TypedElementMiscColumns.h"
+#include "Elements/Columns/TypedElementSlateWidgetColumns.h"
+#include "Elements/Interfaces/TypedElementDataStorageInterface.h"
+#include "Widgets/Layout/SBox.h"
+#include "Widgets/SBoxPanel.h"
+
+
+void USlateVisualizationWidgetFactory::RegisterWidgetConstructors(ITypedElementDataStorageInterface& DataStorage, ITypedElementDataStorageUiInterface& DataStorageUi) const
+{
+	DataStorageUi.RegisterWidgetFactory<FSlateVisualizationWidgetConstructor>(FName(TEXT("SceneOutliner.Cell")),
+	TypedElementDataStorage::FColumn<FTypedElementSlateWidgetReferenceColumn>());
+}
+
+FSlateVisualizationWidgetConstructor::FSlateVisualizationWidgetConstructor()
+	: Super(FSlateVisualizationWidgetConstructor::StaticStruct())
+{
+
+}
+
+TSharedPtr<SWidget> FSlateVisualizationWidgetConstructor::CreateWidget(const TypedElementDataStorage::FMetaDataView& Arguments)
+{
+	return SNew(SHorizontalBox);
+}
+
+bool FSlateVisualizationWidgetConstructor::FinalizeWidget(ITypedElementDataStorageInterface* DataStorage, ITypedElementDataStorageUiInterface* DataStorageUi, TypedElementRowHandle Row, const TSharedPtr<SWidget>& Widget)
+{
+	checkf(Widget, TEXT("Referenced widget is not valid. A constructed widget may not have been cleaned up. This can "
+	"also happen if this processor is running in the same phase as the processors responsible for cleaning up old "
+	"references."));
+
+	TypedElementRowHandle TargetRow = DataStorage->GetColumn<FTypedElementRowReferenceColumn>(Row)->Row;
+
+	if (const FTypedElementSlateWidgetReferenceColumn* SlateWidgetReferenceColumn = DataStorage->GetColumn<FTypedElementSlateWidgetReferenceColumn>(TargetRow))
+	{
+		if(TSharedPtr<SWidget> ActualWidget = SlateWidgetReferenceColumn->Widget.Pin())
+		{
+			checkf(Widget->GetType() == SHorizontalBox::StaticWidgetClass().GetWidgetType(),
+				TEXT("Stored widget with FTypedElementLabelWidgetConstructor doesn't match type %s, but was a %s."),
+				*(SHorizontalBox::StaticWidgetClass().GetWidgetType().ToString()),
+				*(Widget->GetTypeAsString()));
+
+			SHorizontalBox* WidgetInstance = static_cast<SHorizontalBox*>(Widget.Get());
+
+			// Simply display the widget for now
+			// TEDS UI TODO: We can also display some debug info about the widget as a tooltip etc
+			WidgetInstance->AddSlot()
+			.AutoWidth()
+			.HAlign(HAlign_Center)
+			.VAlign(VAlign_Center)
+			[
+				ActualWidget.ToSharedRef()
+			];
+		}
+	}
+
+	return true;
+}
diff --git a/Engine/Plugins/Experimental/EditorDataStorageFeatures/Source/TedsDebugger/Private/Widgets/SlateVisualizationWidget.h b/Engine/Plugins/Experimental/EditorDataStorageFeatures/Source/TedsDebugger/Private/Widgets/SlateVisualizationWidget.h
new file mode 100644
--- /dev/null
+++ b/Engine/Plugins/Experimental/EditorDataStorageFeatures/Source/TedsDebugger/Private/Widgets/SlateVisualizationWidget.h
@@ -0,0 +1,41 @@
+// Copyright Epic Games, Inc. All Rights Reserved.
+
+#pragma once
+
+#include "Elements/Interfaces/TypedElementDataStorageFactory.h"
+#include "Elements/Interfaces/TypedElementDataStorageUiInterface.h"
+
+#include "SlateVisualizationWidget.generated.h"
+
+class ITypedElementDataStorageInterface;
+class SWidget;
+
+/*
+ * Widget for the TEDS Debugger that shows a slate widget reference
+ */
+UCLASS()
+class USlateVisualizationWidgetFactory : public UTypedElementDataStorageFactory
+{
+	GENERATED_BODY()
+
+public:
+	~USlateVisualizationWidgetFactory() override = default;
+
+	void RegisterWidgetConstructors(ITypedElementDataStorageInterface& DataStorage,
+		ITypedElementDataStorageUiInterface& DataStorageUi) const override;
+};
+
+USTRUCT()
+struct FSlateVisualizationWidgetConstructor : public FTypedElementWidgetConstructor
+{
+	GENERATED_BODY()
+
+public:
+	FSlateVisualizationWidgetConstructor();
+	~FSlateVisualizationWidgetConstructor() override = default;
+
+protected:
+	TSharedPtr<SWidget> CreateWidget(const TypedElementDataStorage::FMetaDataView& Arguments) override;
+	bool FinalizeWidget(ITypedElementDataStorageInterface* DataStorage, ITypedElementDataStorageUiInterface* DataStorageUi,
+		TypedElementRowHandle Row, const TSharedPtr<SWidget>& Widget) override;
+};
\ No newline at end of file
diff --git a/Engine/Plugins/Experimental/EditorDataStorageFeatures/Source/TedsDebugger/Private/Widgets/TypedElementRowReferenceWidget.cpp b/Engine/Plugins/Experimental/EditorDataStorageFeatures/Source/TedsDebugger/Private/Widgets/TypedElementRowReferenceWidget.cpp
deleted file mode 100644
diff --git a/Engine/Plugins/Experimental/EditorDataStorageFeatures/Source/TedsDebugger/Private/Widgets/TypedElementRowReferenceWidget.h b/Engine/Plugins/Experimental/EditorDataStorageFeatures/Source/TedsDebugger/Private/Widgets/TypedElementRowReferenceWidget.h
deleted file mode 100644
diff --git a/Engine/Plugins/Experimental/EditorDataStorageFeatures/Source/TedsDebugger/Private/Widgets/TypedElementSlateWidgetReferenceWidget.cpp b/Engine/Plugins/Experimental/EditorDataStorageFeatures/Source/TedsDebugger/Private/Widgets/TypedElementSlateWidgetReferenceWidget.cpp
deleted file mode 100644
diff --git a/Engine/Plugins/Experimental/EditorDataStorageFeatures/Source/TedsDebugger/Private/Widgets/TypedElementSlateWidgetReferenceWidget.h b/Engine/Plugins/Experimental/EditorDataStorageFeatures/Source/TedsDebugger/Private/Widgets/TypedElementSlateWidgetReferenceWidget.h
deleted file mode 100644
diff --git a/Engine/Plugins/Experimental/EditorDataStorageFeatures/Source/TedsDebugger/Public/TedsDebuggerModule.h b/Engine/Plugins/Experimental/EditorDataStorageFeatures/Source/TedsDebugger/Public/TedsDebuggerModule.h
deleted file mode 100644
diff --git a/Engine/Plugins/Experimental/EditorDataStorageFeatures/Source/TedsDebugger/TedsDebugger.Build.cs b/Engine/Plugins/Experimental/EditorDataStorageFeatures/Source/TedsDebugger/TedsDebugger.Build.cs
--- a/Engine/Plugins/Experimental/EditorDataStorageFeatures/Source/TedsDebugger/TedsDebugger.Build.cs
+++ b/Engine/Plugins/Experimental/EditorDataStorageFeatures/Source/TedsDebugger/TedsDebugger.Build.cs
@@ -6,7 +6,6 @@ public class TedsDebugger : ModuleRules
 {
 	public TedsDebugger(ReadOnlyTargetRules Target) : base(Target)
 	{
-		PrivateDependencyModuleNames.AddRange(new string[] { "EditorWidgets", "EditorWidgets", "EditorWidgets" });
 		PCHUsage = ModuleRules.PCHUsageMode.UseExplicitOrSharedPCHs;
 
 		if (Target.bBuildEditor)
