Optimize revision control overlay application query

Stacks queries into subqueries that only fire when necessary, drastically reduces time spent per frame. (~40ms on huge levels down to ~2us)

diff --git a/Engine/Plugins/Developer/PerforceSourceControl/Source/PerforceSourceControl/Private/PerforceSourceControlOperations.cpp b/Engine/Plugins/Developer/PerforceSourceControl/Source/PerforceSourceControl/Private/PerforceSourceControlOperations.cpp
--- a/Engine/Plugins/Developer/PerforceSourceControl/Source/PerforceSourceControl/Private/PerforceSourceControlOperations.cpp
+++ b/Engine/Plugins/Developer/PerforceSourceControl/Source/PerforceSourceControl/Private/PerforceSourceControlOperations.cpp
@@ -2021,6 +2021,7 @@ bool FPerforceUpdateStatusWorker::UpdateStates() const
 	auto UpdateDataStorage = [](const FPerforceSourceControlState& State)
 	{
 		using namespace TypedElementQueryBuilder;
+		using namespace TypedElementDataStorage;
 		using DSI = ITypedElementDataStorageInterface;
 
 		UTypedElementRegistry* Registry = UTypedElementRegistry::GetInstance();
@@ -2041,7 +2042,7 @@ bool FPerforceUpdateStatusWorker::UpdateStates() const
 			FPaths::NormalizeFilename(Filename);
 			Filename = FPaths::ConvertRelativePathToFull(Filename);
 			
-			uint64 Index = TypedElementDataStorage::GenerateIndexHash(Filename);
+			uint64 Index = GenerateIndexHash(Filename);
 			TypedElementRowHandle Row = DataStorage->FindIndexedRow(Index);
 
 			if (!DataStorage->IsRowAvailable(Row))
@@ -2121,6 +2122,16 @@ bool FPerforceUpdateStatusWorker::UpdateStates() const
 		{
 			DataStorage->RemoveColumn(Row, Column);
 		}
+		
+        if (FTypedElementPackageReference* BackReference = DataStorage->GetColumn<FTypedElementPackageReference>(Row))
+        {
+        	if (RowHandle ObjectRow = BackReference->Row; ObjectRow != InvalidRowHandle)
+        	{
+        		static TableHandle Table = DataStorage->FindTable(FName("Editor_PackageUpdateTable"));
+        		RowHandle UpdateRow = DataStorage->AddRow(Table);
+        		DataStorage->AddColumn(UpdateRow, FTypedElementPackageUpdateColumn{ .ObjectRow = ObjectRow, .PackageRow = Row });
+        	}
+        }
 	};
 
 	bool bUpdated = false;
diff --git a/Engine/Plugins/Experimental/EditorDataStorageFeatures/Source/TedsRevisionControl/Private/Compatibility/TypedElementUObjectPackagePathToColumnQueries.cpp b/Engine/Plugins/Experimental/EditorDataStorageFeatures/Source/TedsRevisionControl/Private/Compatibility/TypedElementUObjectPackagePathToColumnQueries.cpp
--- a/Engine/Plugins/Experimental/EditorDataStorageFeatures/Source/TedsRevisionControl/Private/Compatibility/TypedElementUObjectPackagePathToColumnQueries.cpp
+++ b/Engine/Plugins/Experimental/EditorDataStorageFeatures/Source/TedsRevisionControl/Private/Compatibility/TypedElementUObjectPackagePathToColumnQueries.cpp
@@ -32,6 +32,10 @@ static void ResolvePackageReference(ITypedElementDataStorageInterface::IQueryCon
 				
 	Context.AddColumn(PackageRow, MoveTemp(PathColumn));
 	Context.AddColumn(PackageRow, MoveTemp(LoadedPathColumn));
+
+	FTypedElementPackageReference PackageBackReference;
+	PackageBackReference.Row = Row;
+	Context.AddColumn(PackageRow, MoveTemp(PackageBackReference));
 };
 
 void UTypedElementUObjectPackagePathFactory::RegisterQueries(ITypedElementDataStorageInterface& DataStorage)
@@ -52,7 +56,6 @@ void UTypedElementUObjectPackagePathFactory::RegisterQueries(ITypedElementDataSt
 			}
 		}
 	);
-
 	DataStorage.RegisterQuery(
 		Select(
 			TEXT("Resolve package references"),
diff --git a/Engine/Plugins/Experimental/EditorDataStorageFeatures/Source/TedsRevisionControl/Private/RevisionControlProcessors.cpp b/Engine/Plugins/Experimental/EditorDataStorageFeatures/Source/TedsRevisionControl/Private/RevisionControlProcessors.cpp
--- a/Engine/Plugins/Experimental/EditorDataStorageFeatures/Source/TedsRevisionControl/Private/RevisionControlProcessors.cpp
+++ b/Engine/Plugins/Experimental/EditorDataStorageFeatures/Source/TedsRevisionControl/Private/RevisionControlProcessors.cpp
@@ -163,7 +163,11 @@ void UTypedElementRevisionControlFactory::RegisterTables(ITypedElementDataStorag
 		TTypedElementColumnTypeList<
 			FTypedElementPackagePathColumn, FTypedElementPackageLoadedPathColumn,
 			FSCCRevisionIdColumn, FSCCExternalRevisionIdColumn>(),
-		FName("Editor_RevisionControlTable"));
+		FName("Editor_RevisionControlTable"))
+		
+	DataStorage.RegisterTable(
+		TTypedElementColumnTypeList<FTypedElementPackageUpdateColumn>(),
+		FName("Editor_PackageUpdateTable"));
 }
 
 void UTypedElementRevisionControlFactory::RegisterQueries(ITypedElementDataStorageInterface& DataStorage)
@@ -256,91 +260,118 @@ void UTypedElementRevisionControlFactory::RegisterFetchUpdates(ITypedElementData
 
 void UTypedElementRevisionControlFactory::RegisterApplyOverlays(ITypedElementDataStorageInterface& DataStorage)
 {
+	using namespace TypedElementDataStorage;
 	using namespace TypedElementQueryBuilder;
 	using DSI = ITypedElementDataStorageInterface;
 
-	if (ApplyOverlaysObjectToSCC == TypedElementInvalidQueryHandle)
+	if (ApplyOverlaysObjectToSCC == InvalidQueryHandle)
 	{
 		ApplyOverlaysObjectToSCC = DataStorage.RegisterQuery(
 			Select()
 				.ReadOnly<FTypedElementPackagePathColumn>()
-				.ReadOnly<FSCCStatusColumn, FSCCExternallyLockedColumn, FTypedElementViewportOverlayColorColumn>(EOptional::Yes)
+				.ReadOnly<FSCCStatusColumn>(EOptional::Yes)
 			.Compile());
 		}
-
-	if (ApplyNewOverlays == TypedElementInvalidQueryHandle)
+ 
+	if (ApplyNewOverlays == InvalidQueryHandle)
 	{
-		// Query:
-		// For all actors without an overlay color column AND having a package reference:
-		//   Determine if a color should be applied based on SCC status tags
-		//   If so, add the OverlayColorColumn to the actor row
 		ApplyNewOverlays = DataStorage.RegisterQuery(
-			Select(
-				TEXT("Poll to add selection overlay colors based on SCC status"),
-				// This is in PrePhysics because the overlay->actor query is in DuringPhysics and contexts don't flush changes between tick groups
-				FProcessor(DSI::EQueryTickPhase::PrePhysics, DataStorage.GetQueryTickGroupName(DSI::EQueryTickGroups::SyncExternalToDataStorage))
-					.ForceToGameThread(true),
-				[](DSI::IQueryContext& Context, TypedElementRowHandle ObjectRow, const FTypedElementUObjectColumn& Actor, const FTypedElementPackageReference& PackageReference)
-				{
-					Context.RunSubquery(0, PackageReference.Row, CreateSubqueryCallbackBinding(
-						[&Context, &ObjectRow, &Actor](DSI::ISubqueryContext& SubQueryContext)
-						{
-							FColor Color = DetermineOverlayColor(Context, SubQueryContext, Actor);
-							if (Color.Bits != 0)
-							{
-								Context.AddColumn<FTypedElementViewportOverlayColorColumn>(ObjectRow, { .OverlayColor = Color });
-							}
-						})
-					);
-				}
-			)
+			Select()
+			.ReadOnly<FTypedElementUObjectColumn, FTypedElementPackageReference>()
 			.Where()
 				.All<FTypedElementActorTag>()
 				.None<FTypedElementViewportOverlayColorColumn>()
-			.DependsOn()
-				.SubQuery(ApplyOverlaysObjectToSCC)
 			.Compile()
 		);
 	}
-
-	if (ChangeOverlay == TypedElementInvalidQueryHandle)
+ 
+	if (ChangeOverlay == InvalidQueryHandle)
 	{
-		// Query:
-		// For all actors WITH an overlay color column AND having a package reference:
-		//   Re-check the color that should be applied based on SCC status tags
-		//   If the color has changed, remove and re-add the OverlayColorColumn to the actor row
-		//
-		// Note: Remove and re-add will trigger observer in TypedElementActorViewportProcessors to SetOverlayColor on the primitive components
 		ChangeOverlay = DataStorage.RegisterQuery(
+			Select()
+			.ReadOnly<FTypedElementUObjectColumn, FTypedElementPackageReference, FTypedElementViewportOverlayColorColumn>()
+			.Where()
+				.All<FTypedElementActorTag>()
+			.Compile()
+		);
+	}
+ 
+	if (FlushPackageUpdates == InvalidQueryHandle)
+	{
+		check(ApplyOverlaysObjectToSCC != InvalidQueryHandle && ApplyNewOverlays != InvalidQueryHandle && ChangeOverlay!= InvalidQueryHandle);
+		enum EFlushPackageUpdatesSubqueries
+		{
+			EApplyOverlaysObjectToSCC,
+			EApplyNewOverlays,
+			EChangeOverlay,
+			
+			Num
+		};
+		TArray<RowHandle> Subqueries;
+		Subqueries.AddUninitialized(EFlushPackageUpdatesSubqueries::Num);
+		Subqueries[EApplyOverlaysObjectToSCC] = ApplyOverlaysObjectToSCC;
+		Subqueries[EApplyNewOverlays] = ApplyNewOverlays;
+		Subqueries[EChangeOverlay] = ChangeOverlay;
+		
+		FlushPackageUpdates = DataStorage.RegisterQuery(
 			Select(
-				TEXT("Poll for changed selection overlay colors based on SCC status"),
-				// This is in PrePhysics because the overlay->actor query is in DuringPhysics and contexts don't flush changes between tick groups
-				FProcessor(DSI::EQueryTickPhase::PrePhysics, DataStorage.GetQueryTickGroupName(DSI::EQueryTickGroups::SyncExternalToDataStorage))
-					.ForceToGameThread(true),
-				[](DSI::IQueryContext& Context, TypedElementRowHandle ObjectRow, const FTypedElementUObjectColumn& Actor, const FTypedElementPackageReference& PackageReference, const FTypedElementViewportOverlayColorColumn& OverlayColorColumn)
+				TEXT("Consume collected package updates"),
+				FProcessor(DSI::EQueryTickPhase::PrePhysics, DataStorage.GetQueryTickGroupName(DSI::EQueryTickGroups::Update))
+				.ForceToGameThread(true),
+				[](DSI::IQueryContext& Context, RowHandle Row, const FTypedElementPackageUpdateColumn& Update)
 				{
-					Context.RunSubquery(0, PackageReference.Row, CreateSubqueryCallbackBinding(
-						[&Context, &ObjectRow, &Actor, &OverlayColorColumn](DSI::ISubqueryContext& SubQueryContext)
+					// Query:
+					// For all actors without an overlay color column AND having a package reference:
+					//   Determine if a color should be applied based on SCC status tags
+					//   If so, add the OverlayColorColumn to the actor row
+					Context.RunSubquery(EApplyNewOverlays, Update.ObjectRow, CreateSubqueryCallbackBinding(
+						[&Context](RowHandle ObjectRow, const FTypedElementUObjectColumn& Actor, const FTypedElementPackageReference& PackageReference)
 						{
-							FColor Color = DetermineOverlayColor(Context, SubQueryContext, Actor);
-							if (Color.Bits == 0)
-							{
-								Context.RemoveColumns<FTypedElementViewportOverlayColorColumn>(ObjectRow);
-							}
-							else if (Color != OverlayColorColumn.OverlayColor)
-							{
-								// Remove and re-add to trigger the observer
-								Context.RemoveColumns<FTypedElementViewportOverlayColorColumn>(ObjectRow);
-								Context.AddColumn<FTypedElementViewportOverlayColorColumn>(ObjectRow, { .OverlayColor = Color });
-							}
-						})
-					);
+							Context.RunSubquery(EApplyOverlaysObjectToSCC, PackageReference.Row, CreateSubqueryCallbackBinding(
+								[&Context, &ObjectRow, &Actor](DSI::ISubqueryContext& SubQueryContext)
+								{
+									FColor Color = DetermineOverlayColor(Context, SubQueryContext, Actor);
+									if (Color.Bits != 0)
+									{
+										Context.AddColumn<FTypedElementViewportOverlayColorColumn>(ObjectRow, { .OverlayColor = Color });
+									}
+								})
+							);
+						}
+					));
+ 
+					// Query:
+					// For all actors WITH an overlay color column AND having a package reference:
+					//   Re-check the color that should be applied based on SCC status tags
+					//   If the color has changed, remove and re-add the OverlayColorColumn to the actor row
+					//
+					// Note: Remove and re-add will trigger observer in TypedElementActorViewportProcessors to SetOverlayColor on the primitive components
+					Context.RunSubquery(EChangeOverlay, Update.ObjectRow, CreateSubqueryCallbackBinding(
+						[&Context](TypedElementRowHandle ObjectRow, const FTypedElementUObjectColumn& Actor, const FTypedElementPackageReference& PackageReference, const FTypedElementViewportOverlayColorColumn& OverlayColorColumn)
+						{
+							Context.RunSubquery(EApplyOverlaysObjectToSCC, PackageReference.Row, CreateSubqueryCallbackBinding(
+								[&Context, &ObjectRow, &Actor, &OverlayColorColumn](DSI::ISubqueryContext& SubQueryContext)
+								{
+									FColor Color = DetermineOverlayColor(Context, SubQueryContext, Actor);
+									if (Color.Bits == 0)
+									{
+										Context.RemoveColumns<FTypedElementViewportOverlayColorColumn>(ObjectRow);
+									}
+									else if (Color != OverlayColorColumn.OverlayColor)
+									{
+										// Remove and re-add to trigger the observer
+										Context.RemoveColumns<FTypedElementViewportOverlayColorColumn>(ObjectRow);
+										Context.AddColumn<FTypedElementViewportOverlayColorColumn>(ObjectRow, { .OverlayColor = Color });
+									}
+								})
+							);
+						}
+					));
+					Context.RemoveRow(Row);
 				}
 			)
-			.Where()
-				.All<FTypedElementActorTag>()
 			.DependsOn()
-				.SubQuery(ApplyOverlaysObjectToSCC)
+				.SubQuery(Subqueries)
 			.Compile()
 		);
 	}
diff --git a/Engine/Plugins/Experimental/EditorDataStorageFeatures/Source/TedsRevisionControl/Private/RevisionControlProcessors.h b/Engine/Plugins/Experimental/EditorDataStorageFeatures/Source/TedsRevisionControl/Private/RevisionControlProcessors.h
--- a/Engine/Plugins/Experimental/EditorDataStorageFeatures/Source/TedsRevisionControl/Private/RevisionControlProcessors.h
+++ b/Engine/Plugins/Experimental/EditorDataStorageFeatures/Source/TedsRevisionControl/Private/RevisionControlProcessors.h
@@ -28,4 +28,5 @@ private:
 	TypedElementDataStorage::QueryHandle ChangeOverlay = TypedElementDataStorage::InvalidQueryHandle;
 	TypedElementDataStorage::QueryHandle ApplyOverlaysObjectToSCC = TypedElementDataStorage::InvalidQueryHandle;
 	TypedElementDataStorage::QueryHandle RemoveOverlays = TypedElementDataStorage::InvalidQueryHandle;
+	TypedElementDataStorage::QueryHandle FlushPackageUpdates = TypedElementDataStorage::InvalidQueryHandle;
 };
diff --git a/Engine/Source/Runtime/TypedElementFramework/Public/Elements/Columns/TypedElementPackageColumns.h b/Engine/Source/Runtime/TypedElementFramework/Public/Elements/Columns/TypedElementPackageColumns.h
--- a/Engine/Source/Runtime/TypedElementFramework/Public/Elements/Columns/TypedElementPackageColumns.h
+++ b/Engine/Source/Runtime/TypedElementFramework/Public/Elements/Columns/TypedElementPackageColumns.h
@@ -34,6 +34,27 @@ struct FTypedElementPackageReference final : public FTypedElementDataStorageColu
 	TypedElementRowHandle Row;
 };
 
+/**
+ * Tag that indicates some related package information has been modified.
+ */
+USTRUCT(meta = (DisplayName = "Package information has been updated"))
+struct FTypedElementPackageUpdatedTag final : public FTypedElementDataStorageTag
+{
+	GENERATED_BODY()
+};
+
+/**
+ * Column that stores relevant rows when a package is updated
+ */
+USTRUCT()
+struct FTypedElementPackageUpdateColumn final : public FTypedElementDataStorageColumn
+{
+	GENERATED_BODY()
+	
+	TypedElementDataStorage::RowHandle ObjectRow;
+	TypedElementDataStorage::RowHandle PackageRow;
+};
+
 /**
  * Column that stores the path of a package.
  */
