Fix CIS error

diff --git a/Engine/Plugins/Experimental/EditorDataStorageFeatures/Source/TedsRevisionControl/Private/RevisionControlProcessors.cpp b/Engine/Plugins/Experimental/EditorDataStorageFeatures/Source/TedsRevisionControl/Private/RevisionControlProcessors.cpp
--- a/Engine/Plugins/Experimental/EditorDataStorageFeatures/Source/TedsRevisionControl/Private/RevisionControlProcessors.cpp
+++ b/Engine/Plugins/Experimental/EditorDataStorageFeatures/Source/TedsRevisionControl/Private/RevisionControlProcessors.cpp
@@ -163,7 +163,7 @@ void UTypedElementRevisionControlFactory::RegisterTables(ITypedElementDataStorag
 		TTypedElementColumnTypeList<
 			FTypedElementPackagePathColumn, FTypedElementPackageLoadedPathColumn,
 			FSCCRevisionIdColumn, FSCCExternalRevisionIdColumn>(),
-		FName("Editor_RevisionControlTable"))
+		FName("Editor_RevisionControlTable"));
 		
 	DataStorage.RegisterTable(
 		TTypedElementColumnTypeList<FTypedElementPackageUpdateColumn>(),
