First pass for TedsAssetDataCBDataSource. This will populate a few columns type needed for the assets row. At the moment we just want more data into teds to test some stuff. The code is not optimized at all and very inefficient, but this is fine since the feature is still of by default for now.

diff --git a/Engine/Plugins/Experimental/EditorDataStorageFeatures/Source/TedsAssetData/Private/CB/TedsAssetDataCBDataSource.cpp b/Engine/Plugins/Experimental/EditorDataStorageFeatures/Source/TedsAssetData/Private/CB/TedsAssetDataCBDataSource.cpp
new file mode 100644
--- /dev/null
+++ b/Engine/Plugins/Experimental/EditorDataStorageFeatures/Source/TedsAssetData/Private/CB/TedsAssetDataCBDataSource.cpp
@@ -0,0 +1,281 @@
+// Copyright Epic Games, Inc. All Rights Reserved.
+
+#include "TedsAssetDataCBDataSource.h"
+
+#include "TedsAssetDataColumns.h"
+
+#include "AssetRegistry/IAssetRegistry.h"
+#include "ContentBrowserDataUtils.h"
+#include "Elements/Common/TypedElementQueryTypes.h"
+#include "Elements/Framework/TypedElementQueryBuilder.h"
+#include "Elements/Interfaces/TypedElementDataStorageInterface.h"
+#include "Interfaces/IPluginManager.h"
+#include "Interfaces/IPluginManager.h"
+#include "Misc/PathViews.h"
+#include "PluginDescriptor.h"
+#include "Settings/ContentBrowserSettings.h"
+#include "UObject/NameTypes.h"
+
+namespace UE::EditorDataStorage::AssetData::Private
+{
+
+FTedsAssetDataCBDataSource::FTedsAssetDataCBDataSource(ITypedElementDataStorageInterface& InDatabase)
+	: Database(InDatabase)
+{
+	using namespace TypedElementQueryBuilder;
+	using namespace EditorDataStorage;
+
+	InitVirtualPathProcessor();
+
+	auto GenerateVirtualPaths = [DataSource = static_cast<const FTedsAssetDataCBDataSource*>(this)](const FStringView InAssetPath, FNameBuilder& OutVirtualizedPath) -> bool
+	{
+		if (!ContentBrowserDataUtils::PathPassesAttributeFilter(InAssetPath, 0, EContentBrowserItemAttributeFilter::IncludeAll))
+		{
+			return false;
+		}
+
+		DataSource->VirtualPathProcessor.ConvertInternalPathToVirtualPath(InAssetPath, OutVirtualizedPath);
+		return true;
+	};
+
+	ProcessPathQuery = Database.RegisterQuery(
+		Select(
+			TEXT("FTedsAssetDataCBDataSource: Process Path updates"),
+			FProcessor(TypedElementDataStorage::EQueryTickPhase::DuringPhysics, Database.GetQueryTickGroupName(TypedElementDataStorage::EQueryTickGroups::Update)),
+			[GenerateVirtualPaths](TypedElementDataStorage::IQueryContext& Context, const TypedElementDataStorage::RowHandle* Rows, const FAssetPathColumn_Experimental* PathColumn)
+			{
+				int32 NumOfRowToProcess = Context.GetRowCount();
+
+				FNameBuilder InternalPath;
+				FNameBuilder VirtualPath;
+
+				for (int32 Index = 0; Index < NumOfRowToProcess; ++Index)
+				{
+					PathColumn[Index].Path.ToString(InternalPath);
+
+					if (GenerateVirtualPaths(InternalPath, VirtualPath))
+					{
+						FVirtualPathColumn_Experimental VirtualPathColumn;
+						VirtualPathColumn.VirtualPath = *VirtualPath;
+						// Todo investigate for a batch add Maybe?
+						Context.AddColumn(Rows[Index], MoveTemp(VirtualPathColumn));
+					}
+				}
+			}
+		)
+		.Where()
+			.All<FUpdatedPathTag>()
+		.Compile()
+		);
+
+	ProcessAssetDataPathUpdateQuery = Database.RegisterQuery(
+		Select(
+			TEXT("FTedsAssetDataCBDataSource:: Process Asset Data Path Update"),
+			FProcessor(TypedElementDataStorage::EQueryTickPhase::DuringPhysics, Database.GetQueryTickGroupName(TypedElementDataStorage::EQueryTickGroups::Update)),
+			[GenerateVirtualPaths](TypedElementDataStorage::IQueryContext& Context, const TypedElementDataStorage::RowHandle Row, const FAssetDataColumn_Experimental& AssetDataColumn)
+			{
+				int32 NumOfRowToProcess = Context.GetRowCount();
+
+				FNameBuilder InternalPath;
+				FNameBuilder VirtualPath;
+
+				AssetDataColumn.AssetData.AppendObjectPath(InternalPath);
+
+				if (GenerateVirtualPaths(InternalPath, VirtualPath))
+				{
+					FVirtualPathColumn_Experimental VirtualPathColumn;
+					VirtualPathColumn.VirtualPath = *VirtualPath;
+					// Todo investigate for a batch add Maybe?
+					Context.AddColumn(Row, MoveTemp(VirtualPathColumn));
+				}
+			}
+		)
+		.Where()
+			.All<FUpdatedPathTag>()
+			.None<FUpdatedAssetDataTag>()
+		.Compile()
+		);
+
+
+	// For now just add the columns one by one but this should be rework to work in batch
+	auto AddAssetDataColumns = [DataSource = static_cast<const FTedsAssetDataCBDataSource*>(this), AssetRegistry = static_cast<const IAssetRegistry*>(&IAssetRegistry::GetChecked())](const FAssetData& InAssetData, TypedElementDataStorage::IQueryContext& InContext, TypedElementDataStorage::RowHandle Row)
+	{
+		// Not optimized at all but we would like to have the data in sooner for testing purposes.
+		
+		if (InAssetData.HasAnyPackageFlags(PKG_NotExternallyReferenceable))
+		{
+			// Private Asset
+			InContext.AddColumns<FAssetTag, FPrivateAssetTag>(Row);
+		}
+		else
+		{
+			InContext.AddColumns<FAssetTag, FPublicAssetTag>(Row);
+		}
+
+		if (TOptional<FAssetPackageData> PackageData = AssetRegistry->GetAssetPackageDataCopy(InAssetData.PackageName))
+		{
+			FDiskSizeColumn DiskSizeColumn;
+			DiskSizeColumn.DiskSize = PackageData->DiskSize;
+			InContext.AddColumn(Row, MoveTemp(DiskSizeColumn));
+		}
+
+		FItemNameColumn_Experimental ItemNameColumn;
+		ItemNameColumn.Name = InAssetData.AssetName;
+	};
+
+	ProcessAssetDataAndPathUpdateQuery = Database.RegisterQuery(
+		Select(
+			TEXT("FTedsAssetDataCBDataSource: Process Asset Data and Path Updates"),
+			FProcessor(TypedElementDataStorage::EQueryTickPhase::DuringPhysics, Database.GetQueryTickGroupName(TypedElementDataStorage::EQueryTickGroups::Update)),
+			[GenerateVirtualPaths, AddAssetDataColumns](TypedElementDataStorage::IQueryContext& Context, const TypedElementDataStorage::RowHandle Row, const FAssetDataColumn_Experimental& AssetDataColumn)
+			{
+
+				FNameBuilder InternalPath;
+				FNameBuilder VirtualPath;
+
+				InternalPath.Reset();
+				AssetDataColumn.AssetData.AppendObjectPath(InternalPath);
+
+				if (GenerateVirtualPaths(InternalPath, VirtualPath))
+				{
+					FVirtualPathColumn_Experimental VirtualPathColumn;
+					VirtualPathColumn.VirtualPath = *VirtualPath;
+					// Todo investigate for a batch add Maybe?
+					Context.AddColumn(Row, MoveTemp(VirtualPathColumn));
+
+					AddAssetDataColumns(AssetDataColumn.AssetData, Context, Row);
+				}
+			}
+		)
+		.Where()
+			.All<FUpdatedAssetDataTag, FUpdatedPathTag>()
+			.Compile()
+		);
+
+	ProcessAssetDataUpdateQuery = Database.RegisterQuery(
+		Select(
+			TEXT("FTedsAssetDataCBDataSource: Process Asset Data updates"),
+			FProcessor(TypedElementDataStorage::EQueryTickPhase::DuringPhysics, Database.GetQueryTickGroupName(TypedElementDataStorage::EQueryTickGroups::Update)),
+			[AddAssetDataColumns](TypedElementDataStorage::IQueryContext& Context, const TypedElementDataStorage::RowHandle Row, const FAssetDataColumn_Experimental& AssetDataColumn)
+			{
+				AddAssetDataColumns(AssetDataColumn.AssetData, Context, Row);
+			}
+		)
+		.Where()
+			.All<FUpdatedAssetDataTag>()
+			.None<FUpdatedPathTag>()
+		.Compile()
+		);
+
+}
+
+FTedsAssetDataCBDataSource::~FTedsAssetDataCBDataSource()
+{
+	// Not needed on a editor shut down
+	if (!IsEngineExitRequested())
+	{
+		Database.UnregisterQuery(ProcessAssetDataUpdateQuery);
+		Database.UnregisterQuery(ProcessAssetDataAndPathUpdateQuery);
+		Database.UnregisterQuery(ProcessAssetDataPathUpdateQuery);
+		Database.UnregisterQuery(ProcessPathQuery);
+
+		IPluginManager& PluginManager = IPluginManager::Get();
+		PluginManager.OnNewPluginContentMounted().RemoveAll(this);
+		PluginManager.OnPluginEdited().RemoveAll(this);
+		PluginManager.OnPluginUnmounted().RemoveAll(this);
+	}
+}
+
+void FTedsAssetDataCBDataSource::InitVirtualPathProcessor()
+{
+	IPluginManager& PluginManager = IPluginManager::Get();
+
+	PluginManager.OnNewPluginContentMounted().AddRaw(this, &FTedsAssetDataCBDataSource::OnPluginContentMounted);
+	PluginManager.OnPluginEdited().AddRaw(this, &FTedsAssetDataCBDataSource::OnPluginContentMounted);
+	PluginManager.OnPluginUnmounted().AddRaw(this, &FTedsAssetDataCBDataSource::OnPluginUnmounted);
+
+	TArray<TSharedRef<IPlugin>> EnabledPluginsWithContent = PluginManager.GetEnabledPluginsWithContent();
+	VirtualPathProcessor.PluginNameToCachedData.Reserve(EnabledPluginsWithContent.Num());
+
+	for (const TSharedRef<IPlugin>& Plugin : EnabledPluginsWithContent)
+	{
+		FVirtualPathProcessor::FCachedPluginData& Data = VirtualPathProcessor.PluginNameToCachedData.FindOrAdd(Plugin->GetName());
+		Data.LoadedFrom = Plugin->GetLoadedFrom();
+		Data.EditorCustomVirtualPath = Plugin->GetDescriptor().EditorCustomVirtualPath;
+	}
+
+	const UContentBrowserSettings* ContentBrowserSettings = GetDefault<UContentBrowserSettings>();
+	VirtualPathProcessor.bShowAllFolder = ContentBrowserSettings->bShowAllFolder;
+	VirtualPathProcessor.bOrganizeFolders = ContentBrowserSettings->bOrganizeFolders;
+}
+
+void FTedsAssetDataCBDataSource::OnPluginContentMounted(IPlugin& InPlugin)
+{
+	FVirtualPathProcessor::FCachedPluginData& Data = VirtualPathProcessor.PluginNameToCachedData.FindOrAdd(InPlugin.GetName());
+	Data.LoadedFrom = InPlugin.GetLoadedFrom();
+	Data.EditorCustomVirtualPath = InPlugin.GetDescriptor().EditorCustomVirtualPath;
+}
+
+void FTedsAssetDataCBDataSource::OnPluginUnmounted(IPlugin& InPlugin)
+{
+	VirtualPathProcessor.PluginNameToCachedData.Remove(InPlugin.GetName());
+}
+
+void FTedsAssetDataCBDataSource::FVirtualPathProcessor::ConvertInternalPathToVirtualPath(const FStringView InternalPath, FStringBuilderBase& OutVirtualPath) const
+{
+	OutVirtualPath.Reset();
+
+	if (bShowAllFolder)
+	{
+		OutVirtualPath.AppendChar(TEXT('/'));
+		if (InternalPath.Len() == 1 && InternalPath[0] == TEXT('/'))
+		{
+			return;
+		}
+	}
+
+	if (bOrganizeFolders && InternalPath.Len() > 1)
+	{
+		const FStringView MountPoint = FPathViews::GetMountPointNameFromPath(InternalPath);
+		const int32 MountPointHash = GetTypeHash(MountPoint);
+		if (const FTedsAssetDataCBDataSource::FVirtualPathProcessor::FCachedPluginData* Plugin = PluginNameToCachedData.FindByHash(MountPointHash, MountPoint))
+		{
+			if (Plugin->LoadedFrom == EPluginLoadedFrom::Engine)
+			{
+				OutVirtualPath.Append(TEXT("/EngineData/Plugins"));
+			}
+			else
+			{
+				OutVirtualPath.Append(TEXT("/Plugins"));
+			}
+
+			if (!Plugin->EditorCustomVirtualPath.IsEmpty())
+			{
+				int32 NumCharsToCopy = Plugin->EditorCustomVirtualPath.Len();
+				if (Plugin->EditorCustomVirtualPath[NumCharsToCopy - 1] == TEXT('/'))
+				{
+					--NumCharsToCopy;
+				}
+
+				if (NumCharsToCopy > 0)
+				{
+					if (Plugin->EditorCustomVirtualPath[0] != TEXT('/'))
+					{
+						OutVirtualPath.AppendChar(TEXT('/'));
+					}
+
+					OutVirtualPath.Append(*Plugin->EditorCustomVirtualPath, NumCharsToCopy);
+				}
+			}
+		}
+		else if (MountPoint.Equals(TEXT("Engine")))
+		{
+			OutVirtualPath.Append(TEXT("/EngineData"));
+		}
+
+	}
+
+	OutVirtualPath.Append(InternalPath.GetData(), InternalPath.Len());
+}
+
+} // End of Namespace UE::TypedElementDataStorageAssetData::Private
diff --git a/Engine/Plugins/Experimental/EditorDataStorageFeatures/Source/TedsAssetData/Private/CB/TedsAssetDataCBDataSource.h b/Engine/Plugins/Experimental/EditorDataStorageFeatures/Source/TedsAssetData/Private/CB/TedsAssetDataCBDataSource.h
new file mode 100644
--- /dev/null
+++ b/Engine/Plugins/Experimental/EditorDataStorageFeatures/Source/TedsAssetData/Private/CB/TedsAssetDataCBDataSource.h
@@ -0,0 +1,54 @@
+// Copyright Epic Games, Inc. All Rights Reserved.
+
+#pragma once
+
+#include "Elements/Common/TypedElementCommonTypes.h"
+#include "Elements/Common/TypedElementHandles.h"
+
+#include "Containers/StringFwd.h"
+#include "Interfaces/IPluginManager.h"
+
+class IPlugin;
+class ITypedElementDataStorageInterface;
+
+namespace UE::EditorDataStorage::AssetData::Private
+{
+
+class FTedsAssetDataCBDataSource
+{
+public:
+	explicit FTedsAssetDataCBDataSource(ITypedElementDataStorageInterface& InDatabase);
+	~FTedsAssetDataCBDataSource();
+		
+private:
+	ITypedElementDataStorageInterface& Database;
+	TypedElementDataStorage::QueryHandle ProcessPathQuery;
+	TypedElementDataStorage::QueryHandle ProcessAssetDataPathUpdateQuery;
+	TypedElementDataStorage::QueryHandle ProcessAssetDataAndPathUpdateQuery;
+	TypedElementDataStorage::QueryHandle ProcessAssetDataUpdateQuery;
+
+	struct FVirtualPathProcessor
+	{
+		struct FCachedPluginData
+		{
+			EPluginLoadedFrom LoadedFrom;
+			FString EditorCustomVirtualPath;
+		};
+
+		TMap<FString, FCachedPluginData> PluginNameToCachedData;
+
+		bool bShowAllFolder = false;
+		bool bOrganizeFolders = false;
+
+		void ConvertInternalPathToVirtualPath(const FStringView InternalPath, FStringBuilderBase& OutVirtualPath) const;
+	};
+
+	FVirtualPathProcessor VirtualPathProcessor; 
+
+	void InitVirtualPathProcessor();
+
+	void OnPluginContentMounted(IPlugin& InPlugin);
+	void OnPluginUnmounted(IPlugin& InPlugin);
+};
+
+} // End of namespace UE::TypedElementsDataStorageAssetData::Private
diff --git a/Engine/Plugins/Experimental/EditorDataStorageFeatures/Source/TedsAssetData/Private/TedsAssetData.cpp b/Engine/Plugins/Experimental/EditorDataStorageFeatures/Source/TedsAssetData/Private/TedsAssetData.cpp
--- a/Engine/Plugins/Experimental/EditorDataStorageFeatures/Source/TedsAssetData/Private/TedsAssetData.cpp
+++ b/Engine/Plugins/Experimental/EditorDataStorageFeatures/Source/TedsAssetData/Private/TedsAssetData.cpp
@@ -186,10 +186,39 @@ FTedsAssetData::FTedsAssetData(ITypedElementDataStorageInterface& InDatabase)
 	AssetsDataTable = Database.FindTable(FName(TEXT("Editor_AssetRegistryAssetDataTable")));
 	if (AssetsDataTable == TypedElementDataStorage::InvalidTableHandle)
 	{
-		AssetsDataTable = Database.RegisterTable<FAssetDataColumn_Experimental>(FName(TEXT("Editor_AssetRegistryAssetDataTable")));
+		AssetsDataTable = Database.RegisterTable<FAssetDataColumn_Experimental, FUpdatedPathTag, FUpdatedAssetDataTag>(FName(TEXT("Editor_AssetRegistryAssetDataTable")));
 	}
 
-	UpdateAssetsInPath = Database.RegisterQuery(
+
+	RemoveUpdatedPathTagQuery = Database.RegisterQuery(
+			TypedElementQueryBuilder::Select(
+				TEXT("FTedsAssetData: Remove Updated Path Tag"),
+				TypedElementQueryBuilder::FPhaseAmble(TypedElementQueryBuilder::FPhaseAmble::ELocation::Postamble, TypedElementDataStorage::EQueryTickPhase::FrameEnd),
+				[](TypedElementDataStorage::IQueryContext& Context, const TypedElementDataStorage::RowHandle* Rows)
+				{
+					Context.RemoveColumns<FUpdatedPathTag>(TConstArrayView<TypedElementRowHandle>(Rows, Context.GetRowCount()));
+				}
+			)
+			.Where()
+				.All<FUpdatedPathTag>()
+			.Compile()
+		);
+
+	RemoveUpdatedAssetDataTagQuery = Database.RegisterQuery(
+			TypedElementQueryBuilder::Select(
+				TEXT("FTedsAssetData: Remove Updated Asset Data Tag"),
+				TypedElementQueryBuilder::FPhaseAmble(TypedElementQueryBuilder::FPhaseAmble::ELocation::Postamble, TypedElementDataStorage::EQueryTickPhase::FrameEnd),
+				[](TypedElementDataStorage::IQueryContext& Context, const TypedElementDataStorage::RowHandle* Rows)
+				{
+					Context.RemoveColumns<FUpdatedAssetDataTag>(TConstArrayView<TypedElementRowHandle>(Rows, Context.GetRowCount()));
+				}
+			)
+			.Where()
+				.All<FUpdatedAssetDataTag>()
+				.Compile()
+		);
+
+	UpdateAssetsInPathQuery = Database.RegisterQuery(
 		TypedElementQueryBuilder::Select()
 			.ReadWrite<FAssetsInPathColumn_Experimental>()
 			.Compile()
@@ -197,7 +226,7 @@ FTedsAssetData::FTedsAssetData(ITypedElementDataStorageInterface& InDatabase)
 
 	ResolveMissingAssetInPathQuery = Database.RegisterQuery(
 		TypedElementQueryBuilder::Select(
-			TEXT("Resolve Missing Asset In Path"),
+			TEXT("FTedsAssetData: Resolve Missing Asset In Path"),
 			TypedElementQueryBuilder::FProcessor(TypedElementDataStorage::EQueryTickPhase::FrameEnd, Database.GetQueryTickGroupName(TypedElementDataStorage::EQueryTickGroups::Default)),
 			[this](TypedElementDataStorage::IQueryContext& Context, TypedElementDataStorage::RowHandle Row, const FUnresolvedAssetsInPathColumn_Experimental& UnresolvedAssetPath)
 			{
@@ -223,7 +252,7 @@ FTedsAssetData::FTedsAssetData(ITypedElementDataStorageInterface& InDatabase)
 			}
 		)
 		.DependsOn()
-			.SubQuery(UpdateAssetsInPath)
+			.SubQuery(UpdateAssetsInPathQuery)
 		.Compile()
 	);
 
@@ -235,7 +264,7 @@ FTedsAssetData::FTedsAssetData(ITypedElementDataStorageInterface& InDatabase)
 
 	ResolveMissingParentPathQuery = Database.RegisterQuery(
 		TypedElementQueryBuilder::Select(
-			TEXT("Resolve Missing Parent Path Row"),
+			TEXT("FTedsAssetData: Resolve Missing Parent Path Row"),
 			TypedElementQueryBuilder::FProcessor(TypedElementDataStorage::EQueryTickPhase::FrameEnd, Database.GetQueryTickGroupName(TypedElementDataStorage::EQueryTickGroups::Default)),
 			[this](TypedElementDataStorage::IQueryContext& Context, TypedElementDataStorage::RowHandle Row, const FUnresolvedParentAssetPathColumn_Experimental& UnresolvedParentAssetPath, FParentAssetPathColumn_Experimental& ParentAssetPathColumn)
 			{
@@ -346,26 +375,27 @@ FTedsAssetData::~FTedsAssetData()
 	LLM_SCOPE_BYNAME(TEXT("FTedsAssetData"))
 #endif
 
-	if (IAssetRegistry* AssetRegistry = IAssetRegistry::Get())
+	// Not needed on a editor shut down
+	if (!IsEngineExitRequested())
 	{
-		Database.UnregisterQuery(ResolveMissingParentPathQuery);
-		Database.UnregisterQuery(UpdateParentToChildrenAssetPathQuery);
-		Database.UnregisterQuery(ResolveMissingAssetInPathQuery);
-		Database.UnregisterQuery(UpdateAssetsInPath);
-
-
-		AssetRegistry->OnAssetsAdded().RemoveAll(this);
-		AssetRegistry->OnAssetsRemoved().RemoveAll(this);
-		AssetRegistry->OnAssetsUpdated().RemoveAll(this);
-		AssetRegistry->OnAssetsUpdatedOnDisk().RemoveAll(this);
-		AssetRegistry->OnAssetRenamed().RemoveAll(this);
-		AssetRegistry->OnPathsAdded().RemoveAll(this);
-		AssetRegistry->OnPathsRemoved().RemoveAll(this);
-	
-
-		// Not needed on a editor shut down
-		if (!IsEngineExitRequested())
+		if (IAssetRegistry* AssetRegistry = IAssetRegistry::Get())
 		{
+			Database.UnregisterQuery(ResolveMissingParentPathQuery);
+			Database.UnregisterQuery(UpdateParentToChildrenAssetPathQuery);
+			Database.UnregisterQuery(ResolveMissingAssetInPathQuery);
+			Database.UnregisterQuery(UpdateAssetsInPathQuery);
+			Database.UnregisterQuery(RemoveUpdatedAssetDataTagQuery);
+			Database.UnregisterQuery(RemoveUpdatedPathTagQuery);
+
+
+			AssetRegistry->OnAssetsAdded().RemoveAll(this);
+			AssetRegistry->OnAssetsRemoved().RemoveAll(this);
+			AssetRegistry->OnAssetsUpdated().RemoveAll(this);
+			AssetRegistry->OnAssetsUpdatedOnDisk().RemoveAll(this);
+			AssetRegistry->OnAssetRenamed().RemoveAll(this);
+			AssetRegistry->OnPathsAdded().RemoveAll(this);
+			AssetRegistry->OnPathsRemoved().RemoveAll(this);
+	
 			AssetRegistry->EnumerateAllCachedPaths([this](FName InPath)
 				{
 					const TypedElementDataStorage::IndexHash PathHash = TypedElementDataStorage::GenerateIndexHash(InPath);
@@ -499,6 +529,7 @@ void FTedsAssetData::OnAssetsUpdated(TConstArrayView<FAssetData> InAssetsUpdated
 		if (Database.IsRowAssigned(Row))
 		{
 			Database.GetColumn<FAssetDataColumn_Experimental>(Row)->AssetData = Asset;
+			Database.AddColumn<FUpdatedAssetDataTag>(Row);
 		}
 	}
 }
@@ -518,6 +549,7 @@ void FTedsAssetData::OnAssetsUpdatedOnDisk(TConstArrayView<FAssetData> InAssetsU
 		if (Database.IsRowAssigned(Row))
 		{
 			Database.GetColumn<FAssetDataColumn_Experimental>(Row)->AssetData = Asset;
+			Database.AddColumn<FUpdatedAssetDataTag>(Row);
 		}
 	}
 }
@@ -574,6 +606,7 @@ void FTedsAssetData::OnAssetRenamed(const FAssetData& InAsset, const FString& In
 			}
 		}
 
+		Database.AddColumn<FUpdatedPathTag>(Row);
 		Database.ReindexRow(OldAssetHash, NewAssetHash, Row);
 	}
 }
diff --git a/Engine/Plugins/Experimental/EditorDataStorageFeatures/Source/TedsAssetData/Private/TedsAssetData.h b/Engine/Plugins/Experimental/EditorDataStorageFeatures/Source/TedsAssetData/Private/TedsAssetData.h
--- a/Engine/Plugins/Experimental/EditorDataStorageFeatures/Source/TedsAssetData/Private/TedsAssetData.h
+++ b/Engine/Plugins/Experimental/EditorDataStorageFeatures/Source/TedsAssetData/Private/TedsAssetData.h
@@ -40,10 +40,13 @@ private:
 	TypedElementDataStorage::TableHandle PathsTable = TypedElementDataStorage::InvalidTableHandle;
 	TypedElementDataStorage::TableHandle AssetsDataTable = TypedElementDataStorage::InvalidTableHandle;
 
+	TypedElementDataStorage::QueryHandle UpdateAssetsInPathQuery = TypedElementDataStorage::InvalidQueryHandle;
 	TypedElementDataStorage::QueryHandle UpdateAssetsInPath = TypedElementDataStorage::InvalidQueryHandle;
 	TypedElementDataStorage::QueryHandle ResolveMissingAssetInPathQuery = TypedElementDataStorage::InvalidQueryHandle;
 	TypedElementDataStorage::QueryHandle UpdateParentToChildrenAssetPathQuery = TypedElementDataStorage::InvalidQueryHandle;
 	TypedElementDataStorage::QueryHandle ResolveMissingParentPathQuery = TypedElementDataStorage::InvalidQueryHandle;
+	TypedElementDataStorage::QueryHandle RemoveUpdatedPathTagQuery = TypedElementDataStorage::InvalidQueryHandle;
+	TypedElementDataStorage::QueryHandle RemoveUpdatedAssetDataTagQuery = TypedElementDataStorage::InvalidQueryHandle;
 };
 
 }
diff --git a/Engine/Plugins/Experimental/EditorDataStorageFeatures/Source/TedsAssetData/Private/TedsAssetDataModule.cpp b/Engine/Plugins/Experimental/EditorDataStorageFeatures/Source/TedsAssetData/Private/TedsAssetDataModule.cpp
--- a/Engine/Plugins/Experimental/EditorDataStorageFeatures/Source/TedsAssetData/Private/TedsAssetDataModule.cpp
+++ b/Engine/Plugins/Experimental/EditorDataStorageFeatures/Source/TedsAssetData/Private/TedsAssetDataModule.cpp
@@ -2,6 +2,7 @@
 
 #include "TedsAssetDataModule.h"
 
+#include "CB/TedsAssetDataCBDataSource.h"
 #include "Elements/Framework/TypedElementRegistry.h"
 #include "HAL/IConsoleManager.h"
 #include "Modules/ModuleManager.h"
@@ -107,7 +108,10 @@ void FTedsAssetDataModule::ProcessDependentEvents()
 
 void FTedsAssetDataModule::InitAssetRegistryStorage()
 {
-	AssetRegistryStorage = MakeUnique<Private::FTedsAssetData>(*UTypedElementRegistry::GetInstance()->GetMutableDataStorage());
+	ITypedElementDataStorageInterface& MutableDataStorage = *UTypedElementRegistry::GetInstance()->GetMutableDataStorage();
+
+	AssetDataCBDataSource = MakeUnique<Private::FTedsAssetDataCBDataSource>(MutableDataStorage);
+	AssetRegistryStorage = MakeUnique<Private::FTedsAssetData>(MutableDataStorage);
 }
 
 }
\ No newline at end of file
diff --git a/Engine/Plugins/Experimental/EditorDataStorageFeatures/Source/TedsAssetData/Private/TedsAssetDataModule.h b/Engine/Plugins/Experimental/EditorDataStorageFeatures/Source/TedsAssetData/Private/TedsAssetDataModule.h
--- a/Engine/Plugins/Experimental/EditorDataStorageFeatures/Source/TedsAssetData/Private/TedsAssetDataModule.h
+++ b/Engine/Plugins/Experimental/EditorDataStorageFeatures/Source/TedsAssetData/Private/TedsAssetDataModule.h
@@ -11,6 +11,7 @@ namespace UE::EditorDataStorage::AssetData
 namespace Private
 {
 	class FTedsAssetData;
+	class FTedsAssetDataCBDataSource;
 }
 
 class FTedsAssetDataModule : public IModuleInterface
@@ -36,6 +37,7 @@ public:
 private:
 	void InitAssetRegistryStorage();
 
+	TUniquePtr<Private::FTedsAssetDataCBDataSource> AssetDataCBDataSource;
 	TUniquePtr<Private::FTedsAssetData> AssetRegistryStorage;
 };
 
diff --git a/Engine/Plugins/Experimental/EditorDataStorageFeatures/Source/TedsAssetData/Public/TedsAssetDataColumns.h b/Engine/Plugins/Experimental/EditorDataStorageFeatures/Source/TedsAssetData/Public/TedsAssetDataColumns.h
--- a/Engine/Plugins/Experimental/EditorDataStorageFeatures/Source/TedsAssetData/Public/TedsAssetDataColumns.h
+++ b/Engine/Plugins/Experimental/EditorDataStorageFeatures/Source/TedsAssetData/Public/TedsAssetDataColumns.h
@@ -106,4 +106,36 @@ struct FVersePathColumn final : public FTypedElementDataStorageColumn
 	GENERATED_BODY()
 	
 	UE::Core::FVersePath VersePath;
+};
+
+// Used to notify the dependent queries of the update to the path of the row
+USTRUCT()
+struct FUpdatedPathTag : public FTypedElementDataStorageTag
+{
+	GENERATED_BODY()
+};
+
+// Used to notify the dependent queries of a update to the asset data
+USTRUCT()
+struct FUpdatedAssetDataTag : public FTypedElementDataStorageTag
+{
+	GENERATED_BODY()
+};
+
+USTRUCT(meta = (DisplayName = "CB Item Path"))
+struct FVirtualPathColumn_Experimental : public FTypedElementDataStorageColumn
+{
+	GENERATED_BODY()
+
+	UPROPERTY()
+	FName VirtualPath;
+};
+
+USTRUCT(meta = (DisplayName = "Name"))
+struct FItemNameColumn_Experimental : public FTypedElementDataStorageColumn
+{
+	GENERATED_BODY()
+
+	UPROPERTY()
+	FName Name;
 };
\ No newline at end of file
diff --git a/Engine/Plugins/Experimental/EditorDataStorageFeatures/Source/TedsAssetData/TedsAssetData.build.cs b/Engine/Plugins/Experimental/EditorDataStorageFeatures/Source/TedsAssetData/TedsAssetData.build.cs
--- a/Engine/Plugins/Experimental/EditorDataStorageFeatures/Source/TedsAssetData/TedsAssetData.build.cs
+++ b/Engine/Plugins/Experimental/EditorDataStorageFeatures/Source/TedsAssetData/TedsAssetData.build.cs
@@ -20,10 +20,18 @@ public class TedsAssetData : ModuleRules
 				new string[]
 				{
 					"AssetRegistry",
+					"ContentBrowserData",
 					"Core",
 					"CoreUObject",
 					"Engine",
 					"TypedElementFramework",
+					"UnrealEd",
+				});
+
+			PublicDependencyModuleNames.AddRange(
+				new string[]
+				{
+					"Projects",
 				});
 
 		}
diff --git a/Engine/Source/Editor/ContentBrowserData/Private/ContentBrowserDataUtils.cpp b/Engine/Source/Editor/ContentBrowserData/Private/ContentBrowserDataUtils.cpp
--- a/Engine/Source/Editor/ContentBrowserData/Private/ContentBrowserDataUtils.cpp
+++ b/Engine/Source/Editor/ContentBrowserData/Private/ContentBrowserDataUtils.cpp
@@ -72,6 +72,8 @@ int32 ContentBrowserDataUtils::GetMaxFolderDepthRequiredForAttributeFilter()
  
 bool ContentBrowserDataUtils::PathPassesAttributeFilter(const FStringView InPath, const int32 InAlreadyCheckedDepth, const EContentBrowserItemAttributeFilter InAttributeFilter)
 {
+	// Note Keep the implementation of that function pure as it is called from outside of the game thread
+
 	static const FString ProjectContentRootName = TEXT("Game");
 	static const FString EngineContentRootName = TEXT("Engine");
 	static const FString LocalizationFolderName = TEXT("L10N");
