Fixed an issue with the renaming of the TedsAssetData module.

diff --git a/Engine/Plugins/Experimental/EditorDataStorageFeatures/Source/TedsAssetData/Private/TedsAssetDataModule.cpp b/Engine/Plugins/Experimental/EditorDataStorageFeatures/Source/TedsAssetData/Private/TedsAssetDataModule.cpp
--- a/Engine/Plugins/Experimental/EditorDataStorageFeatures/Source/TedsAssetData/Private/TedsAssetDataModule.cpp
+++ b/Engine/Plugins/Experimental/EditorDataStorageFeatures/Source/TedsAssetData/Private/TedsAssetDataModule.cpp
@@ -50,12 +50,12 @@ void FTedsAssetDataModule::ShutdownModule()
 
 FTedsAssetDataModule* FTedsAssetDataModule::Get()
 {
-	return FModuleManager::Get().LoadModulePtr<FTedsAssetDataModule>(TEXT("TypedElementsDataStorageAssetData"));
+	return FModuleManager::Get().LoadModulePtr<FTedsAssetDataModule>(TEXT("TedsAssetData"));
 }
 
 FTedsAssetDataModule& FTedsAssetDataModule::GetChecked()
 {
-	return FModuleManager::Get().LoadModuleChecked<FTedsAssetDataModule>(TEXT("TypedElementsDataStorageAssetData"));
+	return FModuleManager::Get().LoadModuleChecked<FTedsAssetDataModule>(TEXT("TedsAssetData"));
 }
 
 void FTedsAssetDataModule::EnableTedsAssetRegistryStorage()
