Teds Table Viewer: Create a new Table Viewer widget that can be used to visualize rows and columns in TEDS

Create a default implementation of a row query stack that can be used to visualize an array of rows in the Table Viewer
Update the TEDS-Outliner column integration to use the table viewer column implementation internally
Add a "TEDS.UI.OpenTableViewer" and "TEDS.UI.OpenTedsOutliner" Cvar to summon a sample Table Viewer and Teds Outliner respectively

diff --git a/Engine/Plugins/Experimental/EditorDataStorageFeatures/EditorDataStorageFeatures.uplugin b/Engine/Plugins/Experimental/EditorDataStorageFeatures/EditorDataStorageFeatures.uplugin
--- a/Engine/Plugins/Experimental/EditorDataStorageFeatures/EditorDataStorageFeatures.uplugin
+++ b/Engine/Plugins/Experimental/EditorDataStorageFeatures/EditorDataStorageFeatures.uplugin
@@ -34,6 +34,10 @@
 		{
 			"Name": "TedsDebugger",
 			"Type": "Editor"
+		},
+		{
+			"Name": "TedsTableViewer",
+			"Type": "Editor"
 		}
 	],
 	"Plugins": [
diff --git a/Engine/Plugins/Experimental/EditorDataStorageFeatures/Source/TedsDebugger/Private/Widgets/RowReferenceWidget.cpp b/Engine/Plugins/Experimental/EditorDataStorageFeatures/Source/TedsDebugger/Private/Widgets/RowReferenceWidget.cpp
--- a/Engine/Plugins/Experimental/EditorDataStorageFeatures/Source/TedsDebugger/Private/Widgets/RowReferenceWidget.cpp
+++ b/Engine/Plugins/Experimental/EditorDataStorageFeatures/Source/TedsDebugger/Private/Widgets/RowReferenceWidget.cpp
@@ -21,14 +21,14 @@ namespace UE::EditorDataStorage::Debug::Private
 {
 	void OnNavigateHyperlink(const ITypedElementDataStorageInterface* DataStorage, TypedElementDataStorage::RowHandle TargetRowHandle, TypedElementDataStorage::RowHandle UiRowHandle)
 	{
-		const FTableViewerColumn* TableViewerColumn = DataStorage->GetColumn<FTableViewerColumn>(UiRowHandle);
+		const FTedsOutlinerColumn* TedsOutlinerColumn = DataStorage->GetColumn<FTedsOutlinerColumn>(UiRowHandle);
 
-		if(!TableViewerColumn)
+		if(!TedsOutlinerColumn)
 		{
 			return;
 		}
 		
-		TSharedPtr<ISceneOutliner> OwningTableViewer = TableViewerColumn->Outliner.Pin();
+		TSharedPtr<ISceneOutliner> OwningTableViewer = TedsOutlinerColumn->Outliner.Pin();
 		
 		if(!OwningTableViewer)
 		{
diff --git a/Engine/Plugins/Experimental/EditorDataStorageFeatures/Source/TedsOutliner/Private/Compatibility/TedsCompatibilityUtils.cpp b/Engine/Plugins/Experimental/EditorDataStorageFeatures/Source/TedsOutliner/Private/Compatibility/TedsCompatibilityUtils.cpp
--- a/Engine/Plugins/Experimental/EditorDataStorageFeatures/Source/TedsOutliner/Private/Compatibility/TedsCompatibilityUtils.cpp
+++ b/Engine/Plugins/Experimental/EditorDataStorageFeatures/Source/TedsOutliner/Private/Compatibility/TedsCompatibilityUtils.cpp
@@ -3,6 +3,7 @@
 #include "Compatibility/TedsCompatibilityUtils.h"
 
 #include "SSceneOutliner.h"
+#include "TedsTableViewerUtils.h"
 #include "Elements/Columns/TypedElementLabelColumns.h"
 #include "Elements/Columns/TypedElementMiscColumns.h"
 #include "Elements/Columns/TypedElementTypeInfoColumns.h"
@@ -235,14 +236,14 @@ TSharedRef<SWidget> FTedsOutlinerImpl::CreateLabelWidgetForItem(TypedElementRowH
 		TArray<TWeakObjectPtr<const UScriptStruct>> ColumnTypes(QueryDescription.SelectionTypes);
 		TSharedPtr<FTypedElementWidgetConstructor> CellWidgetConstructor = QueryConstructorPair.Value;
 
-		TypedElementRowHandle UiRowHandle = Storage->AddRow(Storage->FindTable(FTypedElementSceneOutlinerQueryBinder::CellWidgetTableName));
+		TypedElementRowHandle UiRowHandle = Storage->AddRow(Storage->FindTable(UE::EditorDataStorage::TableViewerUtils::GetWidgetTableName()));
 
 		if (FTypedElementRowReferenceColumn* RowReference = Storage->GetColumn<FTypedElementRowReferenceColumn>(UiRowHandle))
 		{
 			RowReference->Row = InRowHandle;
 		}
 
-		Storage->AddColumn(UiRowHandle, FTableViewerColumn{.Outliner = StaticCastSharedRef<ISceneOutliner>(SceneOutliner->AsShared())});
+		Storage->AddColumn(UiRowHandle, FTedsOutlinerColumn{.Outliner = StaticCastSharedRef<ISceneOutliner>(SceneOutliner->AsShared())});
 		
 		return StorageUi->ConstructWidget(UiRowHandle, *CellWidgetConstructor, MetaDataArgs);
 	};
@@ -352,9 +353,9 @@ bool FTedsOutlinerImpl::CanDisplayRow(TypedElementDataStorage::RowHandle ItemRow
 	if (Storage->HasColumns<FTypedElementSlateWidgetReferenceColumn>(ItemRowHandle))
 	{
 		// Check if this widget row belongs to the same table viewer it is being displayed in
-		if (const FTableViewerColumn* TableViewerColumn = Storage->GetColumn<FTableViewerColumn>(ItemRowHandle))
+		if (const FTedsOutlinerColumn* TedsOutlinerColumn = Storage->GetColumn<FTedsOutlinerColumn>(ItemRowHandle))
 		{
-			if (const TSharedPtr<ISceneOutliner> TableViewer = TableViewerColumn->Outliner.Pin())
+			if (const TSharedPtr<ISceneOutliner> TableViewer = TedsOutlinerColumn->Outliner.Pin())
 			{
 				return SceneOutliner != TableViewer.Get();
 			}
@@ -721,9 +722,9 @@ void FTedsOutlinerImpl::ClearSelection() const
 
 	// Query to remove the selection column from all rows that belong to this selection set
 	TypedElementDataStorage::DirectQueryCallback RowCollector = CreateDirectQueryCallbackBinding(
-	[this, &RowsToRemoveSelectionColumn](DSI::IDirectQueryContext& Context)
+	[this, &RowsToRemoveSelectionColumn](const DSI::IDirectQueryContext& Context, const TypedElementDataStorage::RowHandle* RowHandles)
 	{
-		TConstArrayView<TypedElementDataStorage::RowHandle> Rows = Context.GetRowHandles();
+		const TConstArrayView<TypedElementDataStorage::RowHandle> Rows(RowHandles, Context.GetRowCount());
 
 		for(const TypedElementDataStorage::RowHandle RowHandle : Rows)
 		{
diff --git a/Engine/Plugins/Experimental/EditorDataStorageFeatures/Source/TedsOutliner/Private/TedsOutlinerModule.cpp b/Engine/Plugins/Experimental/EditorDataStorageFeatures/Source/TedsOutliner/Private/TedsOutlinerModule.cpp
--- a/Engine/Plugins/Experimental/EditorDataStorageFeatures/Source/TedsOutliner/Private/TedsOutlinerModule.cpp
+++ b/Engine/Plugins/Experimental/EditorDataStorageFeatures/Source/TedsOutliner/Private/TedsOutlinerModule.cpp
@@ -49,7 +49,7 @@ namespace UE::TedsOutliner
 
 	// CVar to summon the TEDS-Outliner as a separate tab
 	static FAutoConsoleCommand OpenTableViewerConsoleCommand(
-		TEXT("TEDS.UI.OpenTableViewer"),
+		TEXT("TEDS.UI.OpenTedsOutliner"),
 		TEXT("Spawn the test TEDS-Outliner Integration."),
 		FConsoleCommandWithArgsDelegate::CreateLambda([](const TArray<FString>& Args)
 		{
diff --git a/Engine/Plugins/Experimental/EditorDataStorageFeatures/Source/TedsOutliner/Private/TypedElementOutlinerColumnIntegration.cpp b/Engine/Plugins/Experimental/EditorDataStorageFeatures/Source/TedsOutliner/Private/TypedElementOutlinerColumnIntegration.cpp
--- a/Engine/Plugins/Experimental/EditorDataStorageFeatures/Source/TedsOutliner/Private/TypedElementOutlinerColumnIntegration.cpp
+++ b/Engine/Plugins/Experimental/EditorDataStorageFeatures/Source/TedsOutliner/Private/TypedElementOutlinerColumnIntegration.cpp
@@ -20,13 +20,10 @@
 #include "SceneOutlinerFwd.h"
 #include "SceneOutlinerModule.h"
 #include "SceneOutlinerPublicTypes.h"
+#include "TedsTableViewerColumn.h"
+#include "TedsTableViewerUtils.h"
 #include "Elements/Interfaces/Capabilities/TypedElementUiTextCapability.h"
 #include "TypedElementOutlinerItem.h"
-#include "Columns/TedsOutlinerColumns.h"
-#include "Columns/UIPropertiesColumns.h"
-#include "Elements/Columns/TypedElementCompatibilityColumns.h"
-#include "Elements/Columns/TypedElementSlateWidgetColumns.h"
-#include "Elements/Framework/TypedElementDataStorageWidget.h"
 #include "Modules/ModuleManager.h"
 
 #define LOCTEXT_NAMESPACE "TypedElementsUI_SceneOutliner"
@@ -130,10 +127,7 @@ public:
 		FName InFallbackColumnName,
 		TWeakPtr<ISceneOutliner> InOwningOutliner,
 		const FTreeItemIDDealiaser& InDealiaser)
-		: ColumnTypes(MoveTemp(InColumnTypes))
-		, HeaderWidgetConstructor(MoveTemp(InHeaderWidgetConstructor))
-		, CellWidgetConstructor(MoveTemp(InCellWidgetConstructor))
-		, Storage(InStorage)
+		: Storage(InStorage)
 		, StorageUi(InStorageUi)
 		, StorageCompatibility(InStorageCompatibility)
 		, QueryHandle(InQuery)
@@ -142,20 +136,21 @@ public:
 		, Dealiaser(InDealiaser)
 	{
 		MetaData.AddOrSetMutableData(TEXT("Name"), NameId.ToString());
-		ColumnTypes.Shrink();
 
+		using namespace TypedElementDataStorage;
+
+		TableViewerColumnImpl = MakeUnique<UE::EditorDataStorage::FTedsTableViewerColumn>(NameId, InCellWidgetConstructor, InColumnTypes,
+			InHeaderWidgetConstructor, FComboMetaDataView(FGenericMetaDataView(MetaData)).Next(FQueryMetaDataView(Storage.GetQueryDescription(QueryHandle))));
+
+		TableViewerColumnImpl->SetIsRowVisibleDelegate(
+		UE::EditorDataStorage::FTedsTableViewerColumn::FIsRowVisible::CreateRaw(this, &FOutlinerColumn::IsRowVisible)
+		);
+		
 		// Try to find a fallback column from the regular item, for handling cases like folders which are not in TEDS but want to use TEDS columns
 		FSceneOutlinerModule& SceneOutlinerModule = FModuleManager::LoadModuleChecked<FSceneOutlinerModule>("SceneOutliner");
 		FallbackColumn = SceneOutlinerModule.FactoryColumn(InFallbackColumnName, *OwningOutliner.Pin());
-
-		RegisterQueries();
 	};
 	
-	~FOutlinerColumn() override
-	{
-		UnRegisterQueries();
-	}
-
 	FName GetColumnID() override
 	{
 		return NameId;
@@ -163,91 +158,7 @@ public:
 
 	virtual void Tick(double InCurrentTime, float InDeltaTime)
 	{
-		// Update any rows that could need widget updates
-		if(!RowsToUpdate.IsEmpty())
-		{
-			UpdateWidgets();
-			RowsToUpdate.Empty();
-		}
-	}
-
-	void RegisterQueries()
-	{
-		using namespace TypedElementQueryBuilder;
-		using namespace TypedElementDataStorage;
-
-		// For each TEDS column this column is matched with, we'll add observers to track addition/removal to update any widgets
-		for(const TWeakObjectPtr<const UScriptStruct>& ColumnType : ColumnTypes)
-		{
-			const FName ColumnAddObserverName = *(FString::Printf(TEXT("Column Add Monitor for %s Outliner Column, %s TEDS Column"), *NameId.ToString(), *ColumnType->GetName()));
-			FObserver AddObserver(FObserver::EEvent::Add, ColumnType.Get());
-			AddObserver.ForceToGameThread(true);
-
-			// TEDS-Outliner TODO: Long term if we move this into TypedElementOutlinerMode or similar we can get access to the exact
-			// types the Outliner is looking at and specify them on .Where() to cut down on the things we are observing
-			TypedElementDataStorage::QueryHandle AddQueryHandle = Storage.RegisterQuery(
-				Select(
-					ColumnAddObserverName,
-					AddObserver,
-					[this](IQueryContext& Context, TypedElementRowHandle Row)
-						{
-							RowsToUpdate.Add(TPair<TypedElementRowHandle, bool>(Row, true));
-						})
-				.Where()
-					.All(ColumnType.Get())
-				.Compile()
-				);
-			
-			InternalObserverQueries.Add(AddQueryHandle);
-
-			const FName ColumnRemoveObserverName = *(FString::Printf(TEXT("Column Remove Monitor for %s Outliner Column, %s TEDS Column"), *NameId.ToString(), *ColumnType->GetName()));
-			FObserver RemoveObserver(FObserver::EEvent::Remove, ColumnType.Get());
-			RemoveObserver.ForceToGameThread(true);
-
-			// TEDS-Outliner TODO: Long term if we move this into TypedElementOutlinerMode or similar we can get access to the exact
-			// types the Outliner is looking at and specify them on .Where() to cut down on the things we are observing
-			TypedElementDataStorage::QueryHandle RemoveQueryHandle = Storage.RegisterQuery(
-				Select(
-					ColumnRemoveObserverName,
-					RemoveObserver,
-					[this](IQueryContext& Context, TypedElementRowHandle Row)
-						{
-							RowsToUpdate.Add(TPair<TypedElementRowHandle, bool>(Row, false));
-						})
-				.Where()
-					.All(ColumnType.Get())
-				.Compile()
-			);
-			
-			InternalObserverQueries.Add(RemoveQueryHandle);
-		}
-
-		// We are looking for widgets that have a row reference
-		TArray<const UScriptStruct*> SelectionColumns({FTypedElementSlateWidgetReferenceColumn::StaticStruct(), FTypedElementRowReferenceColumn::StaticStruct()});
-
-		// We need to remove duplicates because TEDS/Mass does not handle having the same column in Select() and Where()
-		TArray<const UScriptStruct*> AdditionalWidgetColumns = CellWidgetConstructor->GetAdditionalColumnsList().FilterByPredicate([&SelectionColumns](const UScriptStruct* Column)
-		{
-			return !SelectionColumns.Contains(Column);
-		});
-		
-		// Query to get all widgets that were created by this column
-		WidgetQuery = Storage.RegisterQuery(
-				Select()
-					.ReadOnly(SelectionColumns)
-				.Where()
-					.All(AdditionalWidgetColumns)
-				.Compile());
-	}
-	
-	void UnRegisterQueries()
-	{
-		for(const TypedElementDataStorage::QueryHandle Query : InternalObserverQueries)
-		{
-			Storage.UnregisterQuery(Query);
-		}
-		
-		Storage.UnregisterQuery(WidgetQuery);
+		TableViewerColumnImpl->Tick();
 	}
 
 	bool IsRowVisible(const TypedElementDataStorage::RowHandle InRowHandle) const
@@ -280,145 +191,9 @@ public:
 		return OutlinerPinned->GetTree().IsItemVisible(Item);
 	}
 
-	void UpdateWidgets()
-	{
-		// Remove any widget rows that don't actually need an update
-		RowsToUpdate = RowsToUpdate.FilterByPredicate([this](const TPair<TypedElementDataStorage::RowHandle, bool>& Pair) -> bool
-		{
-			// We don't have a widget for this item visible, so there is nothing to update
-			if(!IsRowVisible(Pair.Key))
-			{
-				return false;
-			}
-			
-			// Check if the row now matches the query conditions for this widget
-			const bool bMatchesQueryConditions = CellWidgetConstructor->GetQueryConditions() && Storage.MatchesColumns(Pair.Key, *CellWidgetConstructor->GetQueryConditions());
-			
-			// If we are adding a column that we are monitoring and it now matches, or if we are removing a column that we are monitoring and it now
-			// stops matching, there is a potential need for widget update
-			return (bMatchesQueryConditions && Pair.Value) || (!bMatchesQueryConditions && !Pair.Value);
-		});
-		
-		using namespace TypedElementQueryBuilder;
-		using namespace TypedElementDataStorage;
-
-		// Query to find all widgets that belong to the row handles that need updates
-		DirectQueryCallback RowCollector = CreateDirectQueryCallbackBinding(
-		[this](const IDirectQueryContext& Context, const FTypedElementSlateWidgetReferenceColumn* ContainerWidgetReferenceColumns, const FTypedElementRowReferenceColumn* RowReferenceColumns)
-		{
-			const FTypedElementSlateWidgetReferenceColumn* WidgetsIt = ContainerWidgetReferenceColumns;
-			const FTypedElementRowReferenceColumn* RowRefsIt = RowReferenceColumns;
-			const TConstArrayView<TypedElementRowHandle> Rows = Context.GetRowHandles();
-
-			for(unsigned RowIndex = 0; RowIndex < Context.GetRowCount(); ++RowIndex, ++WidgetsIt, ++RowRefsIt)
-			{
-				// Check if this widgets owning row is in our rows to update
-				bool* bColumnAddedPtr = RowsToUpdate.Find(RowRefsIt->Row);
-				// If not, skip it
-				if(!bColumnAddedPtr)
-				{
-					continue;
-				}
-
-				// Check if the container TEDSWidget exists, if not we cannot update this widget
-				const TSharedPtr<STedsWidget> TedsWidget = WidgetsIt->TedsWidget.Pin();
-				if(!TedsWidget)
-				{
-					continue;
-				}
-
-				// A row has numerous widgets, make sure we only update the one that was created by our column by checking the constructor
-				if(WidgetsIt->WidgetConstructor != CellWidgetConstructor)
-				{
-					continue;
-				}
-
-				// If a column was added and we are here, we need to re-create the widget
-				// TEDS-Outliner TODO: Do we need to create the widget only if it doesn't exist? Or should we also update it to automatically respond
-				// to column changes even if it was already created
-				if(*bColumnAddedPtr)
-				{
-					const TSharedPtr<SWidget> RowWidget = CellWidgetConstructor->Construct(Rows[RowIndex], &Storage, &StorageUi,
-						FComboMetaDataView(FGenericMetaDataView(MetaData)).Next(FQueryMetaDataView(Storage.GetQueryDescription(QueryHandle))));
-
-					if(RowWidget)
-					{
-						TedsWidget->SetContent(RowWidget.ToSharedRef());
-					}
-				}
-				// If a column was removed (and we don't match anymore) delete the internal widget
-				else
-				{
-					TedsWidget->SetContent(SNullWidget::NullWidget);
-				}
-			}
-		});
-
-		Storage.RunQuery(WidgetQuery, RowCollector);
-	}
-
 	virtual SHeaderRow::FColumn::FArguments ConstructHeaderRowColumn() override
 	{
-		using namespace TypedElementDataStorage;
-
-		FString TooltipText = TEXT("Data Storage columns:");
-		for (const TWeakObjectPtr<const UScriptStruct>& ColumnType : ColumnTypes)
-		{
-			if (ColumnType.IsValid())
-			{
-				TooltipText += TEXT("\n    ");
-				ColumnType->AppendName(TooltipText);
-			}
-		}
-
-		TSharedPtr<SWidget> Widget;
-		RowHandle UiRowHandle = InvalidRowHandle;
-		if (HeaderWidgetConstructor)
-		{
-			UiRowHandle = Storage.AddRow(Storage.FindTable(FName(TEXT("Editor_WidgetTable"))));
-
-			// TEDS UI TODO: We can't do this from the Widget Constructor because it is a UStruct and does not have access to AsShared(), so we would
-			// be forced to store a raw pointer instead of a weak pointer which is unsafe. Once the widget construction pipleline is improved this can
-			// probably be moved to a better place
-			if(FTypedElementSlateWidgetReferenceColumn* WidgetReferenceColumn = Storage.GetColumn<FTypedElementSlateWidgetReferenceColumn>(UiRowHandle))
-			{
-				WidgetReferenceColumn->WidgetConstructor = HeaderWidgetConstructor;
-			}
-			
-			Widget = StorageUi.ConstructWidget(UiRowHandle, *HeaderWidgetConstructor, 
-				FComboMetaDataView(FGenericMetaDataView(MetaData)).Next(FQueryMetaDataView(Storage.GetQueryDescription(QueryHandle))));
-		}
-		if (!Widget.IsValid())
-		{
-			Widget = SNew(STextBlock)
-				.Text(FText::FromString(NameId.ToString()));
-		}
-		
-		SHeaderRow::FColumn::FArguments Column = SHeaderRow::Column(NameId)
-			.FillWidth(1)
-			.HeaderComboVisibility(EHeaderComboVisibility::OnHover)
-			.DefaultTooltip(FText::FromString(MoveTemp(TooltipText)))
-			.HeaderContent()
-			[
-				SNew(SHorizontalBox)
-				+ SHorizontalBox::Slot()
-				.VAlign(VAlign_Center)
-				[
-					Widget.ToSharedRef()
-				]
-			];
-		if (const FUIHeaderPropertiesColumn* HeaderProperties = Storage.GetColumn<FUIHeaderPropertiesColumn>(UiRowHandle))
-		{
-			float Width = HeaderProperties->Width;
-			switch (HeaderProperties->ColumnSizeMode)
-			{
-				case EColumnSizeMode::Fill: Column.FillWidth(Width); break;
-				case EColumnSizeMode::Fixed: Column.FixedWidth(Width); break;
-				case EColumnSizeMode::Manual: Column.ManualWidth(Width); break;
-				case EColumnSizeMode::FillSized: Column.FillSized(Width); break;
-			}
-		}
-		return Column;
+		return TableViewerColumnImpl->ConstructHeaderRowColumn();
 	}
 
 	// TODO: Sorting is currently handled through the fallback column if it exists because we have no way to sort columns through TEDS
@@ -485,28 +260,7 @@ public:
 
 		if(Storage.IsRowAssigned(RowHandle))
 		{
-			TypedElementRowHandle UiRowHandle = Storage.AddRow(Storage.FindTable(FTypedElementSceneOutlinerQueryBinder::CellWidgetTableName));
-
-			if (ColumnTypes.Num() == 1)
-			{
-				Storage.AddColumn<FTypedElementScriptStructTypeInfoColumn>(UiRowHandle, FTypedElementScriptStructTypeInfoColumn{.TypeInfo = *ColumnTypes.begin()});
-			}
-			
-			if (FTypedElementRowReferenceColumn* RowReference = Storage.GetColumn<FTypedElementRowReferenceColumn>(UiRowHandle))
-			{
-				RowReference->Row = RowHandle;
-			}
-
-			Storage.AddColumn(UiRowHandle, FTableViewerColumn{.Outliner = OwningOutliner});
-			
-			if(FTypedElementSlateWidgetReferenceColumn* WidgetReferenceColumn = Storage.GetColumn<FTypedElementSlateWidgetReferenceColumn>(UiRowHandle))
-			{
-				WidgetReferenceColumn->WidgetConstructor = CellWidgetConstructor;
-			}
-			
-			RowWidget = StorageUi.ConstructWidget(UiRowHandle, *CellWidgetConstructor, 
-							FComboMetaDataView(FGenericMetaDataView(MetaData)).Next(FQueryMetaDataView(Storage.GetQueryDescription(QueryHandle))));
-			
+			RowWidget = TableViewerColumnImpl->ConstructRowWidget(RowHandle);
 		}
 
 		if(RowWidget)
@@ -526,11 +280,10 @@ public:
 			FallbackColumn->PopulateSearchStrings(Item, OutSearchStrings);
 		}
 	}
-	
 
-	TArray<TWeakObjectPtr<const UScriptStruct>> ColumnTypes;
-	TSharedPtr<FTypedElementWidgetConstructor> HeaderWidgetConstructor;
-	TSharedPtr<FTypedElementWidgetConstructor> CellWidgetConstructor;
+	// The table viewer implementation that we internally use to create our widgets
+	TUniquePtr<UE::EditorDataStorage::FTedsTableViewerColumn> TableViewerColumnImpl;
+	
 	ITypedElementDataStorageInterface& Storage;
 	ITypedElementDataStorageUiInterface& StorageUi;
 	ITypedElementDataStorageCompatibilityInterface& StorageCompatibility;
@@ -540,10 +293,6 @@ public:
 	TSharedPtr<ISceneOutlinerColumn> FallbackColumn;
 	TWeakPtr<ISceneOutliner> OwningOutliner;
 	FTreeItemIDDealiaser Dealiaser;
-	
-	TArray<TypedElementDataStorage::QueryHandle> InternalObserverQueries;
-	TypedElementDataStorage::QueryHandle WidgetQuery;
-	TMap<TypedElementDataStorage::RowHandle, bool> RowsToUpdate;
 };
 
 
@@ -552,19 +301,6 @@ public:
 //
 // UTypedElementSceneOutlinerFactory
 // 
-void UTypedElementSceneOutlinerFactory::RegisterTables(ITypedElementDataStorageInterface& DataStorage)
-{
-	TypedElementTableHandle BaseWidgetTable = DataStorage.FindTable(FName(TEXT("Editor_WidgetTable")));
-	if (BaseWidgetTable != TypedElementInvalidTableHandle)
-	{
-		DataStorage.RegisterTable(
-			BaseWidgetTable,
-			{
-				FTypedElementRowReferenceColumn::StaticStruct()
-			},
-			FTypedElementSceneOutlinerQueryBinder::CellWidgetTableName);
-	}
-}
 
 void UTypedElementSceneOutlinerFactory::RegisterWidgetPurposes(ITypedElementDataStorageUiInterface& DataStorageUi) const
 {
@@ -705,82 +441,6 @@ void FTypedElementSceneOutliner::Initialize(
 	Outliner = InOutliner;
 }
 
-TArray<TWeakObjectPtr<const UScriptStruct>> FTypedElementSceneOutliner::CreateVerifiedColumnTypeAray(
-	TConstArrayView<TWeakObjectPtr<const UScriptStruct>> ColumnTypes)
-{
-	TArray<TWeakObjectPtr<const UScriptStruct>> VerifiedColumnTypes;
-	VerifiedColumnTypes.Reserve(ColumnTypes.Num());
-	for (const TWeakObjectPtr<const UScriptStruct>& ColumnType : ColumnTypes)
-	{
-		if (ColumnType.IsValid())
-		{
-			VerifiedColumnTypes.Add(ColumnType.Get());
-		}
-	}
-	return VerifiedColumnTypes;
-}
-
-TSharedPtr<FTypedElementWidgetConstructor> FTypedElementSceneOutliner::CreateHeaderWidgetConstructor(
-	ITypedElementDataStorageInterface& Storage, ITypedElementDataStorageUiInterface& StorageUi, 
-	TypedElementQueryHandle Query, TConstArrayView<TWeakObjectPtr<const UScriptStruct>> ColumnTypes,
-	const TConstArrayView<FName> CellWidgetPurposes)
-{
-	using MatchApproach = ITypedElementDataStorageUiInterface::EMatchApproach;
-
-	const ITypedElementDataStorageInterface::FQueryDescription& Description = Storage.GetQueryDescription(Query);
-	TypedElementDataStorage::FQueryMetaDataView MetaDataView(Description);
-
-	TArray<TWeakObjectPtr<const UScriptStruct>> VerifiedColumnTypes = CreateVerifiedColumnTypeAray(ColumnTypes);
-	TSharedPtr<FTypedElementWidgetConstructor> Constructor;
-
-	for (const FName& Purpose : CellWidgetPurposes)
-	{
-		// Extract the header purpose from the cell purpose (e.g "SceneOutliner.ItemLabel.Cell" -> "SceneOutliner.ItemLabel.Header")
-		FString HeaderPurpose, CellPurpose;
-		Purpose.ToString().Split(TEXT(".Cell"), &HeaderPurpose, &CellPurpose);
-		HeaderPurpose.Append(TEXT(".Header"));
-		
-		StorageUi.CreateWidgetConstructors(FName(HeaderPurpose), MatchApproach::ExactMatch, VerifiedColumnTypes, MetaDataView,
-			[&Constructor, ColumnTypes](
-				TUniquePtr<FTypedElementWidgetConstructor> CreatedConstructor, 
-				TConstArrayView<TWeakObjectPtr<const UScriptStruct>> MatchedColumnTypes)
-			{
-				if (ColumnTypes.Num() == MatchedColumnTypes.Num())
-				{
-					Constructor = TSharedPtr<FTypedElementWidgetConstructor>(CreatedConstructor.Release());
-				}
-				// Either this was the exact match so no need to search further or the longest possible chain didn't match so the next ones will 
-				// always be shorter in both cases just return.
-				return false;
-			});
-		if (Constructor)
-		{
-			return Constructor;
-		}
-	}
-	for (const FName& Purpose : CellWidgetPurposes)
-	{
-		// Extract the default header purpose from the cell purpose (e.g "SceneOutliner.ItemLabel.Cell" -> "SceneOutliner.ItemLabel.Header.Default")
-		FString HeaderPurpose, CellPurpose;
-		Purpose.ToString().Split(TEXT(".Cell"), &HeaderPurpose, &CellPurpose);
-		HeaderPurpose.Append(TEXT(".Header.Default"));
-
-		StorageUi.CreateWidgetConstructors(FName(HeaderPurpose), MetaDataView,
-			[&Constructor, ColumnTypes](
-				TUniquePtr<FTypedElementWidgetConstructor> CreatedConstructor,
-				TConstArrayView<TWeakObjectPtr<const UScriptStruct>> MatchedColumnTypes)
-			{
-				Constructor = TSharedPtr<FTypedElementWidgetConstructor>(CreatedConstructor.Release());
-				return false;
-			});
-		if (Constructor)
-		{
-			return Constructor;
-		}
-	}
-	return nullptr;
-}
-
 void FTypedElementSceneOutliner::RegisterDealiaser(const FTreeItemIDDealiaser& InDealiaser)
 {
 	Dealiaser = InDealiaser;
@@ -793,6 +453,7 @@ void FTypedElementSceneOutliner::AssignQuery(TypedElementQueryHandle Query, cons
 	FTypedElementSceneOutlinerQueryBinder& Binder = FTypedElementSceneOutlinerQueryBinder::GetInstance();
 	CellWidgetPurposes = InCellWidgetPurposes;
 
+
 	if (TSharedPtr<ISceneOutliner> OutlinerPinned = Outliner.Pin())
 	{
 		const ITypedElementDataStorageInterface::FQueryDescription& Description = Storage->GetQueryDescription(Query);
@@ -805,10 +466,10 @@ void FTypedElementSceneOutliner::AssignQuery(TypedElementQueryHandle Query, cons
 			int32 SelectionCount = Description.SelectionTypes.Num();
 			AddedColumns.Reset(SelectionCount);
 
-			TArray<TWeakObjectPtr<const UScriptStruct>> ColumnTypes = CreateVerifiedColumnTypeAray(Description.SelectionTypes);
+			TArray<TWeakObjectPtr<const UScriptStruct>> ColumnTypes = UE::EditorDataStorage::TableViewerUtils::CreateVerifiedColumnTypeArray(Description.SelectionTypes);
 
 			int32 IndexOffset = 0;
-			auto ColumnConstructor = [this, Query, &IndexOffset, &OutlinerPinned, Binder](
+			auto ColumnConstructor = [this, Query, MetaDataView, &IndexOffset, &OutlinerPinned, Binder](
 				TUniquePtr<FTypedElementWidgetConstructor> Constructor, TConstArrayView<TWeakObjectPtr<const UScriptStruct>> ColumnTypes)
 				{
 					TSharedPtr<FTypedElementWidgetConstructor> CellConstructor(Constructor.Release());
@@ -823,15 +484,15 @@ void FTypedElementSceneOutliner::AssignQuery(TypedElementQueryHandle Query, cons
 
 					OutlinerPinned->RemoveColumn(FallbackColumn);
 
-					FName NameId = FindLongestMatchingName(ColumnTypes, IndexOffset);
+					FName NameId = UE::EditorDataStorage::TableViewerUtils::FindLongestMatchingName(ColumnTypes, IndexOffset);
 					AddedColumns.Add(NameId);
 					OutlinerPinned->AddColumn(NameId,
 						FSceneOutlinerColumnInfo(ESceneOutlinerColumnVisibility::Visible, ColumnPriority,
 							FCreateSceneOutlinerColumn::CreateLambda(
-								[this, Query, NameId, &ColumnTypes, CellConstructor, &OutlinerPinned, FallbackColumn](ISceneOutliner&)
+								[this, Query, MetaDataView, NameId, &ColumnTypes, CellConstructor, &OutlinerPinned, FallbackColumn](ISceneOutliner&)
 								{
 									TSharedPtr<FTypedElementWidgetConstructor> HeaderConstructor = 
-										CreateHeaderWidgetConstructor(*Storage, *StorageUi, Query, ColumnTypes, CellWidgetPurposes);
+										UE::EditorDataStorage::TableViewerUtils::CreateHeaderWidgetConstructor(*StorageUi, MetaDataView, ColumnTypes, CellWidgetPurposes);
 									return MakeShared<FOutlinerColumn>(
 										Query, *Storage, *StorageUi, *StorageCompatibility, NameId,
 										TArray<TWeakObjectPtr<const UScriptStruct>>(ColumnTypes.GetData(), ColumnTypes.Num()), 
@@ -854,7 +515,7 @@ void FTypedElementSceneOutliner::AssignQuery(TypedElementQueryHandle Query, cons
 			{
 				FName FallbackColumn = Binder.FindOutlinerColumnFromTEDSColumns({ColumnType});
 
-				auto AssignWidgetToColumn = [this, Query, ColumnType, &IndexOffset, &OutlinerPinned, FallbackColumn](
+				auto AssignWidgetToColumn = [this, Query, MetaDataView, ColumnType, &IndexOffset, &OutlinerPinned, FallbackColumn](
 					TUniquePtr<FTypedElementWidgetConstructor> Constructor, TConstArrayView<TWeakObjectPtr<const UScriptStruct>>)
 				{
 					TSharedPtr<FTypedElementWidgetConstructor> CellConstructor(Constructor.Release());
@@ -863,12 +524,12 @@ void FTypedElementSceneOutliner::AssignQuery(TypedElementQueryHandle Query, cons
 					OutlinerPinned->AddColumn(NameId,
 						FSceneOutlinerColumnInfo(ESceneOutlinerColumnVisibility::Visible, DefaultPriorityIndex + IndexOffset,
 							FCreateSceneOutlinerColumn::CreateLambda(
-								[this, Query, NameId, ColumnType, CellConstructor, &OutlinerPinned, FallbackColumn](ISceneOutliner&)
+								[this, Query, MetaDataView, NameId, ColumnType, CellConstructor, &OutlinerPinned, FallbackColumn](ISceneOutliner&)
 								{
 									TArray<TWeakObjectPtr<const UScriptStruct>> ColumnTypesStored;
 									ColumnTypesStored.Add(ColumnType);
 									TSharedPtr<FTypedElementWidgetConstructor> HeaderConstructor =
-										CreateHeaderWidgetConstructor(*Storage, *StorageUi, Query, { ColumnType }, CellWidgetPurposes);
+										UE::EditorDataStorage::TableViewerUtils::CreateHeaderWidgetConstructor(*StorageUi, MetaDataView, { ColumnType }, CellWidgetPurposes);
 									return MakeShared<FOutlinerColumn>(
 										Query, *Storage, *StorageUi, *StorageCompatibility, NameId, MoveTemp(ColumnTypesStored),
 										HeaderConstructor, CellConstructor, FallbackColumn, OutlinerPinned, Dealiaser);
@@ -902,58 +563,6 @@ void FTypedElementSceneOutliner::AssignQuery(TypedElementQueryHandle Query, cons
 	}
 }
 
-// TEDS UI TODO: Maybe the widget can specify a user facing name derived from the matched columns instead of trying to find the longest matching name
-FName FTypedElementSceneOutliner::FindLongestMatchingName(TConstArrayView<TWeakObjectPtr<const UScriptStruct>> ColumnTypes, int32 DefaultNameIndex)
-{
-	switch (ColumnTypes.Num())
-	{
-	case 0:
-		return FName(TEXT("Column"), DefaultNameIndex);
-	case 1:
-		return FName(ColumnTypes[0]->GetDisplayNameText().ToString());
-	default:
-	{
-		FText LongestMatchText = ColumnTypes[0]->GetDisplayNameText();
-		FStringView LongestMatch = LongestMatchText.ToString();
-		const TWeakObjectPtr<const UScriptStruct>* ItEnd = ColumnTypes.end();
-		const TWeakObjectPtr<const UScriptStruct>* It = ColumnTypes.begin();
-		++It; // Skip the first entry as that's already set.
-		for (; It != ItEnd; ++It)
-		{
-			FText NextMatchText = (*It)->GetDisplayNameText();
-			FStringView NextMatch = NextMatchText.ToString();
-
-			int32 MatchSize = 0;
-			auto ItLeft = LongestMatch.begin();
-			auto ItLeftEnd = LongestMatch.end();
-			auto ItRight = NextMatch.begin();
-			auto ItRightEnd = NextMatch.end();
-			while (
-				ItLeft != ItLeftEnd &&
-				ItRight != ItRightEnd &&
-				*ItLeft == *ItRight)
-			{
-				++MatchSize;
-				++ItLeft;
-				++ItRight;
-			}
-
-			// At least 3 letters have to match to avoid single or double letter names which typically mean nothing.
-			if (MatchSize > 2)
-			{
-				LongestMatch.LeftInline(MatchSize);
-			}
-			else
-			{
-				// There are not enough characters in the string that match. Just return the name of the first column
-				return FName(ColumnTypes[0]->GetDisplayNameText().ToString());
-			}
-		}
-		return FName(LongestMatch);
-	}
-	};
-}
-
 void FTypedElementSceneOutliner::ClearColumns(ISceneOutliner& InOutliner)
 {
 	for (FName ColumnName : AddedColumns)
diff --git a/Engine/Plugins/Experimental/EditorDataStorageFeatures/Source/TedsOutliner/Public/Columns/TedsOutlinerColumns.h b/Engine/Plugins/Experimental/EditorDataStorageFeatures/Source/TedsOutliner/Public/Columns/TedsOutlinerColumns.h
--- a/Engine/Plugins/Experimental/EditorDataStorageFeatures/Source/TedsOutliner/Public/Columns/TedsOutlinerColumns.h
+++ b/Engine/Plugins/Experimental/EditorDataStorageFeatures/Source/TedsOutliner/Public/Columns/TedsOutlinerColumns.h
@@ -9,13 +9,12 @@
 
 class ISceneOutliner;
 
-// Column used to store a reference to the table viewer owning a specific row
-// Currently only added to widget rows in a table viewer
+// Column used to store a reference to the Teds Outliner owning a specific row
+// Currently only added to widget rows in the Teds Outliner
 USTRUCT(meta = (DisplayName = "Owning Table Viewer"))
-struct FTableViewerColumn final : public FTypedElementDataStorageColumn
+struct FTedsOutlinerColumn final : public FTypedElementDataStorageColumn
 {
 	GENERATED_BODY()
 
-	// The table viewer is internally a scene outliner
 	TWeakPtr<ISceneOutliner> Outliner;
 };
diff --git a/Engine/Plugins/Experimental/EditorDataStorageFeatures/Source/TedsOutliner/Public/Columns/UIPropertiesColumns.h b/Engine/Plugins/Experimental/EditorDataStorageFeatures/Source/TedsOutliner/Public/Columns/UIPropertiesColumns.h
deleted file mode 100644
diff --git a/Engine/Plugins/Experimental/EditorDataStorageFeatures/Source/TedsOutliner/Public/TypedElementOutlinerColumnIntegration.h b/Engine/Plugins/Experimental/EditorDataStorageFeatures/Source/TedsOutliner/Public/TypedElementOutlinerColumnIntegration.h
--- a/Engine/Plugins/Experimental/EditorDataStorageFeatures/Source/TedsOutliner/Public/TypedElementOutlinerColumnIntegration.h
+++ b/Engine/Plugins/Experimental/EditorDataStorageFeatures/Source/TedsOutliner/Public/TypedElementOutlinerColumnIntegration.h
@@ -33,12 +33,6 @@ public:
 	void AssignQuery(TypedElementQueryHandle Query, const TConstArrayView<FName> CellWidgetPurposes);
 	void RegisterDealiaser(const FTreeItemIDDealiaser& InDealiaser);
 private:
-	static FName FindLongestMatchingName(TConstArrayView<TWeakObjectPtr<const UScriptStruct>> ColumnTypes, int32 DefaultNameIndex);
-	static TArray<TWeakObjectPtr<const UScriptStruct>> CreateVerifiedColumnTypeAray(
-		TConstArrayView<TWeakObjectPtr<const UScriptStruct>> ColumnTypes);
-	static TSharedPtr<FTypedElementWidgetConstructor> CreateHeaderWidgetConstructor(ITypedElementDataStorageInterface& Storage,
-		ITypedElementDataStorageUiInterface& StorageUI, TypedElementQueryHandle Query,
-		TConstArrayView<TWeakObjectPtr<const UScriptStruct>> ColumnTypes, const TConstArrayView<FName> CellWidgetPurposes);
 	void ClearColumns(ISceneOutliner& InOutliner);
 
 	TArray<FName> AddedColumns;
@@ -100,7 +94,5 @@ class UTypedElementSceneOutlinerFactory : public UTypedElementDataStorageFactory
 public:
 	~UTypedElementSceneOutlinerFactory() override = default;
 
-	void RegisterTables(ITypedElementDataStorageInterface& DataStorage) override;
-	
 	void RegisterWidgetPurposes(ITypedElementDataStorageUiInterface& DataStorageUi) const override;
 };
diff --git a/Engine/Plugins/Experimental/EditorDataStorageFeatures/Source/TedsOutliner/TedsOutliner.Build.cs b/Engine/Plugins/Experimental/EditorDataStorageFeatures/Source/TedsOutliner/TedsOutliner.Build.cs
--- a/Engine/Plugins/Experimental/EditorDataStorageFeatures/Source/TedsOutliner/TedsOutliner.Build.cs
+++ b/Engine/Plugins/Experimental/EditorDataStorageFeatures/Source/TedsOutliner/TedsOutliner.Build.cs
@@ -32,6 +32,7 @@ public class TedsOutliner : ModuleRules
 					"UnrealEd", // FEditorUndoClient used by SSceneOutliner
 					"ToolMenus",
 					"ApplicationCore",
+					"TedsTableViewer"
 				});
 
 			DynamicallyLoadedModuleNames.AddRange(new string[] {});
diff --git a/Engine/Plugins/Experimental/EditorDataStorageFeatures/Source/TedsTableViewer/Private/QueryStack/FQueryStackNode_RowView.cpp b/Engine/Plugins/Experimental/EditorDataStorageFeatures/Source/TedsTableViewer/Private/QueryStack/FQueryStackNode_RowView.cpp
new file mode 100644
--- /dev/null
+++ b/Engine/Plugins/Experimental/EditorDataStorageFeatures/Source/TedsTableViewer/Private/QueryStack/FQueryStackNode_RowView.cpp
@@ -0,0 +1,29 @@
+// Copyright Epic Games, Inc. All Rights Reserved.
+
+#include "QueryStack/FQueryStackNode_RowView.h"
+
+namespace UE::EditorDataStorage
+{
+	TConstArrayView<TypedElementDataStorage::RowHandle> FQueryStackNode_RowView::GetOrderedRowList()
+	{
+		return *Rows;
+	}
+
+	uint32 FQueryStackNode_RowView::GetRevisionId() const
+	{
+		return RevisionId;
+	}
+
+	FQueryStackNode_RowView::FQueryStackNode_RowView(TArray<TypedElementDataStorage::RowHandle>* InRows)
+		: Rows(InRows)
+	{
+		
+	}
+
+	void FQueryStackNode_RowView::MarkDirty()
+	{
+		++RevisionId;
+	}
+}
+
+
diff --git a/Engine/Plugins/Experimental/EditorDataStorageFeatures/Source/TedsTableViewer/Private/TedsTableViewerColumn.cpp b/Engine/Plugins/Experimental/EditorDataStorageFeatures/Source/TedsTableViewer/Private/TedsTableViewerColumn.cpp
new file mode 100644
--- /dev/null
+++ b/Engine/Plugins/Experimental/EditorDataStorageFeatures/Source/TedsTableViewer/Private/TedsTableViewerColumn.cpp
@@ -0,0 +1,336 @@
+// Copyright Epic Games, Inc. All Rights Reserved.
+
+#include "TedsTableViewerColumn.h"
+
+#include "Columns/UIPropertiesColumns.h"
+#include "Elements/Columns/TypedElementMiscColumns.h"
+#include "Elements/Columns/TypedElementSlateWidgetColumns.h"
+#include "Elements/Columns/TypedElementTypeInfoColumns.h"
+#include "Elements/Framework/TypedElementDataStorageWidget.h"
+#include "Elements/Framework/TypedElementRegistry.h"
+#include "TedsTableViewerUtils.h"
+#include "Widgets/Layout/SBox.h"
+#include "Widgets/SBoxPanel.h"
+#include "Widgets/Text/STextBlock.h"
+
+namespace UE::EditorDataStorage
+{
+	FTedsTableViewerColumn::FTedsTableViewerColumn(const FName& InColumnName, const TSharedPtr<FTypedElementWidgetConstructor>& InCellWidgetConstructor,
+		const TArray<TWeakObjectPtr<const UScriptStruct>>& InMatchedColumns, const TSharedPtr<FTypedElementWidgetConstructor>& InHeaderWidgetConstructor,
+		const TypedElementDataStorage::FMetaDataView& InWidgetMetaData)
+		: ColumnName(InColumnName)
+		, CellWidgetConstructor(InCellWidgetConstructor)
+		, HeaderWidgetConstructor(InHeaderWidgetConstructor)
+		, MatchedColumns(InMatchedColumns)
+		, WidgetMetaData(InWidgetMetaData)
+	{
+		UTypedElementRegistry* Registry = UTypedElementRegistry::GetInstance();
+		
+		checkf(Registry, TEXT("Unable to create a Table Viewer column before the Typed Element Registry is initialized."));
+		if (Registry)
+		{
+			Storage = Registry->GetMutableDataStorage();
+			StorageUi = Registry->GetMutableDataStorageUi();
+			StorageCompatibility = Registry->GetMutableDataStorageCompatibility();
+		}
+
+		RegisterQueries();
+
+	}
+	
+	FTedsTableViewerColumn::~FTedsTableViewerColumn()
+	{
+		UnRegisterQueries();
+	}
+
+	TSharedPtr<SWidget> FTedsTableViewerColumn::ConstructRowWidget(TypedElementDataStorage::RowHandle RowHandle) const
+	{
+		TSharedPtr<SWidget> RowWidget;
+		
+		if(Storage->IsRowAssigned(RowHandle))
+		{
+			const TypedElementRowHandle UiRowHandle = Storage->AddRow(Storage->FindTable(TableViewerUtils::GetWidgetTableName()));
+
+			const TConstArrayView<TWeakObjectPtr<const UScriptStruct>> ColumnTypes = GetMatchedColumns();
+			if (ColumnTypes.Num() == 1)
+			{
+				Storage->AddColumn<FTypedElementScriptStructTypeInfoColumn>(UiRowHandle, FTypedElementScriptStructTypeInfoColumn{.TypeInfo = *ColumnTypes.begin()});
+			}
+			
+			if (FTypedElementRowReferenceColumn* RowReference = Storage->GetColumn<FTypedElementRowReferenceColumn>(UiRowHandle))
+			{
+				RowReference->Row = RowHandle;
+			}
+
+			if(FTypedElementSlateWidgetReferenceColumn* WidgetReferenceColumn = Storage->GetColumn<FTypedElementSlateWidgetReferenceColumn>(UiRowHandle))
+			{
+				WidgetReferenceColumn->WidgetConstructor = CellWidgetConstructor;
+			}
+
+			RowWidget = StorageUi->ConstructWidget(UiRowHandle, *CellWidgetConstructor, WidgetMetaData);
+		}
+
+		return RowWidget;
+	}
+
+	SHeaderRow::FColumn::FArguments FTedsTableViewerColumn::ConstructHeaderRowColumn() const
+	{
+		const TConstArrayView<TWeakObjectPtr<const UScriptStruct>> ColumnTypes = GetMatchedColumns();
+		
+		FString TooltipText = TEXT("Data Storage columns:");
+		for (const TWeakObjectPtr<const UScriptStruct>& ColumnType : ColumnTypes)
+		{
+			if (ColumnType.IsValid())
+			{
+				TooltipText += TEXT("\n    ");
+				ColumnType->AppendName(TooltipText);
+			}
+		}
+
+		TSharedPtr<SWidget> Widget;
+		TypedElementDataStorage::RowHandle UiRowHandle = TypedElementDataStorage::InvalidRowHandle;
+		if (HeaderWidgetConstructor)
+		{
+			UiRowHandle = Storage->AddRow(Storage->FindTable(FName(TEXT("Editor_WidgetTable"))));
+
+			// TEDS UI TODO: We can't do this from the Widget Constructor because it is a UStruct and does not have access to AsShared(), so we would
+			// be forced to store a raw pointer instead of a weak pointer which is unsafe. Once the widget construction pipleline is improved this can
+			// probably be moved to a better place
+			if(FTypedElementSlateWidgetReferenceColumn* WidgetReferenceColumn = Storage->GetColumn<FTypedElementSlateWidgetReferenceColumn>(UiRowHandle))
+			{
+				WidgetReferenceColumn->WidgetConstructor = HeaderWidgetConstructor;
+			}
+			
+			Widget = StorageUi->ConstructWidget(UiRowHandle, *HeaderWidgetConstructor, 
+				WidgetMetaData);
+		}
+		if (!Widget.IsValid())
+		{
+			Widget = SNew(STextBlock)
+				.Text(FText::FromString(ColumnName.ToString()));
+		}
+		
+		SHeaderRow::FColumn::FArguments Column = SHeaderRow::Column(ColumnName)
+			.FillWidth(1)
+			.HeaderComboVisibility(EHeaderComboVisibility::OnHover)
+			.DefaultTooltip(FText::FromString(MoveTemp(TooltipText)))
+			.HeaderContent()
+			[
+				SNew(SBox)
+					.MinDesiredHeight(20.0f)
+					.VAlign(VAlign_Center)
+					[
+						Widget.ToSharedRef()
+					]
+			];
+		
+		if (const FUIHeaderPropertiesColumn* HeaderProperties = Storage->GetColumn<FUIHeaderPropertiesColumn>(UiRowHandle))
+		{
+			float Width = HeaderProperties->Width;
+			switch (HeaderProperties->ColumnSizeMode)
+			{
+				case EColumnSizeMode::Fill: Column.FillWidth(Width); break;
+				case EColumnSizeMode::Fixed: Column.FixedWidth(Width); break;
+				case EColumnSizeMode::Manual: Column.ManualWidth(Width); break;
+				case EColumnSizeMode::FillSized: Column.FillSized(Width); break;
+			}
+		}
+		return Column;
+	}
+
+	void FTedsTableViewerColumn::Tick()
+	{
+		// Update any rows that could need widget updates
+		if(!RowsToUpdate.IsEmpty())
+		{
+			UpdateWidgets();
+			RowsToUpdate.Empty();
+		}
+	}
+
+	void FTedsTableViewerColumn::SetIsRowVisibleDelegate(const FIsRowVisible& InIsRowVisibleDelegate)
+	{
+		IsRowVisibleDelegate = InIsRowVisibleDelegate;
+	}
+
+	void FTedsTableViewerColumn::RegisterQueries()
+	{
+		using namespace TypedElementQueryBuilder;
+		using namespace TypedElementDataStorage;
+
+		const TConstArrayView<TWeakObjectPtr<const UScriptStruct>> ColumnTypes = GetMatchedColumns();
+		
+		// For each TEDS column this column is matched with, we'll add observers to track addition/removal to update any widgets
+		for(const TWeakObjectPtr<const UScriptStruct>& ColumnType : ColumnTypes)
+		{
+			const FName ColumnAddObserverName = *(FString::Printf(TEXT("Column Add Monitor for %s Table Viewer Column Column, %s TEDS Column"), *ColumnName.ToString(), *ColumnType->GetName()));
+			FObserver AddObserver(FObserver::EEvent::Add, ColumnType.Get());
+			AddObserver.ForceToGameThread(true);
+
+			// TEDS-Outliner TODO: Long term if we move this into TypedElementOutlinerMode or similar we can get access to the exact
+			// types the Outliner is looking at and specify them on .Where() to cut down on the things we are observing
+			TypedElementDataStorage::QueryHandle AddQueryHandle = Storage->RegisterQuery(
+				Select(
+					ColumnAddObserverName,
+					AddObserver,
+					[this](IQueryContext& Context, TypedElementRowHandle Row)
+						{
+							RowsToUpdate.Add(TPair<TypedElementRowHandle, bool>(Row, true));
+						})
+				.Where()
+					.All(ColumnType.Get())
+				.Compile()
+				);
+			
+			InternalObserverQueries.Add(AddQueryHandle);
+
+			const FName ColumnRemoveObserverName = *(FString::Printf(TEXT("Column Remove Monitor for %s Table Viewer Column Column, %s TEDS Column"), *ColumnName.ToString(), *ColumnType->GetName()));
+			FObserver RemoveObserver(FObserver::EEvent::Remove, ColumnType.Get());
+			RemoveObserver.ForceToGameThread(true);
+
+			// Table Viewer TODO: We might be able to cut down on the rows we are querying for in the future by getting the rows from the query stack
+			// but we currently have to use a generic query so we can support the TEDS-Outliner as well
+			TypedElementDataStorage::QueryHandle RemoveQueryHandle = Storage->RegisterQuery(
+				Select(
+					ColumnRemoveObserverName,
+					RemoveObserver,
+					[this](IQueryContext& Context, TypedElementRowHandle Row)
+						{
+							RowsToUpdate.Add(TPair<TypedElementRowHandle, bool>(Row, false));
+						})
+				.Where()
+					.All(ColumnType.Get())
+				.Compile()
+			);
+			
+			InternalObserverQueries.Add(RemoveQueryHandle);
+		}
+
+		// We are looking for widgets that have a row reference
+		TArray<const UScriptStruct*> SelectionColumns({FTypedElementSlateWidgetReferenceColumn::StaticStruct(), FTypedElementRowReferenceColumn::StaticStruct()});
+
+		// We need to remove duplicates because TEDS/Mass does not handle having the same column in Select() and Where()
+		TArray<const UScriptStruct*> AdditionalWidgetColumns = CellWidgetConstructor->GetAdditionalColumnsList().FilterByPredicate([&SelectionColumns](const UScriptStruct* Column)
+		{
+			return !SelectionColumns.Contains(Column);
+		});
+		
+		// Query to get all widgets that were created by this column
+		WidgetQuery = Storage->RegisterQuery(
+				Select()
+					.ReadOnly(SelectionColumns)
+				.Where()
+					.All(AdditionalWidgetColumns)
+				.Compile());
+	}
+
+	void FTedsTableViewerColumn::UnRegisterQueries()
+	{
+		for(const TypedElementDataStorage::QueryHandle Query : InternalObserverQueries)
+		{
+			Storage->UnregisterQuery(Query);
+		}
+		
+		Storage->UnregisterQuery(WidgetQuery);
+
+	}
+
+	bool FTedsTableViewerColumn::IsRowVisible(const TypedElementDataStorage::RowHandle InRowHandle) const
+	{
+		if(IsRowVisibleDelegate.IsBound())
+		{
+			return IsRowVisibleDelegate.Execute(InRowHandle);
+		}
+
+		// If we don't have a delegate bound we just return true since in the worst case we will just spend time trying to update
+		// rows that aren't visible and therefore don't have widgets due to virtualization
+		return true; 
+	}
+
+	void FTedsTableViewerColumn::UpdateWidgets()
+	{
+		// Remove any widget rows that don't actually need an update
+		RowsToUpdate = RowsToUpdate.FilterByPredicate([this](const TPair<TypedElementDataStorage::RowHandle, bool>& Pair) -> bool
+		{
+			// We don't have a widget for this item visible, so there is nothing to update
+			if(!IsRowVisible(Pair.Key))
+			{
+				return false;
+			}
+			
+			// Check if the row now matches the query conditions for this widget
+			const bool bMatchesQueryConditions = CellWidgetConstructor->GetQueryConditions() && Storage->MatchesColumns(Pair.Key, *CellWidgetConstructor->GetQueryConditions());
+			
+			// If we are adding a column that we are monitoring and it now matches, or if we are removing a column that we are monitoring and it now
+			// stops matching, there is a potential need for widget update
+			return (bMatchesQueryConditions && Pair.Value) || (!bMatchesQueryConditions && !Pair.Value);
+		});
+		
+		using namespace TypedElementQueryBuilder;
+		using namespace TypedElementDataStorage;
+
+		// Query to find all widgets that belong to the row handles that need updates
+		DirectQueryCallback RowCollector = CreateDirectQueryCallbackBinding(
+		[this](const IDirectQueryContext& Context, const FTypedElementSlateWidgetReferenceColumn* ContainerWidgetReferenceColumns, const FTypedElementRowReferenceColumn* RowReferenceColumns)
+		{
+			const FTypedElementSlateWidgetReferenceColumn* WidgetsIt = ContainerWidgetReferenceColumns;
+			const FTypedElementRowReferenceColumn* RowRefsIt = RowReferenceColumns;
+			const TConstArrayView<TypedElementRowHandle> Rows = Context.GetRowHandles();
+
+			for(unsigned RowIndex = 0; RowIndex < Context.GetRowCount(); ++RowIndex, ++WidgetsIt, ++RowRefsIt)
+			{
+				// Check if this widget's owning row is in our rows to update
+				bool* bColumnAddedPtr = RowsToUpdate.Find(RowRefsIt->Row);
+				// If not, skip it
+				if(!bColumnAddedPtr)
+				{
+					continue;
+				}
+
+				// Check if the container TEDSWidget exists, if not we cannot update this widget
+				const TSharedPtr<STedsWidget> TedsWidget = WidgetsIt->TedsWidget.Pin();
+				if(!TedsWidget)
+				{
+					continue;
+				}
+
+				// A row has numerous widgets, make sure we only update the one that was created by our column by checking the constructor
+				if(WidgetsIt->WidgetConstructor != CellWidgetConstructor)
+				{
+					continue;
+				}
+
+				// If a column was added and we are here, we need to re-create the widget
+				// TEDS-Outliner TODO: Do we need to create the widget only if it doesn't exist? Or should we also update it to automatically respond
+				// to column changes even if it was already created
+				if(*bColumnAddedPtr)
+				{
+					const TSharedPtr<SWidget> RowWidget = CellWidgetConstructor->Construct(Rows[RowIndex], Storage, StorageUi, WidgetMetaData);
+
+					if(RowWidget)
+					{
+						TedsWidget->SetContent(RowWidget.ToSharedRef());
+					}
+				}
+				// If a column was removed (and we don't match anymore) delete the internal widget
+				else
+				{
+					TedsWidget->SetContent(SNullWidget::NullWidget);
+				}
+			}
+		});
+
+		Storage->RunQuery(WidgetQuery, RowCollector);
+	}
+	
+	FName FTedsTableViewerColumn::GetColumnName() const
+	{
+		return ColumnName;
+	}
+	
+	TConstArrayView<TWeakObjectPtr<const UScriptStruct>> FTedsTableViewerColumn::GetMatchedColumns() const
+	{
+		return MatchedColumns;
+	}
+}
+
diff --git a/Engine/Plugins/Experimental/EditorDataStorageFeatures/Source/TedsTableViewer/Private/TedsTableViewerModel.cpp b/Engine/Plugins/Experimental/EditorDataStorageFeatures/Source/TedsTableViewer/Private/TedsTableViewerModel.cpp
new file mode 100644
--- /dev/null
+++ b/Engine/Plugins/Experimental/EditorDataStorageFeatures/Source/TedsTableViewer/Private/TedsTableViewerModel.cpp
@@ -0,0 +1,230 @@
+// Copyright Epic Games, Inc. All Rights Reserved.
+
+#include "TedsTableViewerModel.h"
+
+#include "Elements/Framework/TypedElementRegistry.h"
+#include "QueryStack/IQueryStackNode_Row.h"
+#include "TedsTableViewerColumn.h"
+#include "TedsTableViewerUtils.h"
+
+namespace UE::EditorDataStorage
+{
+	FTedsTableViewerModel::FTedsTableViewerModel(const TSharedPtr<IQueryStackNode_Row>& InRowQueryStack,
+		const TArray<TWeakObjectPtr<const UScriptStruct>>& InRequestedColumns, const TArray<FName>& InCellWidgetPurposes,
+		const FIsItemVisible& InIsItemVisibleDelegate)
+		: RowQueryStack(InRowQueryStack)
+		, RequestedTedsColumns(InRequestedColumns)
+		, CellWidgetPurposes(InCellWidgetPurposes)
+		, IsItemVisible(InIsItemVisibleDelegate)
+	{
+		UTypedElementRegistry* Registry = UTypedElementRegistry::GetInstance();
+		
+		checkf(Registry, TEXT("Unable to create a Table Viewer before the Typed Element Registry is initialized."));
+		if (Registry)
+		{
+			Storage = Registry->GetMutableDataStorage();
+			StorageUi = Registry->GetMutableDataStorageUi();
+			StorageCompatibility = Registry->GetMutableDataStorageCompatibility();
+		}
+		
+		TickerHandle = FTSTicker::GetCoreTicker().AddTicker(FTickerDelegate::CreateRaw(this, &FTedsTableViewerModel::Tick), 0);
+
+		GenerateColumns();
+		Refresh();
+	}
+
+	FTedsTableViewerModel::~FTedsTableViewerModel()
+	{
+		FTSTicker::GetCoreTicker().RemoveTicker(TickerHandle);
+	}
+	
+	void FTedsTableViewerModel::Refresh()
+	{
+		Items.Empty();
+		
+		for(TypedElementDataStorage::RowHandle RowHandle : RowQueryStack->GetOrderedRowList())
+		{
+			Items.Add(MakeShared<TedsTableViewerItem>(RowHandle));
+		}
+
+		CachedRowQueryStackRevision = RowQueryStack->GetRevisionId();
+
+		OnModelChanged.Broadcast();
+	}
+	
+	bool FTedsTableViewerModel::Tick(float DeltaTime)
+	{
+		// If the revision ID has changed, refresh to update our rows
+		if(RowQueryStack->GetRevisionId() != CachedRowQueryStackRevision)
+		{
+			Refresh();
+		}
+
+		// Tick all the individual column views
+		for(const auto& ColumnPair : ColumnsView)
+		{
+			ColumnPair.Value->Tick();
+		}
+		
+		return true;
+	}
+
+	const TArray<TableViewerItemPtr>& FTedsTableViewerModel::GetItems() const
+	{
+		return Items;
+	}
+
+	TSharedPtr<FTedsTableViewerColumn> FTedsTableViewerModel::GetColumn(const FName& ColumnName) const
+	{
+		if(const TSharedRef<FTedsTableViewerColumn>* Column = ColumnsView.Find(ColumnName))
+		{
+			return *Column;
+		}
+
+		return nullptr;
+	}
+
+	void FTedsTableViewerModel::ForEachColumn(const TFunctionRef<void(const TSharedRef<FTedsTableViewerColumn>&)>& Delegate) const
+	{
+		for(const TPair<FName, TSharedRef<FTedsTableViewerColumn>>& ColumnPair : ColumnsView)
+		{
+			Delegate(ColumnPair.Value);
+		}
+	}
+
+	FTedsTableViewerModel::FOnModelChanged& FTedsTableViewerModel::GetOnModelChanged()
+	{
+		return OnModelChanged;
+	}
+
+	void FTedsTableViewerModel::SetColumns(const TArray<TWeakObjectPtr<const UScriptStruct>>& InColumns)
+	{
+		RequestedTedsColumns = InColumns;
+		GenerateColumns();
+	}
+
+	void FTedsTableViewerModel::GenerateColumns()
+	{
+		using MatchApproach = ITypedElementDataStorageUiInterface::EMatchApproach;
+		int32 IndexOffset = 0;
+
+		ColumnsView.Empty();
+		
+		// A Map of TEDS Columns -> UI columns so we can add them in the same order they were specified
+		TMap<TWeakObjectPtr<const UScriptStruct>, TSharedRef<FTedsTableViewerColumn>> NewColumnMap;
+
+		// A copy of the columns to preserve the order since TEDS UI modifies the array directly
+		TArray<TWeakObjectPtr<const UScriptStruct>> ColumnsCopy = RequestedTedsColumns;
+		
+		// Lambda to create the constructor for a given list of columns
+		auto ColumnConstructor = [this, &IndexOffset, &NewColumnMap](
+			TUniquePtr<FTypedElementWidgetConstructor> Constructor, TConstArrayView<TWeakObjectPtr<const UScriptStruct>> MatchedColumns)
+			{
+				TSharedPtr<FTypedElementWidgetConstructor> CellConstructor(Constructor.Release());
+
+				TSharedPtr<FTypedElementWidgetConstructor> HeaderConstructor =
+					TableViewerUtils::CreateHeaderWidgetConstructor(*StorageUi, TypedElementDataStorage::FMetaDataView(), MatchedColumns, CellWidgetPurposes);
+			
+				FName NameId = TableViewerUtils::FindLongestMatchingName(MatchedColumns, IndexOffset);
+
+				TArray<TWeakObjectPtr<const UScriptStruct>> MatchedColumnsCopy(MatchedColumns);
+				TSharedRef<FTedsTableViewerColumn> Column = MakeShared<FTedsTableViewerColumn>(NameId, CellConstructor, MoveTemp(MatchedColumnsCopy), HeaderConstructor);
+			
+				Column->SetIsRowVisibleDelegate(FTedsTableViewerColumn::FIsRowVisible::CreateRaw(this, &FTedsTableViewerModel::IsRowVisible));
+
+				for (const TWeakObjectPtr<const UScriptStruct>& ColumnType : MatchedColumns)
+				{
+					NewColumnMap.Emplace(ColumnType, Column);
+				}
+			
+				++IndexOffset;
+				return true;
+			};
+
+		// Create the widget constructors for the columns
+		for(const FName& WidgetPurpose : CellWidgetPurposes)
+		{
+			StorageUi->CreateWidgetConstructors(WidgetPurpose, MatchApproach::LongestMatch, ColumnsCopy, 
+			TypedElementDataStorage::FMetaDataView(), ColumnConstructor);
+		}
+
+		// For any remaining columns, we'll try to find and use any default widgets
+		for (TWeakObjectPtr<const UScriptStruct>& ColumnType : ColumnsCopy)
+		{
+			auto AssignWidgetToColumn = [this, ColumnType, &IndexOffset, &NewColumnMap](
+				TUniquePtr<FTypedElementWidgetConstructor> Constructor, TConstArrayView<TWeakObjectPtr<const UScriptStruct>>)
+			{
+				TSharedPtr<FTypedElementWidgetConstructor> CellConstructor(Constructor.Release());
+
+				TSharedPtr<FTypedElementWidgetConstructor> HeaderConstructor =
+					TableViewerUtils::CreateHeaderWidgetConstructor(*StorageUi, TypedElementDataStorage::FMetaDataView(), {ColumnType}, CellWidgetPurposes);
+
+				FName NameId = FName(ColumnType->GetDisplayNameText().ToString());
+
+				TArray<TWeakObjectPtr<const UScriptStruct>> MatchedColumnsCopy({ColumnType});
+				TSharedRef<FTedsTableViewerColumn> Column = MakeShared<FTedsTableViewerColumn>(NameId, CellConstructor, MoveTemp(MatchedColumnsCopy), HeaderConstructor);
+				
+				Column->SetIsRowVisibleDelegate(FTedsTableViewerColumn::FIsRowVisible::CreateRaw(this, &FTedsTableViewerModel::IsRowVisible));
+
+				NewColumnMap.Emplace(ColumnType, Column);
+				
+				++IndexOffset;
+				return false;
+			};
+
+			const int32 BeforeIndexOffset = IndexOffset;
+			
+			for(const FName& WidgetPurpose : CellWidgetPurposes)
+			{
+				const FName DefaultWidgetPurpose(WidgetPurpose.ToString() + TEXT(".Default"));
+
+				StorageUi->CreateWidgetConstructors(DefaultWidgetPurpose, TypedElementDataStorage::FMetaDataView(), AssignWidgetToColumn);
+				
+				if (BeforeIndexOffset != IndexOffset)
+				{
+					break;
+				}
+			}
+
+			if (BeforeIndexOffset == IndexOffset)
+			{
+				++IndexOffset;
+			}
+		}
+
+		// Add the actual UI columns in the order the Teds Columns were specified
+		for (const TWeakObjectPtr<const UScriptStruct>& ColumnType : RequestedTedsColumns)
+		{
+			if(TSharedRef<FTedsTableViewerColumn>* FoundColumn = NewColumnMap.Find(ColumnType))
+			{
+				// If the column already exists, a widget matched it and a previously encountered column together and was already added
+				// so we can safely ignore it here
+				if(!GetColumn((*FoundColumn)->GetColumnName()))
+				{
+					ColumnsView.Emplace((*FoundColumn)->GetColumnName(), *FoundColumn);
+				}
+			}
+		}
+	}
+
+	bool FTedsTableViewerModel::IsRowVisible(TypedElementDataStorage::RowHandle InRowHandle) const
+	{
+		if(!IsItemVisible.IsBound())
+		{
+			return true;
+		}
+		
+		// Table Viewer TODO: We can probably store a map of the items instead but this works for now
+		for(const TableViewerItemPtr& Item : Items)
+		{
+			if(Item->RowHandle == InRowHandle)
+			{
+				return IsItemVisible.Execute(Item);
+			}
+		}
+
+		return true;
+	}
+
+
+}
diff --git a/Engine/Plugins/Experimental/EditorDataStorageFeatures/Source/TedsTableViewer/Private/TedsTableViewerModule.cpp b/Engine/Plugins/Experimental/EditorDataStorageFeatures/Source/TedsTableViewer/Private/TedsTableViewerModule.cpp
new file mode 100644
--- /dev/null
+++ b/Engine/Plugins/Experimental/EditorDataStorageFeatures/Source/TedsTableViewer/Private/TedsTableViewerModule.cpp
@@ -0,0 +1,124 @@
+// Copyright Epic Games, Inc. All Rights Reserved.
+
+#include "Elements/Columns/TypedElementAlertColumns.h"
+#include "Elements/Columns/TypedElementCompatibilityColumns.h"
+#include "Elements/Columns/TypedElementLabelColumns.h"
+#include "Elements/Framework/TypedElementRegistry.h"
+#include "LevelEditor.h"
+#include "Modules/ModuleManager.h"
+#include "Modules/ModuleInterface.h"
+#include "Widgets/STedsTableViewer.h"
+#include "Widgets/Docking/SDockTab.h"
+#include "WorkspaceMenuStructureModule.h"
+#include "WorkspaceMenuStructure.h"
+#include "Elements/Columns/TypedElementSelectionColumns.h"
+#include "QueryStack/FQueryStackNode_RowView.h"
+
+#define LOCTEXT_NAMESPACE "TedsTableViewerModule"
+
+static FName TableViewerTabName(TEXT("TedsTableViewer"));
+
+class FTedsTableViewerModule : public IModuleInterface
+{
+public:
+
+	virtual void StartupModule() override
+	{
+		RegisterTableViewerTab();
+	}
+
+	virtual void ShutdownModule() override
+	{
+	}
+
+	void RegisterTableViewerTab()
+	{
+		FLevelEditorModule& LevelEditorModule = FModuleManager::GetModuleChecked<FLevelEditorModule>("LevelEditor");
+		
+		LevelEditorTabManagerChangedHandle = LevelEditorModule.OnTabManagerChanged().AddLambda([this]()
+		{
+			FLevelEditorModule& LevelEditorModule = FModuleManager::GetModuleChecked<FLevelEditorModule>("LevelEditor");
+
+			TSharedPtr<FTabManager> LevelEditorTabManager = LevelEditorModule.GetLevelEditorTabManager();
+
+			LevelEditorTabManager->RegisterTabSpawner(TableViewerTabName, FOnSpawnTab::CreateRaw(this, &FTedsTableViewerModule::OpenTableViewer))
+			.SetDisplayName(LOCTEXT("TedsTableVIewerTitle", "Table Viewer (Experimental)"))
+			.SetGroup(WorkspaceMenu::GetMenuStructure().GetLevelEditorOutlinerCategory())
+			.SetIcon(FSlateIcon(FAppStyle::GetAppStyleSetName(), "LevelEditor.Tabs.Outliner"))
+			.SetAutoGenerateMenuEntry(false); // This can only be summoned from the Cvar now
+	
+		});
+	}
+	
+	
+	TSharedRef<SDockTab> OpenTableViewer(const FSpawnTabArgs& SpawnTabArgs)
+	{
+		return SNew(SDockTab)
+		.TabRole(ETabRole::NomadTab)
+		[
+			CreateTableViewer()
+		];
+	}
+	
+	TSharedRef<SWidget> CreateTableViewer()
+	{
+		UTypedElementRegistry* Registry = UTypedElementRegistry::GetInstance();
+		checkf(Registry, TEXT("Unable to initialize the table viewer before TEDS is initialized."));
+
+		if(!Registry->AreDataStorageInterfacesSet())
+		{
+			return SNullWidget::NullWidget;
+		}
+
+		ITypedElementDataStorageInterface* DataStorage = Registry->GetMutableDataStorage();
+		
+		using namespace TypedElementQueryBuilder;
+
+		// We'll just create a test table viewer that views all actor rows without actually updating it dynamically for now
+		static TypedElementDataStorage::QueryHandle QueryHandle =
+			DataStorage->RegisterQuery(
+						Select()
+						.Where()
+							.All<FTypedElementActorTag>()
+						.Compile()
+				);
+
+		Rows.Empty();
+		
+		TypedElementDataStorage::FQueryResult QueryResult = DataStorage->RunQuery(QueryHandle,
+			CreateDirectQueryCallbackBinding([this](const ITypedElementDataStorageInterface::IDirectQueryContext& Context, const TypedElementDataStorage::RowHandle* RowHandles)
+		{
+			Rows.Append(RowHandles, Context.GetRowCount());
+		}));
+
+		return SNew(UE::EditorDataStorage::STedsTableViewer)
+				.QueryStack(MakeShared<UE::EditorDataStorage::FQueryStackNode_RowView>(&Rows))
+				.Columns({FTypedElementLabelColumn::StaticStruct(), FTypedElementSelectionColumn::StaticStruct(),
+					FTypedElementAlertColumn::StaticStruct(), FTypedElementChildAlertColumn::StaticStruct()});
+	}
+	
+	FDelegateHandle LevelEditorTabManagerChangedHandle;
+	TArray<TypedElementDataStorage::RowHandle> Rows;
+};
+
+// CVar to summon a test table viewer that views a snapshot of all actors at the moment when the Cvar is used
+static FAutoConsoleCommand OpenTableViewerConsoleCommand(
+	TEXT("TEDS.UI.OpenTableViewer"),
+	TEXT("Spawn the test TEDS Table Viewer."),
+	FConsoleCommandWithArgsDelegate::CreateLambda([](const TArray<FString>& Args)
+	{
+		FLevelEditorModule& LevelEditorModule = FModuleManager::GetModuleChecked<FLevelEditorModule>("LevelEditor");
+
+		TSharedPtr<FTabManager> LevelEditorTabManager = LevelEditorModule.GetLevelEditorTabManager();
+		if(TSharedPtr<SDockTab> DockTab = LevelEditorTabManager->FindExistingLiveTab(TableViewerTabName))
+		{
+			DockTab->RequestCloseTab();
+		}
+
+		LevelEditorTabManager->TryInvokeTab(TableViewerTabName);
+	}));
+
+IMPLEMENT_MODULE(FTedsTableViewerModule, TedsTableViewer);
+
+#undef LOCTEXT_NAMESPACE //"TedsTableViewerModule"
+
diff --git a/Engine/Plugins/Experimental/EditorDataStorageFeatures/Source/TedsTableViewer/Private/TedsTableViewerUtils.cpp b/Engine/Plugins/Experimental/EditorDataStorageFeatures/Source/TedsTableViewer/Private/TedsTableViewerUtils.cpp
new file mode 100644
--- /dev/null
+++ b/Engine/Plugins/Experimental/EditorDataStorageFeatures/Source/TedsTableViewer/Private/TedsTableViewerUtils.cpp
@@ -0,0 +1,155 @@
+// Copyright Epic Games, Inc. All Rights Reserved.
+
+#include "TedsTableViewerUtils.h"
+
+#include "Elements/Columns/TypedElementMiscColumns.h"
+#include "Elements/Common/TypedElementDataStorageLog.h"
+#include "Elements/Interfaces/TypedElementDataStorageUiInterface.h"
+
+namespace UE::EditorDataStorage::TableViewerUtils
+{
+	static FName TableViewerWidgetTableName("Editor_TableViewerWidgetTable");
+
+	FName GetWidgetTableName()
+	{
+		return TableViewerWidgetTableName;
+	}
+	
+	// TEDS UI TODO: Maybe the widget can specify a user facing name derived from the matched columns instead of trying to find the longest matching name
+	FName FindLongestMatchingName(TConstArrayView<TWeakObjectPtr<const UScriptStruct>> ColumnTypes, int32 DefaultNameIndex)
+	{
+		switch (ColumnTypes.Num())
+		{
+		case 0:
+			return FName(TEXT("Column"), DefaultNameIndex);
+		case 1:
+			return FName(ColumnTypes[0]->GetDisplayNameText().ToString());
+		default:
+			{
+				FText LongestMatchText = ColumnTypes[0]->GetDisplayNameText();
+				FStringView LongestMatch = LongestMatchText.ToString();
+				const TWeakObjectPtr<const UScriptStruct>* ItEnd = ColumnTypes.end();
+				const TWeakObjectPtr<const UScriptStruct>* It = ColumnTypes.begin();
+				++It; // Skip the first entry as that's already set.
+				for (; It != ItEnd; ++It)
+				{
+					FText NextMatchText = (*It)->GetDisplayNameText();
+					FStringView NextMatch = NextMatchText.ToString();
+
+					int32 MatchSize = 0;
+					auto ItLeft = LongestMatch.begin();
+					auto ItLeftEnd = LongestMatch.end();
+					auto ItRight = NextMatch.begin();
+					auto ItRightEnd = NextMatch.end();
+					while (
+						ItLeft != ItLeftEnd &&
+						ItRight != ItRightEnd &&
+						*ItLeft == *ItRight)
+					{
+						++MatchSize;
+						++ItLeft;
+						++ItRight;
+					}
+
+					// At least 3 letters have to match to avoid single or double letter names which typically mean nothing.
+					if (MatchSize > 2)
+					{
+						LongestMatch.LeftInline(MatchSize);
+					}
+					else
+					{
+						// There are not enough characters in the string that match. Just return the name of the first column
+						return FName(ColumnTypes[0]->GetDisplayNameText().ToString());
+					}
+				}
+				return FName(LongestMatch);
+			}
+		};
+	}
+	
+	TArray<TWeakObjectPtr<const UScriptStruct>> CreateVerifiedColumnTypeArray(TConstArrayView<TWeakObjectPtr<const UScriptStruct>> ColumnTypes)
+	{
+		TArray<TWeakObjectPtr<const UScriptStruct>> VerifiedColumnTypes;
+		VerifiedColumnTypes.Reserve(ColumnTypes.Num());
+		for (const TWeakObjectPtr<const UScriptStruct>& ColumnType : ColumnTypes)
+		{
+			if (ColumnType.IsValid())
+			{
+				VerifiedColumnTypes.Add(ColumnType.Get());
+			}
+			else
+			{
+				UE_LOG(LogTypedElementDataStorage, Verbose, TEXT("Invalid column provided to the table viewer"));
+			}
+		}
+		return VerifiedColumnTypes;
+	}
+
+	TSharedPtr<FTypedElementWidgetConstructor> CreateHeaderWidgetConstructor(ITypedElementDataStorageUiInterface& StorageUi, const TypedElementDataStorage::FMetaDataView& InMetaData, TConstArrayView<TWeakObjectPtr<const UScriptStruct>> ColumnTypes, const TConstArrayView<FName> CellWidgetPurposes)
+	{
+		using MatchApproach = ITypedElementDataStorageUiInterface::EMatchApproach;
+
+		TArray<TWeakObjectPtr<const UScriptStruct>> VerifiedColumnTypes = CreateVerifiedColumnTypeArray(ColumnTypes);
+		TSharedPtr<FTypedElementWidgetConstructor> Constructor;
+
+		for (const FName& Purpose : CellWidgetPurposes)
+		{
+			// Extract the header purpose from the cell purpose (e.g "SceneOutliner.ItemLabel.Cell" -> "SceneOutliner.ItemLabel.Header")
+			FString HeaderPurpose, CellPurpose;
+			Purpose.ToString().Split(TEXT(".Cell"), &HeaderPurpose, &CellPurpose);
+			HeaderPurpose.Append(TEXT(".Header"));
+			
+			StorageUi.CreateWidgetConstructors(FName(HeaderPurpose), MatchApproach::ExactMatch, VerifiedColumnTypes, InMetaData,
+				[&Constructor, ColumnTypes](
+					TUniquePtr<FTypedElementWidgetConstructor> CreatedConstructor, 
+					TConstArrayView<TWeakObjectPtr<const UScriptStruct>> MatchedColumnTypes)
+				{
+					if (ColumnTypes.Num() == MatchedColumnTypes.Num())
+					{
+						Constructor = TSharedPtr<FTypedElementWidgetConstructor>(CreatedConstructor.Release());
+					}
+					// Either this was the exact match so no need to search further or the longest possible chain didn't match so the next ones will 
+					// always be shorter in both cases just return.
+					return false;
+				});
+			if (Constructor)
+			{
+				return Constructor;
+			}
+		}
+		for (const FName& Purpose : CellWidgetPurposes)
+		{
+			// Extract the default header purpose from the cell purpose (e.g "SceneOutliner.ItemLabel.Cell" -> "SceneOutliner.ItemLabel.Header.Default")
+			FString HeaderPurpose, CellPurpose;
+			Purpose.ToString().Split(TEXT(".Cell"), &HeaderPurpose, &CellPurpose);
+			HeaderPurpose.Append(TEXT(".Header.Default"));
+
+			StorageUi.CreateWidgetConstructors(FName(HeaderPurpose), InMetaData,
+				[&Constructor, ColumnTypes](
+					TUniquePtr<FTypedElementWidgetConstructor> CreatedConstructor,
+					TConstArrayView<TWeakObjectPtr<const UScriptStruct>> MatchedColumnTypes)
+				{
+					Constructor = TSharedPtr<FTypedElementWidgetConstructor>(CreatedConstructor.Release());
+					return false;
+				});
+			if (Constructor)
+			{
+				return Constructor;
+			}
+		}
+		return nullptr;
+	}
+}
+
+void UTypedElementTableViewerFactory::RegisterTables(ITypedElementDataStorageInterface& DataStorage)
+{
+	const TypedElementTableHandle BaseWidgetTable = DataStorage.FindTable(FName(TEXT("Editor_WidgetTable")));
+	if (BaseWidgetTable != TypedElementInvalidTableHandle)
+	{
+		DataStorage.RegisterTable(
+			BaseWidgetTable,
+			{
+				FTypedElementRowReferenceColumn::StaticStruct()
+			}, UE::EditorDataStorage::TableViewerUtils::TableViewerWidgetTableName);
+	}
+}
diff --git a/Engine/Plugins/Experimental/EditorDataStorageFeatures/Source/TedsTableViewer/Private/Widgets/STedsTableViewer.cpp b/Engine/Plugins/Experimental/EditorDataStorageFeatures/Source/TedsTableViewer/Private/Widgets/STedsTableViewer.cpp
new file mode 100644
--- /dev/null
+++ b/Engine/Plugins/Experimental/EditorDataStorageFeatures/Source/TedsTableViewer/Private/Widgets/STedsTableViewer.cpp
@@ -0,0 +1,64 @@
+// Copyright Epic Games, Inc. All Rights Reserved.
+
+#include "Widgets/STedsTableViewer.h"
+
+#include "TedsTableViewerColumn.h"
+#include "Widgets/STedsTableViewerRow.h"
+#include "Widgets/Views/SListView.h"
+
+namespace UE::EditorDataStorage
+{
+	void STedsTableViewer::Construct(const FArguments& InArgs)
+	{
+		Model = MakeShared<FTedsTableViewerModel>(InArgs._QueryStack, InArgs._Columns, InArgs._CellWidgetPurposes,
+			FTedsTableViewerModel::FIsItemVisible::CreateSP(this, &STedsTableViewer::IsItemVisible));
+		
+		HeaderRowWidget = SNew( SHeaderRow )
+							.CanSelectGeneratedColumn(true);
+
+		ChildSlot
+		[
+			SAssignNew(ListView, SListView<TableViewerItemPtr>)
+			.HeaderRow(HeaderRowWidget)
+			.ListItemsSource(&Model->GetItems())
+			.OnGenerateRow(this, &STedsTableViewer::MakeTableRowWidget)
+		];
+
+		// Add each Teds column from the model to our header row widget
+		Model->ForEachColumn([this](const TSharedRef<FTedsTableViewerColumn>& Column)
+		{
+			HeaderRowWidget->AddColumn(Column->ConstructHeaderRowColumn());
+		});
+
+		// Whenever the model changes, refresh the list to update the UI
+		Model->GetOnModelChanged().AddLambda([this]()
+		{
+			ListView->RequestListRefresh();
+		});
+	}
+
+	void STedsTableViewer::SetColumns(const TArray<TWeakObjectPtr<const UScriptStruct>>& Columns) const
+	{
+		Model->SetColumns(Columns);
+		
+		HeaderRowWidget->ClearColumns();
+		Model->ForEachColumn([this](const TSharedRef<FTedsTableViewerColumn>& Column)
+		{
+			HeaderRowWidget->AddColumn(Column->ConstructHeaderRowColumn());
+		});
+	}
+
+	bool STedsTableViewer::IsItemVisible(TableViewerItemPtr InItem) const
+	{
+		return ListView->IsItemVisible(InItem);
+	}
+
+	TSharedRef<ITableRow> STedsTableViewer::MakeTableRowWidget(TableViewerItemPtr InItem, const TSharedRef<STableViewBase>& OwnerTable) const
+	{
+		return SNew(STedsTableViewerRow, OwnerTable, Model.ToSharedRef())
+				.Item(InItem);
+	}
+
+	
+}
+
diff --git a/Engine/Plugins/Experimental/EditorDataStorageFeatures/Source/TedsTableViewer/Private/Widgets/STedsTableViewerRow.cpp b/Engine/Plugins/Experimental/EditorDataStorageFeatures/Source/TedsTableViewer/Private/Widgets/STedsTableViewerRow.cpp
new file mode 100644
--- /dev/null
+++ b/Engine/Plugins/Experimental/EditorDataStorageFeatures/Source/TedsTableViewer/Private/Widgets/STedsTableViewerRow.cpp
@@ -0,0 +1,39 @@
+// Copyright Epic Games, Inc. All Rights Reserved.
+
+
+#include "Widgets/STedsTableViewerRow.h"
+
+#include "TedsTableViewerColumn.h"
+#include "Widgets/STedsTableViewer.h"
+
+namespace UE::EditorDataStorage
+{
+	void STedsTableViewerRow::Construct(const FArguments& InArgs, const TSharedRef<STableViewBase>& OwnerTableView, const TSharedRef<FTedsTableViewerModel>& InTableViewerModel)
+	{
+		Item = InArgs._Item;
+		TableViewerModel = InTableViewerModel;
+
+		const auto Args = FSuperRowType::FArguments()
+		.Style(&FAppStyle::Get().GetWidgetStyle<FTableRowStyle>("SceneOutliner.TableViewRow"));
+
+		SMultiColumnTableRow<TableViewerItemPtr>::Construct(Args, OwnerTableView);
+	}
+
+	TSharedRef<SWidget> STedsTableViewerRow::GenerateWidgetForColumn(const FName& ColumnName)
+	{
+		if(TSharedPtr<FTedsTableViewerColumn> Column = TableViewerModel->GetColumn(ColumnName))
+		{
+			if(TSharedPtr<SWidget> RowWidget = Column->ConstructRowWidget(Item->RowHandle))
+			{
+				return SNew(SBox)
+						.MinDesiredHeight(20.0f)
+						.VAlign(VAlign_Center)
+						[
+							RowWidget.ToSharedRef()
+						];
+			}
+		}
+
+		return SNullWidget::NullWidget;
+	}
+}
\ No newline at end of file
diff --git a/Engine/Plugins/Experimental/EditorDataStorageFeatures/Source/TedsTableViewer/Public/Columns/UIPropertiesColumns.h b/Engine/Plugins/Experimental/EditorDataStorageFeatures/Source/TedsTableViewer/Public/Columns/UIPropertiesColumns.h
new file mode 100644
--- /dev/null
+++ b/Engine/Plugins/Experimental/EditorDataStorageFeatures/Source/TedsTableViewer/Public/Columns/UIPropertiesColumns.h
@@ -0,0 +1,25 @@
+// Copyright Epic Games, Inc. All Rights Reserved.
+
+#pragma once
+
+#include "Elements/Interfaces/TypedElementDataStorageInterface.h"
+#include "Widgets/Views/SHeaderRow.h"
+
+#include "UIPropertiesColumns.generated.h"
+
+USTRUCT()
+struct FUIHeaderPropertiesColumn final : public FTypedElementDataStorageColumn
+{
+	GENERATED_BODY()
+	
+	EColumnSizeMode::Type ColumnSizeMode;
+
+	/*
+	 * Fill: Column stretches to this fraction of the header row
+	 * Fixed: Column is fixed at this width in slate units and cannot be resized
+	 * Manual: Column defaults to this width in slate units and can be user-sized
+	 * FillSized: Column stretches as Fill but is initialized with this width in slate units
+	 */
+	UPROPERTY()
+	float Width = 0.0f;
+};
diff --git a/Engine/Plugins/Experimental/EditorDataStorageFeatures/Source/TedsTableViewer/Public/QueryStack/FQueryStackNode_RowView.h b/Engine/Plugins/Experimental/EditorDataStorageFeatures/Source/TedsTableViewer/Public/QueryStack/FQueryStackNode_RowView.h
new file mode 100644
--- /dev/null
+++ b/Engine/Plugins/Experimental/EditorDataStorageFeatures/Source/TedsTableViewer/Public/QueryStack/FQueryStackNode_RowView.h
@@ -0,0 +1,33 @@
+// Copyright Epic Games, Inc. All Rights Reserved.
+
+#pragma once
+
+#include "QueryStack/IQueryStackNode_Row.h"
+
+namespace UE::EditorDataStorage
+{
+	/*
+	 * A very simple row query stack that views a list of rows.
+	 * Note: The user is responsible for the lifetime of the Rows array (similar to SListView) and MarkDirty must be called to update
+	 * the table viewer when the row list changes
+	 */
+	class FQueryStackNode_RowView : public IQueryStackNode_Row
+	{
+	public:
+		// IQueryStackNode_Row interface
+		TEDSTABLEVIEWER_API virtual TConstArrayView<TypedElementDataStorage::RowHandle> GetOrderedRowList() override;
+		TEDSTABLEVIEWER_API virtual uint32 GetRevisionId() const override;
+		// IQueryStackNode_Row interface
+
+		TEDSTABLEVIEWER_API FQueryStackNode_RowView(TArray<TypedElementDataStorage::RowHandle>* InRows);
+
+		// Increment the revision ID to update the table viewer when the list of rows changes
+		TEDSTABLEVIEWER_API void MarkDirty();
+
+	private:
+
+		TArray<TypedElementDataStorage::RowHandle>* Rows;
+		uint32 RevisionId = 0;
+	};
+}
+
diff --git a/Engine/Plugins/Experimental/EditorDataStorageFeatures/Source/TedsTableViewer/Public/QueryStack/IQueryStackNode_Row.h b/Engine/Plugins/Experimental/EditorDataStorageFeatures/Source/TedsTableViewer/Public/QueryStack/IQueryStackNode_Row.h
new file mode 100644
--- /dev/null
+++ b/Engine/Plugins/Experimental/EditorDataStorageFeatures/Source/TedsTableViewer/Public/QueryStack/IQueryStackNode_Row.h
@@ -0,0 +1,27 @@
+// Copyright Epic Games, Inc. All Rights Reserved.
+
+#pragma once
+
+#include "Containers/ArrayView.h"
+#include "Elements/Common/TypedElementHandles.h"
+
+namespace UE::EditorDataStorage
+{
+	/*
+	 * A row node in the query stack is used by the Teds Table Viewer to determine which rows to visualize in the table viewer
+	 * @see FQueryStackNode_RowView for a simple implementation that views a list of rows input by the user
+	 */
+	class IQueryStackNode_Row
+	{
+	public:
+		
+		virtual ~IQueryStackNode_Row() = default;
+
+		// Get the rows to visualize
+		virtual TConstArrayView<TypedElementDataStorage::RowHandle> GetOrderedRowList() = 0;
+
+		// Returns an id that updates whenever the locally stored rows have changed.
+		virtual uint32 GetRevisionId() const = 0;
+	};
+}
+
diff --git a/Engine/Plugins/Experimental/EditorDataStorageFeatures/Source/TedsTableViewer/Public/TedsTableViewerColumn.h b/Engine/Plugins/Experimental/EditorDataStorageFeatures/Source/TedsTableViewer/Public/TedsTableViewerColumn.h
new file mode 100644
--- /dev/null
+++ b/Engine/Plugins/Experimental/EditorDataStorageFeatures/Source/TedsTableViewer/Public/TedsTableViewerColumn.h
@@ -0,0 +1,85 @@
+// Copyright Epic Games, Inc. All Rights Reserved.
+
+#pragma once
+#include "Elements/Common/TypedElementHandles.h"
+#include "Elements/Framework/TypedElementMetaData.h"
+#include "Templates/SharedPointer.h"
+#include "UObject/NameTypes.h"
+#include "Widgets/Views/SHeaderRow.h"
+
+class ITypedElementDataStorageCompatibilityInterface;
+class ITypedElementDataStorageUiInterface;
+class ITypedElementDataStorageInterface;
+struct FTypedElementWidgetConstructor;
+class SWidget;
+
+namespace UE::EditorDataStorage
+{
+	/*
+	 * Class representing a column in the UI of the table viewer. Can be constructed using a NameID and a WidgetConstructor to create the actual
+	 * widgets for rows (optionally supplying a header widget constructor and widget metadata to use)
+	 */
+	class FTedsTableViewerColumn
+	{
+	public:
+
+		// Delegate to check if a row is currently visible in the owning table viewer's UI
+		DECLARE_DELEGATE_RetVal_OneParam(bool, FIsRowVisible, const TypedElementDataStorage::RowHandle);
+
+		TEDSTABLEVIEWER_API FTedsTableViewerColumn(
+			const FName& ColumnName, // The unique ID of this column
+			const TSharedPtr<FTypedElementWidgetConstructor>& InCellWidgetConstructor, // The widget constructor to use for this column
+			const TArray<TWeakObjectPtr<const UScriptStruct>>& InMatchedColumns = {}, // An optional list of matched Teds columns 
+			const TSharedPtr<FTypedElementWidgetConstructor>& InHeaderWidgetConstructor = nullptr, // Optional constructor to use for the header widget
+			const TypedElementDataStorage::FMetaDataView& InWidgetMetaData = TypedElementDataStorage::FMetaDataView()); // Optional metadata to use when constructing widgets
+
+		TEDSTABLEVIEWER_API ~FTedsTableViewerColumn();
+		
+		TEDSTABLEVIEWER_API	TSharedPtr<SWidget> ConstructRowWidget(TypedElementDataStorage::RowHandle RowHandle) const;
+		
+		TEDSTABLEVIEWER_API SHeaderRow::FColumn::FArguments ConstructHeaderRowColumn() const;
+		
+		TEDSTABLEVIEWER_API void Tick();
+		
+		TEDSTABLEVIEWER_API void SetIsRowVisibleDelegate(const FIsRowVisible& InIsRowVisibleDelegate);
+		
+		TEDSTABLEVIEWER_API FName GetColumnName() const;
+
+		TEDSTABLEVIEWER_API TConstArrayView<TWeakObjectPtr<const UScriptStruct>> GetMatchedColumns() const;
+
+
+	protected:
+
+		void RegisterQueries();
+		void UnRegisterQueries();
+		bool IsRowVisible(const TypedElementDataStorage::RowHandle InRowHandle) const;
+		void UpdateWidgets();
+		
+	private:
+		// The ID of the column
+		FName ColumnName;
+
+		// Widget Constructors
+		TSharedPtr<FTypedElementWidgetConstructor> CellWidgetConstructor;
+		TSharedPtr<FTypedElementWidgetConstructor> HeaderWidgetConstructor;
+
+		// Teds Columns this widget constructor matched with
+		TArray<TWeakObjectPtr<const UScriptStruct>> MatchedColumns;
+		
+		// The Metadata used to create widgets
+		TypedElementDataStorage::FMetaDataView WidgetMetaData;
+
+		// TEDS Constructs
+		ITypedElementDataStorageInterface* Storage;
+		ITypedElementDataStorageUiInterface* StorageUi;
+		ITypedElementDataStorageCompatibilityInterface* StorageCompatibility;
+
+		// Queries used to virtualize widgets when a column is added to/remove from a row
+		TArray<TypedElementDataStorage::QueryHandle> InternalObserverQueries;
+		TypedElementDataStorage::QueryHandle WidgetQuery;
+		TMap<TypedElementDataStorage::RowHandle, bool> RowsToUpdate;
+
+		// Delegate to check if a row is visible in the owning table viewer
+		FIsRowVisible IsRowVisibleDelegate;
+	};
+}
diff --git a/Engine/Plugins/Experimental/EditorDataStorageFeatures/Source/TedsTableViewer/Public/TedsTableViewerModel.h b/Engine/Plugins/Experimental/EditorDataStorageFeatures/Source/TedsTableViewer/Public/TedsTableViewerModel.h
new file mode 100644
--- /dev/null
+++ b/Engine/Plugins/Experimental/EditorDataStorageFeatures/Source/TedsTableViewer/Public/TedsTableViewerModel.h
@@ -0,0 +1,109 @@
+// Copyright Epic Games, Inc. All Rights Reserved.
+
+#pragma once
+
+#include "Containers/Ticker.h"
+#include "Delegates/Delegate.h"
+#include "Elements/Common/TypedElementHandles.h"
+#include "Templates/SharedPointer.h"
+#include "UObject/NameTypes.h"
+
+class ITypedElementDataStorageCompatibilityInterface;
+class ITypedElementDataStorageUiInterface;
+class ITypedElementDataStorageInterface;
+
+namespace UE::EditorDataStorage
+{
+	class FTedsTableViewerColumn;
+	class IQueryStackNode_Row;
+	
+	// Struct representing a row in the table viewer's UI
+	struct TedsTableViewerItem
+	{
+		TypedElementDataStorage::RowHandle RowHandle;
+
+		TedsTableViewerItem(TypedElementDataStorage::RowHandle InRowHandle)
+			: RowHandle(InRowHandle)
+		{
+			
+		}
+	};
+	
+	// Typedef for an item in the table viewer
+	typedef TSharedPtr<TedsTableViewerItem> TableViewerItemPtr;
+
+	// Model class for the Teds Table Viewer that can be plugged into any widget that is a UI representation of data in Teds
+	// @see STedsTableViewer
+	class FTedsTableViewerModel
+	{
+	public:
+
+		DECLARE_DELEGATE_RetVal_OneParam(bool, FIsItemVisible, TableViewerItemPtr);
+		DECLARE_MULTICAST_DELEGATE(FOnModelChanged);
+
+		TEDSTABLEVIEWER_API FTedsTableViewerModel(const TSharedPtr<IQueryStackNode_Row>& RowQueryStack, const TArray<TWeakObjectPtr<const UScriptStruct>>& RequestedColumns,
+			const TArray<FName>& CellWidgetPurposes, const FIsItemVisible& InIsItemVisibleDelegate);
+
+		~FTedsTableViewerModel();
+
+		// Get the items this table viewer is viewing
+		TEDSTABLEVIEWER_API const TArray<TableViewerItemPtr>& GetItems() const;
+		
+		// Get a specific column that the table viewer is displaying by name
+		TEDSTABLEVIEWER_API TSharedPtr<FTedsTableViewerColumn> GetColumn(const FName& ColumnName) const;
+
+		// Execute a delegate for each column in the model
+		TEDSTABLEVIEWER_API void ForEachColumn(const TFunctionRef<void(const TSharedRef<FTedsTableViewerColumn>&)>& Delegate) const;
+
+		// Delegate when the item list changes
+		TEDSTABLEVIEWER_API FOnModelChanged& GetOnModelChanged();
+		
+		// Clear the current list of columns being displayed and set it to the given list
+		TEDSTABLEVIEWER_API void SetColumns(const TArray<TWeakObjectPtr<const UScriptStruct>>& InColumns);
+
+	protected:
+
+		// Generate the actual columns to display in the UI using TEDS UI
+		void GenerateColumns();
+
+		// Check if the given row is currently visible in the UI
+		bool IsRowVisible(TypedElementDataStorage::RowHandle InRowHandle) const;
+
+		bool Tick(float DeltaTime);
+
+		void Refresh();
+
+	private:
+
+		// The row query stack used to supply the rows to display
+		TSharedPtr<IQueryStackNode_Row> RowQueryStack;
+
+		// The cached list of rows we are currently displaying
+		TArray<TableViewerItemPtr> Items;
+
+		// List of columns the table viewer is currently displaying
+		TMap<FName, TSharedRef<FTedsTableViewerColumn>> ColumnsView;
+
+		// The initial TEDS columns the widget was requested to display
+		TArray<TWeakObjectPtr<const UScriptStruct>> RequestedTedsColumns;
+
+		// The widget purposes used to create widgets in this table viewer
+		TArray<FName> CellWidgetPurposes;
+
+		// Cached revision ID for the query stack used to check when the table viewer needs a refresh
+		uint32 CachedRowQueryStackRevision = 0;
+
+		// Delegate supplied by the widget to check if an item is visible in the UI currently
+		FIsItemVisible IsItemVisible;
+
+		FTSTicker::FDelegateHandle TickerHandle;
+
+		// Delegate executed when the row list changes
+		FOnModelChanged OnModelChanged;
+		
+		// Teds Constructs
+		ITypedElementDataStorageInterface* Storage = nullptr;
+		ITypedElementDataStorageUiInterface* StorageUi = nullptr;
+		ITypedElementDataStorageCompatibilityInterface* StorageCompatibility = nullptr;
+	};
+}
\ No newline at end of file
diff --git a/Engine/Plugins/Experimental/EditorDataStorageFeatures/Source/TedsTableViewer/Public/TedsTableViewerUtils.h b/Engine/Plugins/Experimental/EditorDataStorageFeatures/Source/TedsTableViewer/Public/TedsTableViewerUtils.h
new file mode 100644
--- /dev/null
+++ b/Engine/Plugins/Experimental/EditorDataStorageFeatures/Source/TedsTableViewer/Public/TedsTableViewerUtils.h
@@ -0,0 +1,46 @@
+// Copyright Epic Games, Inc. All Rights Reserved.
+
+#pragma once
+
+#include "Containers/ArrayView.h"
+#include "Elements/Interfaces/TypedElementDataStorageFactory.h"
+#include "UObject/Class.h"
+
+#include "TedsTableViewerUtils.generated.h"
+
+namespace TypedElementDataStorage
+{
+	class FMetaDataView;
+}
+class ITypedElementDataStorageInterface;
+class ITypedElementDataStorageUiInterface;
+struct FTypedElementWidgetConstructor;
+
+// Util library for functions shared by the Teds Table Viewer and the Teds Outliner
+namespace UE::EditorDataStorage::TableViewerUtils
+{
+	TEDSTABLEVIEWER_API FName GetWidgetTableName();
+	
+	// Find the longest matching common column name given a list of columns
+	TEDSTABLEVIEWER_API FName FindLongestMatchingName(TConstArrayView<TWeakObjectPtr<const UScriptStruct>> ColumnTypes, int32 DefaultNameIndex);
+
+	// Create a header widget constructor for the given columns
+	TEDSTABLEVIEWER_API TSharedPtr<FTypedElementWidgetConstructor> CreateHeaderWidgetConstructor(ITypedElementDataStorageUiInterface& StorageUi, 
+	const TypedElementDataStorage::FMetaDataView& InMetaData, TConstArrayView<TWeakObjectPtr<const UScriptStruct>> ColumnTypes,
+	const TConstArrayView<FName> CellWidgetPurposes);
+
+	// Create a copy of the provided column types array after discarding invalid entries
+	TEDSTABLEVIEWER_API TArray<TWeakObjectPtr<const UScriptStruct>> CreateVerifiedColumnTypeArray(TConstArrayView<TWeakObjectPtr<const UScriptStruct>> ColumnTypes);
+}
+
+UCLASS()
+class UTypedElementTableViewerFactory : public UTypedElementDataStorageFactory
+{
+	GENERATED_BODY()
+
+public:
+	~UTypedElementTableViewerFactory() override = default;
+
+	void RegisterTables(ITypedElementDataStorageInterface& DataStorage) override;
+};
+
diff --git a/Engine/Plugins/Experimental/EditorDataStorageFeatures/Source/TedsTableViewer/Public/Widgets/STedsTableViewer.h b/Engine/Plugins/Experimental/EditorDataStorageFeatures/Source/TedsTableViewer/Public/Widgets/STedsTableViewer.h
new file mode 100644
--- /dev/null
+++ b/Engine/Plugins/Experimental/EditorDataStorageFeatures/Source/TedsTableViewer/Public/Widgets/STedsTableViewer.h
@@ -0,0 +1,74 @@
+// Copyright Epic Games, Inc. All Rights Reserved.
+#pragma once
+
+#include "Elements/Common/TypedElementHandles.h"
+#include "QueryStack/IQueryStackNode_Row.h"
+#include "TedsTableViewerModel.h"
+#include "Widgets/SCompoundWidget.h"
+#include "Widgets/Views/STableRow.h"
+#include "Widgets/Views/SListView.h"
+
+class ITableRow;
+class STableViewBase;
+class SHeaderRow;
+
+
+namespace UE::EditorDataStorage
+{
+	class FTedsTableViewerColumn;
+
+	/*
+	 * A table viewer widget can be used to show a visual representation of data in TEDS.
+	 * The rows to display can be specified using a RowQueryStack, and the columns to display are directly input into the widget
+	 * Example usage:
+	 * 
+	 *	SNew(STedsTableViewer)
+     *		.QueryStack(MakeShared<UE::EditorDataStorage::FQueryStackNode_RowView>(&Rows))
+	 *		.Columns({FTypedElementLabelColumn::StaticStruct(), FTypedElementClassTypeInfoColumn::StaticStruct());
+	 */
+	class STedsTableViewer : public SCompoundWidget
+	{
+		SLATE_BEGIN_ARGS(STedsTableViewer)
+			: _CellWidgetPurposes({TEXT("General.Cell")})
+		{
+			
+		}
+		
+
+		// Query Stack that will supply the rows to be displayed
+		SLATE_ARGUMENT(TSharedPtr<IQueryStackNode_Row>, QueryStack)
+
+		// The Columns that this table viewer will display
+		// Table Viewer TODO: How do we specify column metadata (ReadOnly or ReadWrite)?
+		SLATE_ARGUMENT(TArray<TWeakObjectPtr<const UScriptStruct>>, Columns)
+
+		// The widget purposes to use to create the widgets
+		SLATE_ARGUMENT(TArray<FName>, CellWidgetPurposes)
+
+		SLATE_END_ARGS()
+
+	public:
+		
+		TEDSTABLEVIEWER_API void Construct(const FArguments& InArgs);
+
+		// Clear the current list of columns being displayed and set it to the given list
+		TEDSTABLEVIEWER_API void SetColumns(const TArray<TWeakObjectPtr<const UScriptStruct>>& Columns) const;
+
+	protected:
+		
+		TSharedRef<ITableRow> MakeTableRowWidget(TableViewerItemPtr InItem, const TSharedRef<STableViewBase>& OwnerTable) const;
+
+		bool IsItemVisible(TableViewerItemPtr InItem) const;
+
+	private:
+
+		// The actual ListView widget that displays the rows
+		TSharedPtr<SListView<TableViewerItemPtr>> ListView;
+
+		// The actual header widget
+		TSharedPtr<SHeaderRow> HeaderRowWidget;
+
+		// Our model class
+		TSharedPtr<FTedsTableViewerModel> Model;
+	};
+}
\ No newline at end of file
diff --git a/Engine/Plugins/Experimental/EditorDataStorageFeatures/Source/TedsTableViewer/Public/Widgets/STedsTableViewerRow.h b/Engine/Plugins/Experimental/EditorDataStorageFeatures/Source/TedsTableViewer/Public/Widgets/STedsTableViewerRow.h
new file mode 100644
--- /dev/null
+++ b/Engine/Plugins/Experimental/EditorDataStorageFeatures/Source/TedsTableViewer/Public/Widgets/STedsTableViewerRow.h
@@ -0,0 +1,38 @@
+// Copyright Epic Games, Inc. All Rights Reserved.
+
+#pragma once
+
+#include "TedsTableViewerModel.h"
+#include "Widgets/Views/SListView.h"
+#include "Widgets/Views/STableRow.h"
+
+namespace UE::EditorDataStorage
+{
+	class STedsTableViewer;
+	
+	/** Widget that represents a row.  Generates widgets for each column on demand. */
+	class STedsTableViewerRow
+		: public SMultiColumnTableRow< TableViewerItemPtr >
+	{
+
+	public:
+
+		SLATE_BEGIN_ARGS( STedsTableViewerRow ) {}
+
+		/** The list item for this row */
+		SLATE_ARGUMENT( TableViewerItemPtr, Item )
+
+		SLATE_END_ARGS()
+		
+		/** Construct function for this widget */
+		void Construct( const FArguments& InArgs, const TSharedRef<STableViewBase>& OwnerTableView, const TSharedRef<FTedsTableViewerModel>& InTableViewerModel );
+
+		/** Overridden from SMultiColumnTableRow.  Generates a widget for this column of the tree row. */
+		virtual TSharedRef<SWidget> GenerateWidgetForColumn( const FName& ColumnName ) override;
+
+	protected:
+
+		TSharedPtr<FTedsTableViewerModel> TableViewerModel;
+		TableViewerItemPtr Item;
+	};
+}
diff --git a/Engine/Plugins/Experimental/EditorDataStorageFeatures/Source/TedsTableViewer/TedsTableViewer.Build.cs b/Engine/Plugins/Experimental/EditorDataStorageFeatures/Source/TedsTableViewer/TedsTableViewer.Build.cs
new file mode 100644
--- /dev/null
+++ b/Engine/Plugins/Experimental/EditorDataStorageFeatures/Source/TedsTableViewer/TedsTableViewer.Build.cs
@@ -0,0 +1,34 @@
+// Copyright Epic Games, Inc. All Rights Reserved.
+
+using UnrealBuildTool;
+
+public class TedsTableViewer : ModuleRules
+{
+	public TedsTableViewer(ReadOnlyTargetRules Target) : base(Target)
+	{
+
+		if (Target.bBuildEditor)
+		{
+			PublicIncludePaths.AddRange(new string[] {});
+			PrivateIncludePaths.AddRange(new string[] {});
+
+			PublicDependencyModuleNames.AddRange(
+				new string[]
+				{
+					"TypedElementFramework"
+				});
+
+			PrivateDependencyModuleNames.AddRange(
+				new string[]
+				{
+					"Core",
+					"CoreUObject",
+					"Engine",
+					"InputCore",
+					"Slate",
+					"SlateCore",
+					"WorkspaceMenuStructure"
+				});
+		}
+	}
+}
\ No newline at end of file
diff --git a/Engine/Source/Runtime/TypedElementFramework/Public/Elements/Interfaces/TypedElementDataStorageUiInterface.h b/Engine/Source/Runtime/TypedElementFramework/Public/Elements/Interfaces/TypedElementDataStorageUiInterface.h
--- a/Engine/Source/Runtime/TypedElementFramework/Public/Elements/Interfaces/TypedElementDataStorageUiInterface.h
+++ b/Engine/Source/Runtime/TypedElementFramework/Public/Elements/Interfaces/TypedElementDataStorageUiInterface.h
@@ -165,6 +165,10 @@ public:
 	 * Register a widget purpose. Widget purposes indicates how widgets can be used and categorizes/organizes the available 
 	 * widget factories. If the same purpose is registered multiple times, only the first will be recorded and later registrations
 	 * will be silently ignored.
+	 * Widget purposes follow a specific naming convention currently:
+	 * "PurposeName.Cell" for widgets created for (row, columns) pairs
+	 * "PurposeName.Header" for widgets created for column headers
+	 * "PurposeName.Cell.Default" or "PurposeName.Header.Default" for generic widgets not registered against column(s)
 	 */
 	virtual void RegisterWidgetPurpose(FName Purpose, EPurposeType Type, FText Description) = 0;
 	/** 
