Teds Debugger: Update the Teds Query Editor to use the new STedsTableViewer instead of the TedsOutliner

diff --git a/Engine/Plugins/Experimental/EditorDataStorageFeatures/Source/TedsDebugger/Private/Widgets/QueryEditor/TedsQueryEditorResultsView.cpp b/Engine/Plugins/Experimental/EditorDataStorageFeatures/Source/TedsDebugger/Private/Widgets/QueryEditor/TedsQueryEditorResultsView.cpp
--- a/Engine/Plugins/Experimental/EditorDataStorageFeatures/Source/TedsDebugger/Private/Widgets/QueryEditor/TedsQueryEditorResultsView.cpp
+++ b/Engine/Plugins/Experimental/EditorDataStorageFeatures/Source/TedsDebugger/Private/Widgets/QueryEditor/TedsQueryEditorResultsView.cpp
@@ -13,6 +13,8 @@
 #include "Elements/Framework/TypedElementQueryBuilder.h"
 #include "Modules/ModuleManager.h"
 #include "QueryEditor/TedsQueryEditorModel.h"
+#include "QueryStack/FQueryStackNode_RowView.h"
+#include "Widgets/STedsTableViewer.h"
 #include "Widgets/Text/STextBlock.h"
 
 #define LOCTEXT_NAMESPACE "TedsDebuggerModule"
@@ -30,6 +32,8 @@ void SResultsView::Construct(const FArguments& InArgs, FTedsQueryEditorModel& In
 	ModelChangedDelegateHandle = Model->GetModelChangedDelegate().AddRaw(this, &SResultsView::OnModelChanged);
 	// RowQueryHandle = TypedElementDataStorage::InvalidQueryHandle;
 
+	RowQueryStack = MakeShared<UE::EditorDataStorage::FQueryStackNode_RowView>(&TableViewerRows);
+
 	ChildSlot
 	[
 		SNew(SVerticalBox)
@@ -42,12 +46,8 @@ void SResultsView::Construct(const FArguments& InArgs, FTedsQueryEditorModel& In
 		]
 		+SVerticalBox::Slot()
 		[
-			SAssignNew(TableViewHolder, SHorizontalBox)
-			+SHorizontalBox::Slot()
-			[
-				// Placeholder for Table View
-				SNullWidget::NullWidget
-			]
+			SAssignNew(TableViewer, UE::EditorDataStorage::STedsTableViewer)
+			.QueryStack(RowQueryStack)
 		]
 		+SVerticalBox::Slot()
 		.AutoHeight()
@@ -77,106 +77,69 @@ void SResultsView::Tick(const FGeometry& AllottedGeometry, const double InCurren
 
 	ITypedElementDataStorageInterface& TedsInterface = Model->GetTedsInterface();
 
+	if(bModelDirty)
 	{
-		TypedElementDataStorage::FQueryDescription CountQueryDescription = Model->GenerateNoSelectQueryDescription();
-		
-		if (CountQueryHandle != TypedElementDataStorage::InvalidQueryHandle)
 		{
-			TedsInterface.UnregisterQuery(CountQueryHandle);
-			CountQueryHandle = TypedElementDataStorage::InvalidQueryHandle;
-		}
-		CountQueryHandle = TedsInterface.RegisterQuery(MoveTemp(CountQueryDescription));
-	}
-
-	{
-		// TODO: For now, destroy and re-create the table view
-		// This could be all that is required at some point
-		//
-		// FTypedElementSceneOutlinerQueryBinder::GetInstance().AssignQuery(SelectQueryHandle, TableView);
-		// TableView->FullRefresh();
+			TypedElementDataStorage::FQueryDescription CountQueryDescription = Model->GenerateNoSelectQueryDescription();
 		
-		TableViewHolder->ClearChildren();
-
-		FSceneOutlinerInitializationOptions InitOptions;
-		InitOptions.bShowHeaderRow = true;
-		InitOptions.bShowSearchBox = false;
-		InitOptions.FilterBarOptions.bHasFilterBar = false;
-		InitOptions.OutlinerIdentifier = "QueryEditorResultsView";
-		
-		// if (RowQueryHandle != TypedElementDataStorage::InvalidQueryHandle)
-		// {
-		// 	TedsInterface.UnregisterQuery(RowQueryHandle);
-		// 	RowQueryHandle = TypedElementDataStorage::InvalidQueryHandle;
-		// } 
-		RowQueryDescription = Model->GenerateNoSelectQueryDescription();
-		if (!RowQueryDescription.ConditionTypes.IsEmpty())
-		{
-			// Work around since query registration takes by RRef but later the
-			// FTypedElementOutlinerModeParams.QueryDescription needs a reference to it.
-			// We are just going to copy it for now...
-			// TypedElementDataStorage::FQueryDescription Copy = SelectQueryDescription;
-			RowQueryDescription.Action = TypedElementDataStorage::FQueryDescription::EActionType::Select;
-			//RowQueryHandle = TedsInterface.RegisterQuery(MoveTemp(RowQueryDescription));
+			if (CountQueryHandle != TypedElementDataStorage::InvalidQueryHandle)
+			{
+				TedsInterface.UnregisterQuery(CountQueryHandle);
+				CountQueryHandle = TypedElementDataStorage::InvalidQueryHandle;
+			}
+			CountQueryHandle = TedsInterface.RegisterQuery(MoveTemp(CountQueryDescription));
 		}
-	
-		FTypedElementOutlinerModeParams Params(nullptr);
-		Params.QueryDescription = TAttribute<TypedElementDataStorage::FQueryDescription>::CreateLambda([WeakResultsView = AsWeak()]()
+
 		{
-			if (TSharedPtr<SWidget> Shared = WeakResultsView.Pin())
+			TypedElementDataStorage::FQueryDescription TableViewerQueryDescription = Model->GenerateQueryDescription();
+
+			// Update the columns in the table viewer using the selection types from the query description
+			TableViewer->SetColumns(TArray<TWeakObjectPtr<const UScriptStruct>>(TableViewerQueryDescription.SelectionTypes));
+			
+			if (TableViewerQueryHandle != TypedElementDataStorage::InvalidQueryHandle)
 			{
-				SResultsView* This = static_cast<SResultsView*>(Shared.Get());
-				return This->RowQueryDescription;
+				TedsInterface.UnregisterQuery(TableViewerQueryHandle);
+				TableViewerQueryHandle = TypedElementDataStorage::InvalidQueryHandle;
 			}
-			return TypedElementDataStorage::FQueryDescription();
-		});
-		Params.HierarchyData = TOptional<FTypedElementOutlinerHierarchyData>(); // We don't want to show hierarchies
-		Params.CellWidgetPurposes = TArray<FName>{TEXT("General.Cell")};
 
-		TypedElementDataStorage::FQueryDescription ColumnQueryDescription = Model->GenerateQueryDescription();
-		if (ColumnQueryHandle != TypedElementDataStorage::InvalidQueryHandle)
-		{
-			TedsInterface.UnregisterQuery(ColumnQueryHandle);
-			ColumnQueryHandle = TypedElementDataStorage::InvalidQueryHandle;
+			// Mass doesn't like empty queries, so we only set it if there are actual conditions
+			if(TableViewerQueryDescription.ConditionTypes.Num())
+			{
+				TableViewerQueryHandle = TedsInterface.RegisterQuery(MoveTemp(TableViewerQueryDescription));
+			}
 		}
+		
+		bModelDirty = false;
+	}
 
-		if (!ColumnQueryDescription.ConditionTypes.IsEmpty() || !ColumnQueryDescription.SelectionTypes.IsEmpty())
-		{
-			// Work around since query registration takes by RRef but later the
-			// FTypedElementOutlinerModeParams.QueryDescription needs a reference to it.
-			// We are just going to copy it for now...
-			// TypedElementDataStorage::FQueryDescription Copy = SelectQueryDescription;
-			ColumnQueryDescription.Action = TypedElementDataStorage::FQueryDescription::EActionType::Select;
-			ColumnQueryHandle = TedsInterface.RegisterQuery(MoveTemp(ColumnQueryDescription));
-		}
+	// Every frame we re-run the query to update the rows the table viewer is showing
+	if(TableViewerQueryHandle != TypedElementDataStorage::InvalidQueryHandle)
+	{
+		TSet<TypedElementDataStorage::RowHandle> NewTableViewerRows_Set;
+		NewTableViewerRows_Set.Reserve(TableViewerRows_Set.Num());
 
-		FTedsOutlinerModule& TedsOutlinerModule = FModuleManager::GetModuleChecked<FTedsOutlinerModule>("TedsOutliner");
-		if (ColumnQueryHandle != TypedElementDataStorage::InvalidQueryHandle)
+		TypedElementDataStorage::FQueryResult QueryResult = Model->GetTedsInterface().RunQuery(TableViewerQueryHandle,
+			CreateDirectQueryCallbackBinding([&NewTableViewerRows_Set](const ITypedElementDataStorageInterface::IDirectQueryContext& Context, const TypedElementDataStorage::RowHandle*)
 		{
-			TSharedPtr<ISceneOutliner> TableView = TedsOutlinerModule.CreateTedsOutliner(InitOptions, Params, ColumnQueryHandle);
+			NewTableViewerRows_Set.Append(Context.GetRowHandles());
+		}));
 
-			TableViewHolder->AddSlot()
-			[
-				TableView.ToSharedRef()
-			];
-		}
-		else
+		// Check if the two sets are equal, i.e not changes and no need to update the table viewer
+		const bool bSetsEqual = (TableViewerRows_Set.Num() == NewTableViewerRows_Set.Num()) && TableViewerRows_Set.Includes(NewTableViewerRows_Set);
+
+		if(!bSetsEqual)
 		{
-			TableViewHolder->AddSlot()
-			[
-				// Placeholder for table view.  Mass doesn't like invalid/empty queries
-				SNullWidget::NullWidget
-			];
+			Swap(TableViewerRows_Set, NewTableViewerRows_Set);
+			TableViewerRows = TableViewerRows_Set.Array();
+			RowQueryStack->MarkDirty();
 		}
 	}
-
-	SetCanTick(false);
 }
 
 void SResultsView::OnModelChanged()
 {
-	SetCanTick(true);
+	bModelDirty = true;
 	Invalidate(EInvalidateWidgetReason::Layout);
-	
 }
 
 #undef LOCTEXT_NAMESPACE
\ No newline at end of file
diff --git a/Engine/Plugins/Experimental/EditorDataStorageFeatures/Source/TedsDebugger/Private/Widgets/QueryEditor/TedsQueryEditorResultsView.h b/Engine/Plugins/Experimental/EditorDataStorageFeatures/Source/TedsDebugger/Private/Widgets/QueryEditor/TedsQueryEditorResultsView.h
--- a/Engine/Plugins/Experimental/EditorDataStorageFeatures/Source/TedsDebugger/Private/Widgets/QueryEditor/TedsQueryEditorResultsView.h
+++ b/Engine/Plugins/Experimental/EditorDataStorageFeatures/Source/TedsDebugger/Private/Widgets/QueryEditor/TedsQueryEditorResultsView.h
@@ -8,6 +8,12 @@ class ISceneOutliner;
 class SSceneOutliner;
 class SHorizontalBox;
 
+namespace UE::EditorDataStorage
+{
+	class STedsTableViewer;
+	class FQueryStackNode_RowView;
+}
+
 namespace UE::EditorDataStorage::Debug::QueryEditor
 {
 	class FTedsQueryEditorModel;
@@ -28,14 +34,17 @@ namespace UE::EditorDataStorage::Debug::QueryEditor
 		
 		FTedsQueryEditorModel* Model = nullptr;
 		FDelegateHandle ModelChangedDelegateHandle;
+		bool bModelDirty = true;
 
-		TypedElementDataStorage::FQueryDescription RowQueryDescription;
-		TypedElementQueryHandle CountQueryHandle = TypedElementDataStorage::InvalidQueryHandle;
-		TypedElementQueryHandle ColumnQueryHandle = TypedElementDataStorage::InvalidQueryHandle;
 
-		TSharedPtr<SHorizontalBox> TableViewHolder;
+		TypedElementQueryHandle CountQueryHandle = TypedElementDataStorage::InvalidQueryHandle;
+		TypedElementQueryHandle TableViewerQueryHandle = TypedElementDataStorage::InvalidQueryHandle;
 
-		
+		TArray<TypedElementDataStorage::RowHandle> TableViewerRows;
+		// We have to keep a TSet copy because queries return duplicate rows sometimes and to have some form of sorted order for the rows for now
+		TSet<TypedElementDataStorage::RowHandle> TableViewerRows_Set;
+		TSharedPtr<UE::EditorDataStorage::STedsTableViewer> TableViewer;
+		TSharedPtr<UE::EditorDataStorage::FQueryStackNode_RowView> RowQueryStack;
 	};
 
 	
diff --git a/Engine/Plugins/Experimental/EditorDataStorageFeatures/Source/TedsDebugger/TedsDebugger.Build.cs b/Engine/Plugins/Experimental/EditorDataStorageFeatures/Source/TedsDebugger/TedsDebugger.Build.cs
--- a/Engine/Plugins/Experimental/EditorDataStorageFeatures/Source/TedsDebugger/TedsDebugger.Build.cs
+++ b/Engine/Plugins/Experimental/EditorDataStorageFeatures/Source/TedsDebugger/TedsDebugger.Build.cs
@@ -31,7 +31,9 @@ public class TedsDebugger : ModuleRules
 		            "Slate",
 		            "InputCore",
 		            "EditorWidgets",
-		            "ToolWidgets"
+		            "ToolWidgets",
+		            "EditorWidgets",
+		            "TedsTableViewer"
             	});
 			
 			DynamicallyLoadedModuleNames.AddRange(new string[] {});
