[Teds] Rename TedsRevisionControl files and classes to remove TypedElement strings

#rnx

diff --git a/Engine/Plugins/Experimental/EditorDataStorageFeatures/Source/TedsRevisionControl/Private/Compatibility/TypedElementUObjectPackagePathToColumnQueries.cpp b/Engine/Plugins/Experimental/EditorDataStorageFeatures/Source/TedsRevisionControl/Private/Compatibility/TypedElementUObjectPackagePathToColumnQueries.cpp
deleted file mode 100644
diff --git a/Engine/Plugins/Experimental/EditorDataStorageFeatures/Source/TedsRevisionControl/Private/Compatibility/TypedElementUObjectPackagePathToColumnQueries.h b/Engine/Plugins/Experimental/EditorDataStorageFeatures/Source/TedsRevisionControl/Private/Compatibility/TypedElementUObjectPackagePathToColumnQueries.h
deleted file mode 100644
diff --git a/Engine/Plugins/Experimental/EditorDataStorageFeatures/Source/TedsRevisionControl/Private/Processors/RevisionControlProcessors.cpp b/Engine/Plugins/Experimental/EditorDataStorageFeatures/Source/TedsRevisionControl/Private/Processors/RevisionControlProcessors.cpp
new file mode 100644
--- /dev/null
+++ b/Engine/Plugins/Experimental/EditorDataStorageFeatures/Source/TedsRevisionControl/Private/Processors/RevisionControlProcessors.cpp
@@ -0,0 +1,378 @@
+// Copyright Epic Games, Inc. All Rights Reserved.
+
+#include "Processors/RevisionControlProcessors.h"
+
+#include "ISourceControlModule.h"
+#include "SourceControlFileStatusMonitor.h"
+#include "GameFramework/Actor.h"
+#include "HAL/IConsoleManager.h"
+
+#include "Elements/Columns/TypedElementCompatibilityColumns.h"
+#include "Elements/Columns/TypedElementMiscColumns.h"
+#include "Elements/Columns/TypedElementPackageColumns.h"
+#include "Elements/Columns/TypedElementRevisionControlColumns.h"
+#include "Elements/Columns/TypedElementSelectionColumns.h"
+#include "Elements/Columns/TypedElementViewportColumns.h"
+#include "Elements/Common/TypedElementDataStorageLog.h"
+#include "Elements/Framework/TypedElementQueryBuilder.h"
+#include "Elements/Framework/TypedElementRegistry.h"
+#include "Elements/Interfaces/TypedElementQueryStorageInterfaces.h"
+
+extern FAutoConsoleVariableRef CVarAutoPopulateState;
+
+static bool gEnableOverlays = false;
+TAutoConsoleVariable<bool> CVarEnableOverlays(
+	TEXT("RevisionControl.Overlays.Enable"),
+	gEnableOverlays,
+	TEXT("Enables overlays."),
+	ECVF_Default);
+
+static bool gEnableOverlayCheckedOutByOtherUser = true;
+TAutoConsoleVariable<bool> CVarEnableOverlayCheckedOutByOtherUser(
+	TEXT("RevisionControl.Overlays.CheckedOutByOtherUser.Enable"),
+	gEnableOverlayCheckedOutByOtherUser,
+	TEXT("Enables overlays for files that are checked out by another user."),
+	ECVF_Default);
+
+static bool gEnableOverlayNotAtHeadRevision = true;
+TAutoConsoleVariable<bool> CVarEnableOverlayNotAtHeadRevision(
+	TEXT("RevisionControl.Overlays.NotAtHeadRevision.Enable"),
+	gEnableOverlayNotAtHeadRevision,
+	TEXT("Enables overlays for files that are not at the latest revision."),
+	ECVF_Default);
+
+static bool gEnableOverlayCheckedOut = false;
+TAutoConsoleVariable<bool> CVarEnableOverlayCheckedOut(
+	TEXT("RevisionControl.Overlays.CheckedOut.Enable"),
+	gEnableOverlayCheckedOut,
+	TEXT("Enables overlays for files that are checked out by user."),
+	ECVF_Default);
+
+static bool gEnableOverlayOpenForAdd = false;
+TAutoConsoleVariable<bool> CVarEnableOverlayOpenForAdd(
+	TEXT("RevisionControl.Overlays.OpenForAdd.Enable"),
+	gEnableOverlayOpenForAdd,
+	TEXT("Enables overlays for files that are newly added."),
+	ECVF_Default);
+
+static int32 gOverlayAlpha = 20; // [0..100]
+TAutoConsoleVariable<int32> CVarOverlayAlpha(
+	TEXT("RevisionControl.Overlays.Alpha"),
+	gOverlayAlpha,
+	TEXT("Configures overlay opacity."),
+	ECVF_Default);
+
+#if UE_BUILD_SHIPPING
+#define ENABLE_OVERLAY_DEBUG 0
+#else
+#define ENABLE_OVERLAY_DEBUG 1
+#endif
+
+#if ENABLE_OVERLAY_DEBUG
+static int32 gDefaultDebugForceColorOnAllValue = false; // [0..100]
+TAutoConsoleVariable<int32> CVarDebugForceColorOnAll(
+	TEXT("RevisionControl.Overlays.Debug.ForceColorOnAll"),
+	gDefaultDebugForceColorOnAllValue,
+	TEXT("Debug to force overlay color on everything. 1 = Red, 2 = Green, 3 = Blue, 4 = White. 0 = off  ."),
+	ECVF_Default);
+#endif
+
+static FColor DetermineOverlayColor(const TypedElementDataStorage::IQueryContext& ObjectContext, const TypedElementDataStorage::ICommonQueryContext& SCCContext, const FTypedElementUObjectColumn& Actor)
+{
+	check(IsInGameThread());
+
+#if ENABLE_OVERLAY_DEBUG
+	if (int32 Force = CVarDebugForceColorOnAll.GetValueOnGameThread())
+	{
+		int32 Alpha = FMath::Lerp<float>(0.f, 255.f, gOverlayAlpha / 100.f);
+		switch(Force)
+		{
+		case 1:
+			return FColor(255, 0, 0, Alpha);
+		case 2:
+			return FColor(0, 255, 0, Alpha);
+		case 3:
+			return FColor(0, 0, 255, Alpha);
+		case 4:
+			return FColor(255, 255, 255, Alpha);
+		// Do normal determination for higher than 4
+		}
+	}
+#endif
+
+	bool bExternal = Actor.Object.IsValid() ? Cast<AActor>(Actor.Object)->IsPackageExternal() : false;
+	bool bIgnored = !bExternal;
+	bool bSelected = ObjectContext.HasColumn<FTypedElementSelectionColumn>();
+	if (!bIgnored && !bSelected)
+	{
+		// Convert CVar value from [0..100] to [0..255] range.
+		int32 Alpha = FMath::Lerp<float>(0.f, 255.f, CVarOverlayAlpha.GetValueOnGameThread() / 100.f);
+
+		// Check if the package is outdated because there is a newer version available.
+		if (SCCContext.HasColumn<FSCCNotCurrentTag>())
+		{
+			if (CVarEnableOverlayNotAtHeadRevision.GetValueOnGameThread())
+			{
+				// Yellow.
+				return FColor(225, 255, 61, Alpha);
+			}
+		}
+
+		// Check if the package is locked by someone else.
+		if (SCCContext.HasColumn<FSCCExternallyLockedColumn>())
+		{
+			if (CVarEnableOverlayCheckedOutByOtherUser.GetValueOnGameThread())
+			{
+				// Red.
+				return FColor(239, 53, 53, Alpha);
+			}
+		}
+
+		// Check if the package is added locally.
+		if (SCCContext.HasColumn<FSCCStatusColumn>())
+		{
+			if (CVarEnableOverlayOpenForAdd.GetValueOnGameThread())
+			{
+				if (const FSCCStatusColumn* StatusColumn = SCCContext.GetColumn<FSCCStatusColumn>())
+				{
+					if (StatusColumn->Modification == ESCCModification::Added)
+					{
+						// Blue.
+						return FColor(0, 112, 224, Alpha);
+					}
+				}
+			}
+		}
+
+		// Check if the package is locked by self.
+		if (SCCContext.HasColumn<FSCCLockedTag>())
+		{
+			if (CVarEnableOverlayCheckedOut.GetValueOnGameThread())
+			{
+				// Green.
+				return FColor(31, 228, 75, Alpha);
+			}
+		}
+	}
+
+	return FColor(ForceInitToZero);
+}
+
+void URevisionControlDataStorageFactory::RegisterTables(ITypedElementDataStorageInterface& DataStorage)
+{
+	DataStorage.RegisterTable(
+		TTypedElementColumnTypeList<
+			FTypedElementPackagePathColumn, FTypedElementPackageLoadedPathColumn,
+			FSCCRevisionIdColumn, FSCCExternalRevisionIdColumn>(),
+		FName("Editor_RevisionControlTable"));
+}
+
+void URevisionControlDataStorageFactory::RegisterQueries(ITypedElementDataStorageInterface& DataStorage)
+{
+	CVarAutoPopulateState->AsVariable()->OnChangedDelegate().AddLambda(
+		[this, &DataStorage](IConsoleVariable* AutoPopulate)
+		{
+			if (AutoPopulate->GetBool())
+			{
+				RegisterFetchUpdates(DataStorage);
+			}
+			else
+			{
+				DataStorage.UnregisterQuery(FetchUpdates);
+				FetchUpdates = TypedElementInvalidQueryHandle;
+			}
+		}
+	);
+
+	CVarEnableOverlays->AsVariable()->OnChangedDelegate().AddLambda(
+		[this, &DataStorage](IConsoleVariable* EnableOverlays)
+		{
+			if (EnableOverlays->GetBool())
+			{
+				DataStorage.UnregisterQuery(RemoveOverlays);
+				RemoveOverlays = TypedElementInvalidQueryHandle;
+
+				RegisterApplyOverlays(DataStorage);
+			}
+			else
+			{
+				DataStorage.UnregisterQuery(ApplyNewOverlays);
+				ApplyNewOverlays = TypedElementInvalidQueryHandle;
+				DataStorage.UnregisterQuery(ChangeOverlay);
+				ChangeOverlay = TypedElementInvalidQueryHandle;
+
+				DataStorage.UnregisterQuery(ApplyOverlaysObjectToSCC);
+				ApplyOverlaysObjectToSCC = TypedElementInvalidQueryHandle;
+
+				RegisterRemoveOverlays(DataStorage);
+			}
+		}
+	);
+	
+	if (CVarAutoPopulateState->GetBool())
+	{
+		RegisterFetchUpdates(DataStorage);
+	}
+
+	if (CVarEnableOverlays->GetBool())
+	{
+		RegisterApplyOverlays(DataStorage);
+	}
+	else
+	{
+		RegisterRemoveOverlays(DataStorage);
+	}
+}
+
+void URevisionControlDataStorageFactory::RegisterFetchUpdates(ITypedElementDataStorageInterface& DataStorage)
+{
+	using namespace TypedElementQueryBuilder;
+	using DSI = ITypedElementDataStorageInterface;
+	
+	FSourceControlFileStatusMonitor& FileStatusMonitor = ISourceControlModule::Get().GetSourceControlFileStatusMonitor();
+
+	if (FetchUpdates == TypedElementDataStorage::InvalidQueryHandle)
+	{
+		FetchUpdates = DataStorage.RegisterQuery(
+			Select(
+				TEXT("Gather source control statuses for objects with unresolved package paths"),
+				FObserver::OnAdd<FTypedElementPackageUnresolvedReference>()
+				.ForceToGameThread(true),
+				[this, &FileStatusMonitor](DSI::IQueryContext& Context, const FTypedElementPackageUnresolvedReference& UnresolvedReference)
+				{
+					static FSourceControlFileStatusMonitor::FOnSourceControlFileStatus EmptyDelegate{};
+				
+					FileStatusMonitor.StartMonitoringFile(
+						reinterpret_cast<uintptr_t>(this),
+						UnresolvedReference.PathOnDisk,
+						EmptyDelegate
+						
+					);
+				}
+			)
+			.Compile()
+		);
+	}
+}
+
+void URevisionControlDataStorageFactory::RegisterApplyOverlays(ITypedElementDataStorageInterface& DataStorage)
+{
+	using namespace TypedElementQueryBuilder;
+	using DSI = ITypedElementDataStorageInterface;
+
+	if (ApplyOverlaysObjectToSCC == TypedElementInvalidQueryHandle)
+	{
+		ApplyOverlaysObjectToSCC = DataStorage.RegisterQuery(
+			Select()
+				.ReadOnly<FTypedElementPackagePathColumn>()
+				.ReadOnly<FSCCStatusColumn, FSCCExternallyLockedColumn, FTypedElementViewportOverlayColorColumn>(EOptional::Yes)
+			.Compile());
+		}
+
+	if (ApplyNewOverlays == TypedElementInvalidQueryHandle)
+	{
+		// Query:
+		// For all actors without an overlay color column AND having a package reference:
+		//   Determine if a color should be applied based on SCC status tags
+		//   If so, add the OverlayColorColumn to the actor row
+		ApplyNewOverlays = DataStorage.RegisterQuery(
+			Select(
+				TEXT("Poll to add selection overlay colors based on SCC status"),
+				// This is in PrePhysics because the overlay->actor query is in DuringPhysics and contexts don't flush changes between tick groups
+				FProcessor(DSI::EQueryTickPhase::PrePhysics, DataStorage.GetQueryTickGroupName(DSI::EQueryTickGroups::SyncExternalToDataStorage))
+					.ForceToGameThread(true),
+				[](DSI::IQueryContext& Context, TypedElementRowHandle ObjectRow, const FTypedElementUObjectColumn& Actor, const FTypedElementPackageReference& PackageReference)
+				{
+					Context.RunSubquery(0, PackageReference.Row, CreateSubqueryCallbackBinding(
+						[&Context, &ObjectRow, &Actor](DSI::ISubqueryContext& SubQueryContext)
+						{
+							FColor Color = DetermineOverlayColor(Context, SubQueryContext, Actor);
+							if (Color.Bits != 0)
+							{
+								Context.AddColumn<FTypedElementViewportOverlayColorColumn>(ObjectRow, { .OverlayColor = Color });
+							}
+						})
+					);
+				}
+			)
+			.Where()
+				.All<FTypedElementActorTag>()
+				.None<FTypedElementViewportOverlayColorColumn>()
+			.DependsOn()
+				.SubQuery(ApplyOverlaysObjectToSCC)
+			.Compile()
+		);
+	}
+
+	if (ChangeOverlay == TypedElementInvalidQueryHandle)
+	{
+		// Query:
+		// For all actors WITH an overlay color column AND having a package reference:
+		//   Re-check the color that should be applied based on SCC status tags
+		//   If the color has changed, remove and re-add the OverlayColorColumn to the actor row
+		//
+		// Note: Remove and re-add will trigger observer in TypedElementActorViewportProcessors to SetOverlayColor on the primitive components
+		ChangeOverlay = DataStorage.RegisterQuery(
+			Select(
+				TEXT("Poll for changed selection overlay colors based on SCC status"),
+				// This is in PrePhysics because the overlay->actor query is in DuringPhysics and contexts don't flush changes between tick groups
+				FProcessor(DSI::EQueryTickPhase::PrePhysics, DataStorage.GetQueryTickGroupName(DSI::EQueryTickGroups::SyncExternalToDataStorage))
+					.ForceToGameThread(true),
+				[](DSI::IQueryContext& Context, TypedElementRowHandle ObjectRow, const FTypedElementUObjectColumn& Actor, const FTypedElementPackageReference& PackageReference, const FTypedElementViewportOverlayColorColumn& OverlayColorColumn)
+				{
+					Context.RunSubquery(0, PackageReference.Row, CreateSubqueryCallbackBinding(
+						[&Context, &ObjectRow, &Actor, &OverlayColorColumn](DSI::ISubqueryContext& SubQueryContext)
+						{
+							FColor Color = DetermineOverlayColor(Context, SubQueryContext, Actor);
+							if (Color.Bits == 0)
+							{
+								Context.RemoveColumns<FTypedElementViewportOverlayColorColumn>(ObjectRow);
+							}
+							else if (Color != OverlayColorColumn.OverlayColor)
+							{
+								// Remove and re-add to trigger the observer
+								Context.RemoveColumns<FTypedElementViewportOverlayColorColumn>(ObjectRow);
+								Context.AddColumn<FTypedElementViewportOverlayColorColumn>(ObjectRow, { .OverlayColor = Color });
+							}
+						})
+					);
+				}
+			)
+			.Where()
+				.All<FTypedElementActorTag>()
+			.DependsOn()
+				.SubQuery(ApplyOverlaysObjectToSCC)
+			.Compile()
+		);
+	}
+}
+
+void URevisionControlDataStorageFactory::RegisterRemoveOverlays(ITypedElementDataStorageInterface& DataStorage)
+{
+	using namespace TypedElementQueryBuilder;
+	using DSI = ITypedElementDataStorageInterface;
+
+	if (RemoveOverlays == TypedElementInvalidQueryHandle)
+	{
+		// Query:
+		// For all actors WITH an overlay color column AND having a package reference:
+		//   Remove the overlay color column
+		//
+		// This query is used to clean up the color columns if the overlay feature is disabled dynamically
+		RemoveOverlays = DataStorage.RegisterQuery(
+			Select(
+				TEXT("Remove selection overlay colors"),
+				// This is in PrePhysics because the overlay->actor query is in DuringPhysics and contexts don't flush changes between tick groups
+				FProcessor(DSI::EQueryTickPhase::PrePhysics, DataStorage.GetQueryTickGroupName(DSI::EQueryTickGroups::SyncExternalToDataStorage))
+				.ForceToGameThread(true),
+				[](DSI::IQueryContext& Context, TypedElementRowHandle ObjectRow, FTypedElementUObjectColumn& Actor, const FTypedElementViewportOverlayColorColumn& ViewportColor)
+				{
+					Context.RemoveColumns<FTypedElementViewportOverlayColorColumn>(ObjectRow);
+				}
+			)
+			.Where()
+				.All<FTypedElementActorTag>()
+			.Compile()
+		);
+	}
+}
\ No newline at end of file
diff --git a/Engine/Plugins/Experimental/EditorDataStorageFeatures/Source/TedsRevisionControl/Private/Processors/RevisionControlProcessors.h b/Engine/Plugins/Experimental/EditorDataStorageFeatures/Source/TedsRevisionControl/Private/Processors/RevisionControlProcessors.h
new file mode 100644
--- /dev/null
+++ b/Engine/Plugins/Experimental/EditorDataStorageFeatures/Source/TedsRevisionControl/Private/Processors/RevisionControlProcessors.h
@@ -0,0 +1,33 @@
+// Copyright Epic Games, Inc. All Rights Reserved.
+
+#pragma once
+
+#include "Elements/Common/TypedElementHandles.h"
+#include "Elements/Interfaces/TypedElementDataStorageFactory.h"
+
+#include "RevisionControlProcessors.generated.h"
+
+class ITypedElementDataStorageInterface;
+
+UCLASS()
+class URevisionControlDataStorageFactory : public UTypedElementDataStorageFactory
+{
+	GENERATED_BODY()
+
+public:
+	~URevisionControlDataStorageFactory() override = default;
+
+	void RegisterTables(ITypedElementDataStorageInterface& DataStorage) override;
+	void RegisterQueries(ITypedElementDataStorageInterface& DataStorage) override;
+
+private:
+	void RegisterFetchUpdates(ITypedElementDataStorageInterface& DataStorage);
+	void RegisterApplyOverlays(ITypedElementDataStorageInterface& DataStorage);
+	void RegisterRemoveOverlays(ITypedElementDataStorageInterface& DataStorage);
+	TypedElementDataStorage::QueryHandle FetchUpdates = TypedElementDataStorage::InvalidQueryHandle;
+	TypedElementDataStorage::QueryHandle ApplyNewOverlays = TypedElementDataStorage::InvalidQueryHandle;
+	TypedElementDataStorage::QueryHandle ChangeOverlay = TypedElementDataStorage::InvalidQueryHandle;
+	TypedElementDataStorage::QueryHandle ApplyOverlaysObjectToSCC = TypedElementDataStorage::InvalidQueryHandle;
+	TypedElementDataStorage::QueryHandle RemoveOverlays = TypedElementDataStorage::InvalidQueryHandle;
+	TypedElementDataStorage::QueryHandle FlushPackageUpdates = TypedElementDataStorage::InvalidQueryHandle;
+};
diff --git a/Engine/Plugins/Experimental/EditorDataStorageFeatures/Source/TedsRevisionControl/Private/Queries/ObjectPackagePathToColumnQueries.cpp b/Engine/Plugins/Experimental/EditorDataStorageFeatures/Source/TedsRevisionControl/Private/Queries/ObjectPackagePathToColumnQueries.cpp
new file mode 100644
--- /dev/null
+++ b/Engine/Plugins/Experimental/EditorDataStorageFeatures/Source/TedsRevisionControl/Private/Queries/ObjectPackagePathToColumnQueries.cpp
@@ -0,0 +1,131 @@
+// Copyright Epic Games, Inc. All Rights Reserved.
+
+#include "Queries/ObjectPackagePathToColumnQueries.h"
+
+#include "Elements/Columns/TypedElementCompatibilityColumns.h"
+#include "Elements/Columns/TypedElementPackageColumns.h"
+#include "Elements/Framework/TypedElementIndexHasher.h"
+#include "Elements/Framework/TypedElementQueryBuilder.h"
+#include "HAL/IConsoleManager.h"
+#include "Misc/PackageName.h"
+#include "Misc/Paths.h"
+#include "UObject/Package.h"
+
+static bool bAutoPopulateRevisionControlState = false;
+FAutoConsoleVariableRef CVarAutoPopulateState(
+	TEXT("TEDS.RevisionControl.AutoPopulateState"),
+	bAutoPopulateRevisionControlState,
+	TEXT("Automatically query revision control provider and fill information into TEDS")
+);
+
+static void ResolvePackageReference(ITypedElementDataStorageInterface::IQueryContext& Context, const UPackage* Package, TypedElementRowHandle Row, TypedElementRowHandle PackageRow)
+{
+	FTypedElementPackageReference PackageReference;
+	PackageReference.Row = PackageRow;
+	Context.AddColumn(Row, MoveTemp(PackageReference));
+				
+	FTypedElementPackagePathColumn PathColumn;
+	FTypedElementPackageLoadedPathColumn LoadedPathColumn;
+		
+	Package->GetPathName(nullptr, PathColumn.Path);
+	LoadedPathColumn.LoadedPath = Package->GetLoadedPath();
+				
+	Context.AddColumn(PackageRow, MoveTemp(PathColumn));
+	Context.AddColumn(PackageRow, MoveTemp(LoadedPathColumn));
+
+	FTypedElementPackageReference PackageBackReference;
+	PackageBackReference.Row = Row;
+	Context.AddColumn(PackageRow, MoveTemp(PackageBackReference));
+};
+
+void UTypedElementUObjectPackagePathFactory::RegisterQueries(ITypedElementDataStorageInterface& DataStorage)
+{
+	using namespace TypedElementQueryBuilder;
+	using DSI = ITypedElementDataStorageInterface;
+
+	CVarAutoPopulateState->AsVariable()->OnChangedDelegate().AddLambda(
+		[this, &DataStorage](IConsoleVariable* AutoPopulate)
+		{
+			if (AutoPopulate->GetBool())
+			{
+				RegisterTryAddPackageRef(DataStorage);
+			}
+			else
+			{
+				DataStorage.UnregisterQuery(TryAddPackageRef);
+			}
+		}
+	);
+	DataStorage.RegisterQuery(
+		Select(
+			TEXT("Resolve package references"),
+			FProcessor(DSI::EQueryTickPhase::FrameEnd, DataStorage.GetQueryTickGroupName(DSI::EQueryTickGroups::SyncExternalToDataStorage)),
+			[](DSI::IQueryContext& Context, TypedElementRowHandle Row, const FTypedElementUObjectColumn& Object, const FTypedElementPackageUnresolvedReference& UnresolvedPackageReference)
+			{
+				TypedElementRowHandle PackageRow = Context.FindIndexedRow(UnresolvedPackageReference.Index);
+				if (!Context.IsRowAvailable(PackageRow))
+				{
+					return;
+				}
+				const UObject* ObjectInstance = Object.Object.Get();
+				if (!ObjectInstance)
+				{
+					return;
+				}
+				const UPackage* Package = ObjectInstance->GetPackage();
+				Context.RemoveColumns(Row, { FTypedElementPackageUnresolvedReference::StaticStruct() });
+
+				ResolvePackageReference(Context, Package, Row, PackageRow);
+			}
+		)
+		.Compile()
+	);
+
+	if (CVarAutoPopulateState->GetBool())
+	{
+		RegisterTryAddPackageRef(DataStorage);
+	}
+}
+
+void UTypedElementUObjectPackagePathFactory::RegisterTryAddPackageRef(ITypedElementDataStorageInterface& DataStorage)
+{
+	using namespace TypedElementQueryBuilder;
+	using DSI = ITypedElementDataStorageInterface;
+	
+	TryAddPackageRef = DataStorage.RegisterQuery(
+		Select(
+			TEXT("Sync UObject package info to columns"),
+			FObserver::OnAdd<FTypedElementUObjectColumn>()
+				.ForceToGameThread(true),
+			[](DSI::IQueryContext& Context, TypedElementRowHandle Row, const FTypedElementUObjectColumn& Object)
+			{
+				if (const UObject* ObjectInstance = Object.Object.Get(); ObjectInstance != nullptr)
+				{
+					const UPackage* Target = ObjectInstance->GetPackage();
+					FString Path;
+					Target->GetPathName(nullptr, Path);
+					if (FString PackageFilename; FPackageName::TryConvertLongPackageNameToFilename(Path, PackageFilename))
+					{
+						FPaths::NormalizeFilename(PackageFilename);
+						FString FullPackageFilename = FPaths::ConvertRelativePathToFull(PackageFilename);
+						TypedElementDataStorage::IndexHash Index = TypedElementDataStorage::GenerateIndexHash(FullPackageFilename);
+						TypedElementRowHandle PackageRow = Context.FindIndexedRow(Index);
+						if (Context.IsRowAvailable(PackageRow))
+						{
+							const UPackage* Package = ObjectInstance->GetPackage();
+							ResolvePackageReference(Context, Package, Row, PackageRow);
+						}
+						else
+						{
+							FTypedElementPackageUnresolvedReference UnresolvedPackageReference;
+							UnresolvedPackageReference.Index = Index;
+							UnresolvedPackageReference.PathOnDisk = MoveTemp(FullPackageFilename);
+							Context.AddColumn(Row, MoveTemp(UnresolvedPackageReference));
+						}
+					}
+				}
+			}
+		)
+		.Compile()
+	);
+}
diff --git a/Engine/Plugins/Experimental/EditorDataStorageFeatures/Source/TedsRevisionControl/Private/Queries/ObjectPackagePathToColumnQueries.h b/Engine/Plugins/Experimental/EditorDataStorageFeatures/Source/TedsRevisionControl/Private/Queries/ObjectPackagePathToColumnQueries.h
new file mode 100644
--- /dev/null
+++ b/Engine/Plugins/Experimental/EditorDataStorageFeatures/Source/TedsRevisionControl/Private/Queries/ObjectPackagePathToColumnQueries.h
@@ -0,0 +1,25 @@
+// Copyright Epic Games, Inc. All Rights Reserved.
+
+#pragma once
+
+#include "Elements/Common/TypedElementHandles.h"
+#include "Elements/Interfaces/TypedElementDataStorageFactory.h"
+
+#include "ObjectPackagePathToColumnQueries.generated.h"
+
+class ITypedElementDataStorageInterface;
+
+UCLASS()
+class UTypedElementUObjectPackagePathFactory : public UTypedElementDataStorageFactory
+{
+	GENERATED_BODY()
+
+public:
+	~UTypedElementUObjectPackagePathFactory() override = default;
+
+	void RegisterQueries(ITypedElementDataStorageInterface& DataStorage) override;
+
+private:
+	void RegisterTryAddPackageRef(ITypedElementDataStorageInterface& DataStorage);
+	TypedElementDataStorage::QueryHandle TryAddPackageRef = TypedElementDataStorage::InvalidQueryHandle;
+};
diff --git a/Engine/Plugins/Experimental/EditorDataStorageFeatures/Source/TedsRevisionControl/Private/RevisionControlProcessors.cpp b/Engine/Plugins/Experimental/EditorDataStorageFeatures/Source/TedsRevisionControl/Private/RevisionControlProcessors.cpp
deleted file mode 100644
diff --git a/Engine/Plugins/Experimental/EditorDataStorageFeatures/Source/TedsRevisionControl/Private/RevisionControlProcessors.h b/Engine/Plugins/Experimental/EditorDataStorageFeatures/Source/TedsRevisionControl/Private/RevisionControlProcessors.h
deleted file mode 100644
diff --git a/Engine/Plugins/Experimental/EditorDataStorageFeatures/Source/TedsRevisionControl/Private/TedsRevisionControlModule.cpp b/Engine/Plugins/Experimental/EditorDataStorageFeatures/Source/TedsRevisionControl/Private/TedsRevisionControlModule.cpp
new file mode 100644
--- /dev/null
+++ b/Engine/Plugins/Experimental/EditorDataStorageFeatures/Source/TedsRevisionControl/Private/TedsRevisionControlModule.cpp
@@ -0,0 +1,22 @@
+// Copyright Epic Games, Inc. All Rights Reserved.
+
+#include "TedsRevisionControlModule.h"
+
+void FTedsRevisionControlModule::StartupModule()
+{
+}
+
+void FTedsRevisionControlModule::ShutdownModule()
+{
+}
+
+void FTedsRevisionControlModule::AddReferencedObjects(FReferenceCollector& Collector)
+{
+}
+
+FString FTedsRevisionControlModule::GetReferencerName() const
+{
+	return TEXT("Teds: Revision Control Module");
+}
+
+IMPLEMENT_MODULE(FTedsRevisionControlModule, TedsRevisionControl)
diff --git a/Engine/Plugins/Experimental/EditorDataStorageFeatures/Source/TedsRevisionControl/Private/TedsRevisionControlModule.h b/Engine/Plugins/Experimental/EditorDataStorageFeatures/Source/TedsRevisionControl/Private/TedsRevisionControlModule.h
new file mode 100644
--- /dev/null
+++ b/Engine/Plugins/Experimental/EditorDataStorageFeatures/Source/TedsRevisionControl/Private/TedsRevisionControlModule.h
@@ -0,0 +1,29 @@
+// Copyright Epic Games, Inc. All Rights Reserved.
+
+#pragma once
+
+#include "Modules/ModuleManager.h"
+#include "UObject/GCObject.h"
+
+class FReferenceCollector;
+
+class FTedsRevisionControlModule : public IModuleInterface, public FGCObject
+{
+public:
+	~FTedsRevisionControlModule() override = default;
+
+	//
+	// IModuleInterface
+	//
+
+	virtual void StartupModule() override;
+	virtual void ShutdownModule() override;
+
+	
+	//
+	// FGCObject
+	//
+	
+	void AddReferencedObjects(FReferenceCollector& Collector) override;
+	FString GetReferencerName() const override;
+};
diff --git a/Engine/Plugins/Experimental/EditorDataStorageFeatures/Source/TedsRevisionControl/Private/TypedElementsDataStorageRevisionControl.cpp b/Engine/Plugins/Experimental/EditorDataStorageFeatures/Source/TedsRevisionControl/Private/TypedElementsDataStorageRevisionControl.cpp
deleted file mode 100644
diff --git a/Engine/Plugins/Experimental/EditorDataStorageFeatures/Source/TedsRevisionControl/Public/TypedElementsDataStorageRevisionControl.h b/Engine/Plugins/Experimental/EditorDataStorageFeatures/Source/TedsRevisionControl/Public/TypedElementsDataStorageRevisionControl.h
deleted file mode 100644
diff --git a/Engine/Plugins/Experimental/EditorDataStorageFeatures/Source/TedsRevisionControl/TedsRevisionControl.Build.cs b/Engine/Plugins/Experimental/EditorDataStorageFeatures/Source/TedsRevisionControl/TedsRevisionControl.Build.cs
--- a/Engine/Plugins/Experimental/EditorDataStorageFeatures/Source/TedsRevisionControl/TedsRevisionControl.Build.cs
+++ b/Engine/Plugins/Experimental/EditorDataStorageFeatures/Source/TedsRevisionControl/TedsRevisionControl.Build.cs
@@ -19,18 +19,11 @@ public class TedsRevisionControl : ModuleRules
 					"Core",
 					"CoreUObject",
 					"Engine",
-					"EditorFramework",
-					"Slate",
-					"SlateCore",
 					"SourceControl",
 					"TypedElementFramework",
-					"EditorSubsystem",
-					"UnrealEd",
 				});
 
 			DynamicallyLoadedModuleNames.AddRange(new string[] {});
-			
-			ShortName = "TElmntsDataStorRevCtrl";
 		}
 	}
 }
\ No newline at end of file
