Fix merge issue causing crash

diff --git a/Engine/Plugins/Experimental/EditorDataStorageFeatures/Source/TedsRevisionControl/Private/Processors/RevisionControlProcessors.cpp b/Engine/Plugins/Experimental/EditorDataStorageFeatures/Source/TedsRevisionControl/Private/Processors/RevisionControlProcessors.cpp
--- a/Engine/Plugins/Experimental/EditorDataStorageFeatures/Source/TedsRevisionControl/Private/Processors/RevisionControlProcessors.cpp
+++ b/Engine/Plugins/Experimental/EditorDataStorageFeatures/Source/TedsRevisionControl/Private/Processors/RevisionControlProcessors.cpp
@@ -165,6 +165,10 @@ void URevisionControlDataStorageFactory::RegisterTables(ITypedElementDataStorage
 			FTypedElementPackagePathColumn, FTypedElementPackageLoadedPathColumn,
 			FSCCRevisionIdColumn, FSCCExternalRevisionIdColumn>(),
 		FName("Editor_RevisionControlTable"));
+
+	DataStorage.RegisterTable(
+		TTypedElementColumnTypeList<FTypedElementPackageUpdateColumn>(),
+		FName("Editor_PackageUpdateTable"));
 }
 
 void URevisionControlDataStorageFactory::RegisterQueries(ITypedElementDataStorageInterface& DataStorage)
@@ -257,91 +261,118 @@ void URevisionControlDataStorageFactory::RegisterFetchUpdates(ITypedElementDataS
 
 void URevisionControlDataStorageFactory::RegisterApplyOverlays(ITypedElementDataStorageInterface& DataStorage)
 {
+	using namespace TypedElementDataStorage;
 	using namespace TypedElementQueryBuilder;
 	using DSI = ITypedElementDataStorageInterface;
 
-	if (ApplyOverlaysObjectToSCC == TypedElementInvalidQueryHandle)
+	if (ApplyOverlaysObjectToSCC == InvalidQueryHandle)
 	{
 		ApplyOverlaysObjectToSCC = DataStorage.RegisterQuery(
 			Select()
 				.ReadOnly<FTypedElementPackagePathColumn>()
-				.ReadOnly<FSCCStatusColumn, FSCCExternallyLockedColumn, FTypedElementViewportOverlayColorColumn>(EOptional::Yes)
+				.ReadOnly<FSCCStatusColumn>(EOptional::Yes)
 			.Compile());
 		}
 
-	if (ApplyNewOverlays == TypedElementInvalidQueryHandle)
+	if (ApplyNewOverlays == InvalidQueryHandle)
 	{
-		// Query:
-		// For all actors without an overlay color column AND having a package reference:
-		//   Determine if a color should be applied based on SCC status tags
-		//   If so, add the OverlayColorColumn to the actor row
 		ApplyNewOverlays = DataStorage.RegisterQuery(
-			Select(
-				TEXT("Poll to add selection overlay colors based on SCC status"),
-				// This is in PrePhysics because the overlay->actor query is in DuringPhysics and contexts don't flush changes between tick groups
-				FProcessor(DSI::EQueryTickPhase::PrePhysics, DataStorage.GetQueryTickGroupName(DSI::EQueryTickGroups::SyncExternalToDataStorage))
-					.ForceToGameThread(true),
-				[](DSI::IQueryContext& Context, TypedElementRowHandle ObjectRow, const FTypedElementUObjectColumn& Actor, const FTypedElementPackageReference& PackageReference)
-				{
-					Context.RunSubquery(0, PackageReference.Row, CreateSubqueryCallbackBinding(
-						[&Context, &ObjectRow, &Actor](DSI::ISubqueryContext& SubQueryContext)
-						{
-							FColor Color = DetermineOverlayColor(Context, SubQueryContext, Actor);
-							if (Color.Bits != 0)
-							{
-								Context.AddColumn<FTypedElementViewportOverlayColorColumn>(ObjectRow, { .OverlayColor = Color });
-							}
-						})
-					);
-				}
-			)
+			Select()
+			.ReadOnly<FTypedElementUObjectColumn, FTypedElementPackageReference>()
 			.Where()
 				.All<FTypedElementActorTag>()
 				.None<FTypedElementViewportOverlayColorColumn>()
-			.DependsOn()
-				.SubQuery(ApplyOverlaysObjectToSCC)
 			.Compile()
 		);
 	}
 
-	if (ChangeOverlay == TypedElementInvalidQueryHandle)
+	if (ChangeOverlay == InvalidQueryHandle)
 	{
-		// Query:
-		// For all actors WITH an overlay color column AND having a package reference:
-		//   Re-check the color that should be applied based on SCC status tags
-		//   If the color has changed, remove and re-add the OverlayColorColumn to the actor row
-		//
-		// Note: Remove and re-add will trigger observer in TypedElementActorViewportProcessors to SetOverlayColor on the primitive components
 		ChangeOverlay = DataStorage.RegisterQuery(
+			Select()
+			.ReadOnly<FTypedElementUObjectColumn, FTypedElementPackageReference, FTypedElementViewportOverlayColorColumn>()
+			.Where()
+				.All<FTypedElementActorTag>()
+			.Compile()
+		);
+	}
+
+	if (FlushPackageUpdates == InvalidQueryHandle)
+	{
+		check(ApplyOverlaysObjectToSCC != InvalidQueryHandle && ApplyNewOverlays != InvalidQueryHandle && ChangeOverlay!= InvalidQueryHandle);
+		enum EFlushPackageUpdatesSubqueries
+		{
+			EApplyOverlaysObjectToSCC,
+			EApplyNewOverlays,
+			EChangeOverlay,
+			
+			Num
+		};
+		TArray<RowHandle> Subqueries;
+		Subqueries.AddUninitialized(EFlushPackageUpdatesSubqueries::Num);
+		Subqueries[EApplyOverlaysObjectToSCC] = ApplyOverlaysObjectToSCC;
+		Subqueries[EApplyNewOverlays] = ApplyNewOverlays;
+		Subqueries[EChangeOverlay] = ChangeOverlay;
+		
+		FlushPackageUpdates = DataStorage.RegisterQuery(
 			Select(
-				TEXT("Poll for changed selection overlay colors based on SCC status"),
-				// This is in PrePhysics because the overlay->actor query is in DuringPhysics and contexts don't flush changes between tick groups
-				FProcessor(DSI::EQueryTickPhase::PrePhysics, DataStorage.GetQueryTickGroupName(DSI::EQueryTickGroups::SyncExternalToDataStorage))
-					.ForceToGameThread(true),
-				[](DSI::IQueryContext& Context, TypedElementRowHandle ObjectRow, const FTypedElementUObjectColumn& Actor, const FTypedElementPackageReference& PackageReference, const FTypedElementViewportOverlayColorColumn& OverlayColorColumn)
+				TEXT("Consume collected package updates"),
+				FProcessor(DSI::EQueryTickPhase::PrePhysics, DataStorage.GetQueryTickGroupName(DSI::EQueryTickGroups::Update))
+				.ForceToGameThread(true),
+				[](DSI::IQueryContext& Context, RowHandle Row, const FTypedElementPackageUpdateColumn& Update)
 				{
-					Context.RunSubquery(0, PackageReference.Row, CreateSubqueryCallbackBinding(
-						[&Context, &ObjectRow, &Actor, &OverlayColorColumn](DSI::ISubqueryContext& SubQueryContext)
+					// Query:
+					// For all actors without an overlay color column AND having a package reference:
+					//   Determine if a color should be applied based on SCC status tags
+					//   If so, add the OverlayColorColumn to the actor row
+					Context.RunSubquery(EApplyNewOverlays, Update.ObjectRow, CreateSubqueryCallbackBinding(
+						[&Context](RowHandle ObjectRow, const FTypedElementUObjectColumn& Actor, const FTypedElementPackageReference& PackageReference)
 						{
-							FColor Color = DetermineOverlayColor(Context, SubQueryContext, Actor);
-							if (Color.Bits == 0)
-							{
-								Context.RemoveColumns<FTypedElementViewportOverlayColorColumn>(ObjectRow);
-							}
-							else if (Color != OverlayColorColumn.OverlayColor)
-							{
-								// Remove and re-add to trigger the observer
-								Context.RemoveColumns<FTypedElementViewportOverlayColorColumn>(ObjectRow);
-								Context.AddColumn<FTypedElementViewportOverlayColorColumn>(ObjectRow, { .OverlayColor = Color });
-							}
-						})
-					);
+							Context.RunSubquery(EApplyOverlaysObjectToSCC, PackageReference.Row, CreateSubqueryCallbackBinding(
+								[&Context, &ObjectRow, &Actor](DSI::ISubqueryContext& SubQueryContext)
+								{
+									FColor Color = DetermineOverlayColor(Context, SubQueryContext, Actor);
+									if (Color.Bits != 0)
+									{
+										Context.AddColumn<FTypedElementViewportOverlayColorColumn>(ObjectRow, { .OverlayColor = Color });
+									}
+								})
+							);
+						}
+					));
+
+					// Query:
+					// For all actors WITH an overlay color column AND having a package reference:
+					//   Re-check the color that should be applied based on SCC status tags
+					//   If the color has changed, remove and re-add the OverlayColorColumn to the actor row
+					//
+					// Note: Remove and re-add will trigger observer in TypedElementActorViewportProcessors to SetOverlayColor on the primitive components
+					Context.RunSubquery(EChangeOverlay, Update.ObjectRow, CreateSubqueryCallbackBinding(
+						[&Context](TypedElementRowHandle ObjectRow, const FTypedElementUObjectColumn& Actor, const FTypedElementPackageReference& PackageReference, const FTypedElementViewportOverlayColorColumn& OverlayColorColumn)
+						{
+							Context.RunSubquery(EApplyOverlaysObjectToSCC, PackageReference.Row, CreateSubqueryCallbackBinding(
+								[&Context, &ObjectRow, &Actor, &OverlayColorColumn](DSI::ISubqueryContext& SubQueryContext)
+								{
+									FColor Color = DetermineOverlayColor(Context, SubQueryContext, Actor);
+									if (Color.Bits == 0)
+									{
+										Context.RemoveColumns<FTypedElementViewportOverlayColorColumn>(ObjectRow);
+									}
+									else if (Color != OverlayColorColumn.OverlayColor)
+									{
+										// Remove and re-add to trigger the observer
+										Context.RemoveColumns<FTypedElementViewportOverlayColorColumn>(ObjectRow);
+										Context.AddColumn<FTypedElementViewportOverlayColorColumn>(ObjectRow, { .OverlayColor = Color });
+									}
+								})
+							);
+						}
+					));
+					Context.RemoveRow(Row);
 				}
 			)
-			.Where()
-				.All<FTypedElementActorTag>()
 			.DependsOn()
-				.SubQuery(ApplyOverlaysObjectToSCC)
+				.SubQuery(Subqueries)
 			.Compile()
 		);
 	}
