STedsTableViewer:
Add a custom widget to show row handles and display it as a column in the query editor
Add a message when the table viewer is empty because of no columns or empty query results

diff --git a/Engine/Plugins/Experimental/EditorDataStorageFeatures/Source/TedsDebugger/Private/Widgets/QueryEditor/TedsQueryEditorResultsView.cpp b/Engine/Plugins/Experimental/EditorDataStorageFeatures/Source/TedsDebugger/Private/Widgets/QueryEditor/TedsQueryEditorResultsView.cpp
--- a/Engine/Plugins/Experimental/EditorDataStorageFeatures/Source/TedsDebugger/Private/Widgets/QueryEditor/TedsQueryEditorResultsView.cpp
+++ b/Engine/Plugins/Experimental/EditorDataStorageFeatures/Source/TedsDebugger/Private/Widgets/QueryEditor/TedsQueryEditorResultsView.cpp
@@ -2,15 +2,12 @@
 
 #include "TedsQueryEditorResultsView.h"
 
-#include "SceneOutlinerPublicTypes.h"
 #include "SWarningOrErrorBox.h"
 #include "TedsOutlinerModule.h"
-#include "TypedElementOutlinerColumnIntegration.h"
-#include "TypedElementOutlinerMode.h"
-#include "Components/VerticalBox.h"
-#include "Elements/Columns/TypedElementLabelColumns.h"
+#include "TedsTableViewerColumn.h"
 #include "Elements/Interfaces/TypedElementDataStorageInterface.h"
 #include "Elements/Framework/TypedElementQueryBuilder.h"
+#include "Elements/Framework/TypedElementRegistry.h"
 #include "Modules/ModuleManager.h"
 #include "QueryEditor/TedsQueryEditorModel.h"
 #include "QueryStack/FQueryStackNode_RowView.h"
@@ -32,6 +29,9 @@ void SResultsView::Construct(const FArguments& InArgs, FTedsQueryEditorModel& In
 	ModelChangedDelegateHandle = Model->GetModelChangedDelegate().AddRaw(this, &SResultsView::OnModelChanged);
 	// RowQueryHandle = TypedElementDataStorage::InvalidQueryHandle;
 
+	// Create a custom column for the table viewer to display row handles
+	CreateRowHandleColumn();
+
 	RowQueryStack = MakeShared<UE::EditorDataStorage::FQueryStackNode_RowView>(&TableViewerRows);
 
 	ChildSlot
@@ -69,6 +69,11 @@ void SResultsView::Construct(const FArguments& InArgs, FTedsQueryEditorModel& In
 			})
 		]
 	];
+
+	if(RowHandleColumn)
+	{
+		TableViewer->AddCustomColumn(RowHandleColumn.ToSharedRef());
+	}
 }
 
 void SResultsView::Tick(const FGeometry& AllottedGeometry, const double InCurrentTime, const float InDeltaTime)
@@ -95,6 +100,12 @@ void SResultsView::Tick(const FGeometry& AllottedGeometry, const double InCurren
 
 			// Update the columns in the table viewer using the selection types from the query description
 			TableViewer->SetColumns(TArray<TWeakObjectPtr<const UScriptStruct>>(TableViewerQueryDescription.SelectionTypes));
+
+			// Since setting the columns clears all columns, we are going to re-add the custom column back
+			if(RowHandleColumn)
+			{
+				TableViewer->AddCustomColumn(RowHandleColumn.ToSharedRef());
+			}
 			
 			if (TableViewerQueryHandle != TypedElementDataStorage::InvalidQueryHandle)
 			{
@@ -103,7 +114,7 @@ void SResultsView::Tick(const FGeometry& AllottedGeometry, const double InCurren
 			}
 
 			// Mass doesn't like empty queries, so we only set it if there are actual conditions
-			if(TableViewerQueryDescription.ConditionTypes.Num())
+			if(TableViewerQueryDescription.ConditionTypes.Num() || TableViewerQueryDescription.SelectionTypes.Num())
 			{
 				TableViewerQueryHandle = TedsInterface.RegisterQuery(MoveTemp(TableViewerQueryDescription));
 			}
@@ -112,10 +123,11 @@ void SResultsView::Tick(const FGeometry& AllottedGeometry, const double InCurren
 		bModelDirty = false;
 	}
 
+	TSet<TypedElementDataStorage::RowHandle> NewTableViewerRows_Set;
+	
 	// Every frame we re-run the query to update the rows the table viewer is showing
 	if(TableViewerQueryHandle != TypedElementDataStorage::InvalidQueryHandle)
 	{
-		TSet<TypedElementDataStorage::RowHandle> NewTableViewerRows_Set;
 		NewTableViewerRows_Set.Reserve(TableViewerRows_Set.Num());
 
 		TypedElementDataStorage::FQueryResult QueryResult = Model->GetTedsInterface().RunQuery(TableViewerQueryHandle,
@@ -123,16 +135,16 @@ void SResultsView::Tick(const FGeometry& AllottedGeometry, const double InCurren
 		{
 			NewTableViewerRows_Set.Append(Context.GetRowHandles());
 		}));
+	}
+	
+	// Check if the two sets are equal, i.e not changes and no need to update the table viewer
+	const bool bSetsEqual = (TableViewerRows_Set.Num() == NewTableViewerRows_Set.Num()) && TableViewerRows_Set.Includes(NewTableViewerRows_Set);
 
-		// Check if the two sets are equal, i.e not changes and no need to update the table viewer
-		const bool bSetsEqual = (TableViewerRows_Set.Num() == NewTableViewerRows_Set.Num()) && TableViewerRows_Set.Includes(NewTableViewerRows_Set);
-
-		if(!bSetsEqual)
-		{
-			Swap(TableViewerRows_Set, NewTableViewerRows_Set);
-			TableViewerRows = TableViewerRows_Set.Array();
-			RowQueryStack->MarkDirty();
-		}
+	if(!bSetsEqual)
+	{
+		Swap(TableViewerRows_Set, NewTableViewerRows_Set);
+		TableViewerRows = TableViewerRows_Set.Array();
+		RowQueryStack->MarkDirty();
 	}
 }
 
@@ -142,4 +154,21 @@ void SResultsView::OnModelChanged()
 	Invalidate(EInvalidateWidgetReason::Layout);
 }
 
-#undef LOCTEXT_NAMESPACE
\ No newline at end of file
+void SResultsView::CreateRowHandleColumn()
+{
+	auto AssignWidgetToColumn = [this](TUniquePtr<FTypedElementWidgetConstructor> Constructor, TConstArrayView<TWeakObjectPtr<const UScriptStruct>>)
+	{
+		TSharedPtr<FTypedElementWidgetConstructor> WidgetConstructor(Constructor.Release());
+		RowHandleColumn = MakeShared<FTedsTableViewerColumn>(TEXT("Row Handle"), WidgetConstructor);
+		return false;
+	};
+	
+	UTypedElementRegistry* Registry = UTypedElementRegistry::GetInstance();
+	checkf(Registry, TEXT("SResultsView created before UTypedElementRegistry is available."));
+	ITypedElementDataStorageUiInterface* StorageUi = Registry->GetMutableDataStorageUi();
+	checkf(StorageUi, TEXT("SResultsView created before data storage interfaces were initialized."))
+
+	StorageUi->CreateWidgetConstructors(TEXT("General.Cell.RowHandle"), TypedElementDataStorage::FMetaDataView(), AssignWidgetToColumn);
+}
+
+#undef LOCTEXT_NAMESPACE
diff --git a/Engine/Plugins/Experimental/EditorDataStorageFeatures/Source/TedsDebugger/Private/Widgets/QueryEditor/TedsQueryEditorResultsView.h b/Engine/Plugins/Experimental/EditorDataStorageFeatures/Source/TedsDebugger/Private/Widgets/QueryEditor/TedsQueryEditorResultsView.h
--- a/Engine/Plugins/Experimental/EditorDataStorageFeatures/Source/TedsDebugger/Private/Widgets/QueryEditor/TedsQueryEditorResultsView.h
+++ b/Engine/Plugins/Experimental/EditorDataStorageFeatures/Source/TedsDebugger/Private/Widgets/QueryEditor/TedsQueryEditorResultsView.h
@@ -7,11 +7,13 @@
 class ISceneOutliner;
 class SSceneOutliner;
 class SHorizontalBox;
+struct FTypedElementWidgetConstructor;
 
 namespace UE::EditorDataStorage
 {
 	class STedsTableViewer;
 	class FQueryStackNode_RowView;
+	class FTedsTableViewerColumn;
 }
 
 namespace UE::EditorDataStorage::Debug::QueryEditor
@@ -31,6 +33,7 @@ namespace UE::EditorDataStorage::Debug::QueryEditor
 	private:
 
 		void OnModelChanged();
+		void CreateRowHandleColumn();
 		
 		FTedsQueryEditorModel* Model = nullptr;
 		FDelegateHandle ModelChangedDelegateHandle;
@@ -45,6 +48,9 @@ namespace UE::EditorDataStorage::Debug::QueryEditor
 		TSet<TypedElementDataStorage::RowHandle> TableViewerRows_Set;
 		TSharedPtr<UE::EditorDataStorage::STedsTableViewer> TableViewer;
 		TSharedPtr<UE::EditorDataStorage::FQueryStackNode_RowView> RowQueryStack;
+
+		// Custom column for the table viewer to display row handles
+		TSharedPtr<UE::EditorDataStorage::FTedsTableViewerColumn> RowHandleColumn;
 	};
 
 	
diff --git a/Engine/Plugins/Experimental/EditorDataStorageFeatures/Source/TedsOutliner/Private/Widgets/TedsOutlinerRowHandleColumn.cpp b/Engine/Plugins/Experimental/EditorDataStorageFeatures/Source/TedsOutliner/Private/Widgets/TedsOutlinerRowHandleColumn.cpp
--- a/Engine/Plugins/Experimental/EditorDataStorageFeatures/Source/TedsOutliner/Private/Widgets/TedsOutlinerRowHandleColumn.cpp
+++ b/Engine/Plugins/Experimental/EditorDataStorageFeatures/Source/TedsOutliner/Private/Widgets/TedsOutlinerRowHandleColumn.cpp
@@ -4,9 +4,30 @@
 
 #include "SortHelper.h"
 #include "TypedElementOutlinerItem.h"
+#include "Elements/Framework/TypedElementRegistry.h"
+#include "TedsTableViewerColumn.h"
 
 #define LOCTEXT_NAMESPACE "TedsOutlinerRowHandleColumn"
 
+FTedsOutlinerRowHandleColumn::FTedsOutlinerRowHandleColumn(ISceneOutliner& SceneOutliner)
+	: WeakSceneOutliner(StaticCastSharedRef<ISceneOutliner>(SceneOutliner.AsShared()))
+{
+	auto AssignWidgetToColumn = [this](TUniquePtr<FTypedElementWidgetConstructor> Constructor, TConstArrayView<TWeakObjectPtr<const UScriptStruct>>)
+	{
+		TSharedPtr<FTypedElementWidgetConstructor> WidgetConstructor(Constructor.Release());
+		TableViewerColumn = MakeShared<UE::EditorDataStorage::FTedsTableViewerColumn>(TEXT("Row Handle"), WidgetConstructor);
+		return false;
+	};
+	
+	UTypedElementRegistry* Registry = UTypedElementRegistry::GetInstance();
+	checkf(Registry, TEXT("FTedsOutlinerRowHandleColumn created before UTypedElementRegistry is available."));
+	ITypedElementDataStorageUiInterface* StorageUi = Registry->GetMutableDataStorageUi();
+	checkf(StorageUi, TEXT("FTedsOutlinerRowHandleColumn created before data storage interfaces were initialized."))
+
+	StorageUi->CreateWidgetConstructors(TEXT("General.Cell.RowHandle"), TypedElementDataStorage::FMetaDataView(), AssignWidgetToColumn);
+}
+
+
 FName FTedsOutlinerRowHandleColumn::GetID()
 {
 	static const FName ID("Row Handle");
@@ -34,21 +55,10 @@ const TSharedRef<SWidget> FTedsOutlinerRowHandleColumn::ConstructRowWidget(FScen
 	{
 		const TypedElementDataStorage::RowHandle RowHandle = OutlinerTreeItem->GetRowHandle();
 
-		FNumberFormattingOptions NumberFormattingOptions;
-		NumberFormattingOptions.SetUseGrouping(false);
-		const FText Text = FText::AsNumber(RowHandle, &NumberFormattingOptions);
-		
-		return SNew(SBox)
-			.HAlign(HAlign_Left)
-			.VAlign(VAlign_Center)
-			.Padding(8, 0, 0, 0)
-			.Content()
-			[
-				SNew(STextBlock)
-					.Text(Text)
-					.HighlightText(SceneOutliner->GetFilterHighlightText())
-					.ColorAndOpacity(FSlateColor::UseForeground())
-			];
+		if(TSharedPtr<SWidget> Widget = TableViewerColumn->ConstructRowWidget(RowHandle))
+		{
+			return Widget.ToSharedRef();
+		}
 	}
 	return SNullWidget::NullWidget;
 }
diff --git a/Engine/Plugins/Experimental/EditorDataStorageFeatures/Source/TedsOutliner/Private/Widgets/TedsOutlinerRowHandleColumn.h b/Engine/Plugins/Experimental/EditorDataStorageFeatures/Source/TedsOutliner/Private/Widgets/TedsOutlinerRowHandleColumn.h
--- a/Engine/Plugins/Experimental/EditorDataStorageFeatures/Source/TedsOutliner/Private/Widgets/TedsOutlinerRowHandleColumn.h
+++ b/Engine/Plugins/Experimental/EditorDataStorageFeatures/Source/TedsOutliner/Private/Widgets/TedsOutlinerRowHandleColumn.h
@@ -6,6 +6,11 @@
 #include "ISceneOutliner.h"
 #include "ISceneOutlinerColumn.h"
 
+namespace UE::EditorDataStorage
+{
+	class FTedsTableViewerColumn;
+}
+
 /**
  * This is a custom column for the TEDS Outliner/Table Viewer to display row handles for items. It is a special case-d column instead of going through
  * the TEDS UI layer because row handles are not information stored in a TEDS column
@@ -13,10 +18,7 @@
 class FTedsOutlinerRowHandleColumn : public ISceneOutlinerColumn
 {
 public:
-	FTedsOutlinerRowHandleColumn(ISceneOutliner& SceneOutliner) : WeakSceneOutliner(StaticCastSharedRef<ISceneOutliner>(SceneOutliner.AsShared()))
-	{
-		
-	}
+	FTedsOutlinerRowHandleColumn(ISceneOutliner& SceneOutliner);
 	
 	virtual ~FTedsOutlinerRowHandleColumn() {}
 
@@ -35,4 +37,7 @@ public:
 
 private:
 	TWeakPtr<ISceneOutliner> WeakSceneOutliner;
+
+	// The Table Viewer column we are going to use internally to create the widget
+	TSharedPtr<UE::EditorDataStorage::FTedsTableViewerColumn> TableViewerColumn;
 };
diff --git a/Engine/Plugins/Experimental/EditorDataStorageFeatures/Source/TedsTableViewer/Private/TedsTableViewerColumn.cpp b/Engine/Plugins/Experimental/EditorDataStorageFeatures/Source/TedsTableViewer/Private/TedsTableViewerColumn.cpp
--- a/Engine/Plugins/Experimental/EditorDataStorageFeatures/Source/TedsTableViewer/Private/TedsTableViewerColumn.cpp
+++ b/Engine/Plugins/Experimental/EditorDataStorageFeatures/Source/TedsTableViewer/Private/TedsTableViewerColumn.cpp
@@ -114,6 +114,7 @@ namespace UE::EditorDataStorage
 			.FillWidth(1)
 			.HeaderComboVisibility(EHeaderComboVisibility::OnHover)
 			.DefaultTooltip(FText::FromString(MoveTemp(TooltipText)))
+			.DefaultLabel(FText::FromName(ColumnName))
 			.HeaderContent()
 			[
 				SNew(SBox)
diff --git a/Engine/Plugins/Experimental/EditorDataStorageFeatures/Source/TedsTableViewer/Private/TedsTableViewerModel.cpp b/Engine/Plugins/Experimental/EditorDataStorageFeatures/Source/TedsTableViewer/Private/TedsTableViewerModel.cpp
--- a/Engine/Plugins/Experimental/EditorDataStorageFeatures/Source/TedsTableViewer/Private/TedsTableViewerModel.cpp
+++ b/Engine/Plugins/Experimental/EditorDataStorageFeatures/Source/TedsTableViewer/Private/TedsTableViewerModel.cpp
@@ -61,9 +61,9 @@ namespace UE::EditorDataStorage
 		}
 
 		// Tick all the individual column views
-		for(const auto& ColumnPair : ColumnsView)
+		for(const TSharedRef<FTedsTableViewerColumn>&Column : ColumnsView)
 		{
-			ColumnPair.Value->Tick();
+			Column->Tick();
 		}
 		
 		return true;
@@ -74,9 +74,25 @@ namespace UE::EditorDataStorage
 		return Items;
 	}
 
+	uint64 FTedsTableViewerModel::GetRowCount() const
+	{
+		return Items.Num();
+	}
+
+	uint64 FTedsTableViewerModel::GetColumnCount() const
+	{
+		return ColumnsView.Num();
+	}
+
 	TSharedPtr<FTedsTableViewerColumn> FTedsTableViewerModel::GetColumn(const FName& ColumnName) const
 	{
-		if(const TSharedRef<FTedsTableViewerColumn>* Column = ColumnsView.Find(ColumnName))
+		const TSharedRef<FTedsTableViewerColumn>* Column = ColumnsView.FindByPredicate([ColumnName]
+			(const TSharedRef<FTedsTableViewerColumn>& InColumn)
+		{
+			return InColumn->GetColumnName() == ColumnName;
+		});
+		
+		if(Column)
 		{
 			return *Column;
 		}
@@ -86,9 +102,9 @@ namespace UE::EditorDataStorage
 
 	void FTedsTableViewerModel::ForEachColumn(const TFunctionRef<void(const TSharedRef<FTedsTableViewerColumn>&)>& Delegate) const
 	{
-		for(const TPair<FName, TSharedRef<FTedsTableViewerColumn>>& ColumnPair : ColumnsView)
+		for(const TSharedRef<FTedsTableViewerColumn>& Column : ColumnsView)
 		{
-			Delegate(ColumnPair.Value);
+			Delegate(Column);
 		}
 	}
 
@@ -103,6 +119,13 @@ namespace UE::EditorDataStorage
 		GenerateColumns();
 	}
 
+	void FTedsTableViewerModel::AddCustomColumn(const TSharedRef<FTedsTableViewerColumn>& InColumn)
+	{
+		// Table Viewer TODO: We should allow users to specify sort order using a TEDS column on the UI row, but for now we put any custom
+		// columns on the front
+		ColumnsView.Insert(InColumn, 0);
+	}
+
 	void FTedsTableViewerModel::GenerateColumns()
 	{
 		using MatchApproach = ITypedElementDataStorageUiInterface::EMatchApproach;
@@ -195,13 +218,13 @@ namespace UE::EditorDataStorage
 		// Add the actual UI columns in the order the Teds Columns were specified
 		for (const TWeakObjectPtr<const UScriptStruct>& ColumnType : RequestedTedsColumns)
 		{
-			if(TSharedRef<FTedsTableViewerColumn>* FoundColumn = NewColumnMap.Find(ColumnType))
+			if(const TSharedRef<FTedsTableViewerColumn>* FoundColumn = NewColumnMap.Find(ColumnType))
 			{
 				// If the column already exists, a widget matched it and a previously encountered column together and was already added
 				// so we can safely ignore it here
 				if(!GetColumn((*FoundColumn)->GetColumnName()))
 				{
-					ColumnsView.Emplace((*FoundColumn)->GetColumnName(), *FoundColumn);
+					ColumnsView.Add(*FoundColumn);
 				}
 			}
 		}
diff --git a/Engine/Plugins/Experimental/EditorDataStorageFeatures/Source/TedsTableViewer/Private/Widgets/RowHandleWidget.cpp b/Engine/Plugins/Experimental/EditorDataStorageFeatures/Source/TedsTableViewer/Private/Widgets/RowHandleWidget.cpp
new file mode 100644
--- /dev/null
+++ b/Engine/Plugins/Experimental/EditorDataStorageFeatures/Source/TedsTableViewer/Private/Widgets/RowHandleWidget.cpp
@@ -0,0 +1,65 @@
+// Copyright Epic Games, Inc. All Rights Reserved.
+
+#include "Widgets/RowHandleWidget.h"
+
+#include "Elements/Columns/TypedElementMiscColumns.h"
+#include "Widgets/Layout/SBox.h"
+#include "Widgets/Text/STextBlock.h"
+
+#define LOCTEXT_NAMESPACE "RowHandleWidget"
+
+
+void URowHandleWidgetFactory::RegisterWidgetConstructors(ITypedElementDataStorageInterface& DataStorage, ITypedElementDataStorageUiInterface& DataStorageUi) const
+{
+	DataStorageUi.RegisterWidgetFactory(FName(TEXT("General.Cell.RowHandle")), FRowHandleWidgetConstructor::StaticStruct());
+}
+
+void URowHandleWidgetFactory::RegisterWidgetPurposes(ITypedElementDataStorageUiInterface& DataStorageUi) const
+{
+	DataStorageUi.RegisterWidgetPurpose(FName(TEXT("General.Cell.RowHandle")), ITypedElementDataStorageUiInterface::EPurposeType::UniqueByName,
+	LOCTEXT("RowHandlePurpose", "Specific purpose to request a widget to display row handles."));
+}
+
+FRowHandleWidgetConstructor::FRowHandleWidgetConstructor()
+	: Super(FRowHandleWidgetConstructor::StaticStruct())
+{
+	
+}
+
+TSharedPtr<SWidget> FRowHandleWidgetConstructor::CreateWidget(const TypedElementDataStorage::FMetaDataView& Arguments)
+{
+	return SNew(SBox)
+		.HAlign(HAlign_Left)
+		.VAlign(VAlign_Center)
+		.Padding(8, 0, 0, 0);}
+
+bool FRowHandleWidgetConstructor::FinalizeWidget(ITypedElementDataStorageInterface* DataStorage, ITypedElementDataStorageUiInterface* DataStorageUi, TypedElementDataStorage::RowHandle Row, const TSharedPtr<SWidget>& Widget)
+{
+
+	checkf(Widget->GetType() == SBox::StaticWidgetClass().GetWidgetType(),
+		TEXT("Stored widget with FRowHandleWidgetConstructor doesn't match type %s, but was a %s."),
+		*(SBox::StaticWidgetClass().GetWidgetType().ToString()),
+		*(Widget->GetTypeAsString()));
+	
+	SBox* BoxWidget = static_cast<SBox*>(Widget.Get());
+
+	TypedElementDataStorage::RowHandle TargetRowHandle = TypedElementDataStorage::InvalidRowHandle;
+
+	if(const FTypedElementRowReferenceColumn* RowReferenceColumn = DataStorage->GetColumn<FTypedElementRowReferenceColumn>(Row))
+	{
+		TargetRowHandle = RowReferenceColumn->Row;
+	}
+	
+	FNumberFormattingOptions NumberFormattingOptions;
+	NumberFormattingOptions.SetUseGrouping(false);
+	const FText Text = FText::AsNumber(TargetRowHandle, &NumberFormattingOptions);
+
+	BoxWidget->SetContent(
+			SNew(STextBlock)
+				.Text(Text)
+				.ColorAndOpacity(FSlateColor::UseForeground())
+		);
+	return true;
+}
+
+#undef LOCTEXT_NAMESPACE //"RowHandleWidget"
diff --git a/Engine/Plugins/Experimental/EditorDataStorageFeatures/Source/TedsTableViewer/Private/Widgets/STedsTableViewer.cpp b/Engine/Plugins/Experimental/EditorDataStorageFeatures/Source/TedsTableViewer/Private/Widgets/STedsTableViewer.cpp
--- a/Engine/Plugins/Experimental/EditorDataStorageFeatures/Source/TedsTableViewer/Private/Widgets/STedsTableViewer.cpp
+++ b/Engine/Plugins/Experimental/EditorDataStorageFeatures/Source/TedsTableViewer/Private/Widgets/STedsTableViewer.cpp
@@ -5,6 +5,9 @@
 #include "TedsTableViewerColumn.h"
 #include "Widgets/STedsTableViewerRow.h"
 #include "Widgets/Views/SListView.h"
+#include "Widgets/Text/STextBlock.h"
+
+#define LOCTEXT_NAMESPACE "STedsTableViewer"
 
 namespace UE::EditorDataStorage
 {
@@ -16,14 +19,13 @@ namespace UE::EditorDataStorage
 		HeaderRowWidget = SNew( SHeaderRow )
 							.CanSelectGeneratedColumn(true);
 
-		ChildSlot
-		[
-			SAssignNew(ListView, SListView<TableViewerItemPtr>)
+		ListView = SNew(SListView<TableViewerItemPtr>)
 			.HeaderRow(HeaderRowWidget)
 			.ListItemsSource(&Model->GetItems())
-			.OnGenerateRow(this, &STedsTableViewer::MakeTableRowWidget)
-		];
-
+			.OnGenerateRow(this, &STedsTableViewer::MakeTableRowWidget);
+		
+		AssignChildSlot();
+		
 		// Add each Teds column from the model to our header row widget
 		Model->ForEachColumn([this](const TSharedRef<FTedsTableViewerColumn>& Column)
 		{
@@ -34,18 +36,68 @@ namespace UE::EditorDataStorage
 		Model->GetOnModelChanged().AddLambda([this]()
 		{
 			ListView->RequestListRefresh();
+			AssignChildSlot();
 		});
 	}
 
-	void STedsTableViewer::SetColumns(const TArray<TWeakObjectPtr<const UScriptStruct>>& Columns) const
+	void STedsTableViewer::AssignChildSlot()
+	{
+		if(Model->GetRowCount() == 0)
+		{
+			ChildSlot
+			[
+				SNew(SBox)
+					.HAlign(HAlign_Center)
+					.VAlign(VAlign_Center)
+					[
+						SNew(STextBlock)
+							.Text(LOCTEXT("EmptyTableViewerText", "The input query has no results"))
+					]
+			];
+		}
+		else if(Model->GetColumnCount() == 0)
+		{
+			ChildSlot
+			[
+				SNew(SBox)
+					.HAlign(HAlign_Center)
+					.VAlign(VAlign_Center)
+					[
+						SNew(STextBlock)
+							.Text(LOCTEXT("EmptyTableViewerText", "There were no columns specified to display"))
+					]
+			];
+		}
+		else
+		{
+			ChildSlot
+			[
+				ListView.ToSharedRef()	
+			];
+		}
+	}
+
+	void STedsTableViewer::RefreshColumnWidgets()
 	{
-		Model->SetColumns(Columns);
-		
 		HeaderRowWidget->ClearColumns();
 		Model->ForEachColumn([this](const TSharedRef<FTedsTableViewerColumn>& Column)
 		{
 			HeaderRowWidget->AddColumn(Column->ConstructHeaderRowColumn());
 		});
+
+		AssignChildSlot();
+	}
+
+	void STedsTableViewer::SetColumns(const TArray<TWeakObjectPtr<const UScriptStruct>>& Columns)
+	{
+		Model->SetColumns(Columns);
+		RefreshColumnWidgets();
+	}
+
+	void STedsTableViewer::AddCustomColumn(const TSharedRef<FTedsTableViewerColumn>& InColumn)
+	{
+		Model->AddCustomColumn(InColumn);
+		RefreshColumnWidgets();
 	}
 
 	bool STedsTableViewer::IsItemVisible(TableViewerItemPtr InItem) const
@@ -58,7 +110,7 @@ namespace UE::EditorDataStorage
 		return SNew(STedsTableViewerRow, OwnerTable, Model.ToSharedRef())
 				.Item(InItem);
 	}
-
-	
 }
 
+#undef LOCTEXT_NAMESPACE //"STedsTableViewer"
+
diff --git a/Engine/Plugins/Experimental/EditorDataStorageFeatures/Source/TedsTableViewer/Public/TedsTableViewerModel.h b/Engine/Plugins/Experimental/EditorDataStorageFeatures/Source/TedsTableViewer/Public/TedsTableViewerModel.h
--- a/Engine/Plugins/Experimental/EditorDataStorageFeatures/Source/TedsTableViewer/Public/TedsTableViewerModel.h
+++ b/Engine/Plugins/Experimental/EditorDataStorageFeatures/Source/TedsTableViewer/Public/TedsTableViewerModel.h
@@ -48,6 +48,12 @@ namespace UE::EditorDataStorage
 
 		// Get the items this table viewer is viewing
 		TEDSTABLEVIEWER_API const TArray<TableViewerItemPtr>& GetItems() const;
+
+		// Get the number of rows currently being observed
+		TEDSTABLEVIEWER_API uint64 GetRowCount() const;
+		
+		// Get the number of columns being displayed
+		TEDSTABLEVIEWER_API uint64 GetColumnCount() const;
 		
 		// Get a specific column that the table viewer is displaying by name
 		TEDSTABLEVIEWER_API TSharedPtr<FTedsTableViewerColumn> GetColumn(const FName& ColumnName) const;
@@ -61,6 +67,9 @@ namespace UE::EditorDataStorage
 		// Clear the current list of columns being displayed and set it to the given list
 		TEDSTABLEVIEWER_API void SetColumns(const TArray<TWeakObjectPtr<const UScriptStruct>>& InColumns);
 
+		// Add a custom column to display in the table viewer, that doesn't necessarily map to a Teds column
+		TEDSTABLEVIEWER_API void AddCustomColumn(const TSharedRef<FTedsTableViewerColumn>& InColumn);
+
 	protected:
 
 		// Generate the actual columns to display in the UI using TEDS UI
@@ -82,7 +91,7 @@ namespace UE::EditorDataStorage
 		TArray<TableViewerItemPtr> Items;
 
 		// List of columns the table viewer is currently displaying
-		TMap<FName, TSharedRef<FTedsTableViewerColumn>> ColumnsView;
+		TArray<TSharedRef<FTedsTableViewerColumn>> ColumnsView;
 
 		// The initial TEDS columns the widget was requested to display
 		TArray<TWeakObjectPtr<const UScriptStruct>> RequestedTedsColumns;
diff --git a/Engine/Plugins/Experimental/EditorDataStorageFeatures/Source/TedsTableViewer/Public/Widgets/RowHandleWidget.h b/Engine/Plugins/Experimental/EditorDataStorageFeatures/Source/TedsTableViewer/Public/Widgets/RowHandleWidget.h
new file mode 100644
--- /dev/null
+++ b/Engine/Plugins/Experimental/EditorDataStorageFeatures/Source/TedsTableViewer/Public/Widgets/RowHandleWidget.h
@@ -0,0 +1,45 @@
+// Copyright Epic Games, Inc. All Rights Reserved.
+
+#pragma once
+
+#include "Elements/Interfaces/TypedElementDataStorageFactory.h"
+#include "Elements/Interfaces/TypedElementDataStorageUiInterface.h"
+#include "Elements/Common/TypedElementQueryConditions.h"
+
+#include "RowHandleWidget.generated.h"
+
+class ITypedElementDataStorageInterface;
+class UScriptStruct;
+
+UCLASS()
+class TEDSTABLEVIEWER_API URowHandleWidgetFactory : public UTypedElementDataStorageFactory
+{
+	GENERATED_BODY()
+
+public:
+	virtual ~URowHandleWidgetFactory() override = default;
+
+	virtual void RegisterWidgetConstructors(ITypedElementDataStorageInterface& DataStorage,
+		ITypedElementDataStorageUiInterface& DataStorageUi) const override;
+
+	void RegisterWidgetPurposes(ITypedElementDataStorageUiInterface& DataStorageUi) const override;
+};
+
+// A custom widget to display the row handle of a row as text
+USTRUCT()
+struct TEDSTABLEVIEWER_API FRowHandleWidgetConstructor : public FTypedElementWidgetConstructor
+{
+	GENERATED_BODY()
+
+public:
+	FRowHandleWidgetConstructor();
+	virtual ~FRowHandleWidgetConstructor() override = default;
+
+protected:
+	virtual TSharedPtr<SWidget> CreateWidget(const TypedElementDataStorage::FMetaDataView& Arguments) override;
+	virtual bool FinalizeWidget(
+		ITypedElementDataStorageInterface* DataStorage,
+		ITypedElementDataStorageUiInterface* DataStorageUi,
+		TypedElementRowHandle Row,
+		const TSharedPtr<SWidget>& Widget) override;
+};
diff --git a/Engine/Plugins/Experimental/EditorDataStorageFeatures/Source/TedsTableViewer/Public/Widgets/STedsTableViewer.h b/Engine/Plugins/Experimental/EditorDataStorageFeatures/Source/TedsTableViewer/Public/Widgets/STedsTableViewer.h
--- a/Engine/Plugins/Experimental/EditorDataStorageFeatures/Source/TedsTableViewer/Public/Widgets/STedsTableViewer.h
+++ b/Engine/Plugins/Experimental/EditorDataStorageFeatures/Source/TedsTableViewer/Public/Widgets/STedsTableViewer.h
@@ -12,7 +12,6 @@ class ITableRow;
 class STableViewBase;
 class SHeaderRow;
 
-
 namespace UE::EditorDataStorage
 {
 	class FTedsTableViewerColumn;
@@ -52,7 +51,11 @@ namespace UE::EditorDataStorage
 		TEDSTABLEVIEWER_API void Construct(const FArguments& InArgs);
 
 		// Clear the current list of columns being displayed and set it to the given list
-		TEDSTABLEVIEWER_API void SetColumns(const TArray<TWeakObjectPtr<const UScriptStruct>>& Columns) const;
+		TEDSTABLEVIEWER_API void SetColumns(const TArray<TWeakObjectPtr<const UScriptStruct>>& Columns);
+
+		// Add a custom column to display in the table viewer, that doesn't necessarily map to a Teds column
+		TEDSTABLEVIEWER_API void AddCustomColumn(const TSharedRef<FTedsTableViewerColumn>& InColumn);
+
 
 	protected:
 		
@@ -60,6 +63,10 @@ namespace UE::EditorDataStorage
 
 		bool IsItemVisible(TableViewerItemPtr InItem) const;
 
+		void AssignChildSlot();
+
+		void RefreshColumnWidgets();
+
 	private:
 
 		// The actual ListView widget that displays the rows
