[Teds] Rename files and update class names for SceneOutliner derived types for the TedsOutliner module

#rnx

diff --git a/Engine/Plugins/Experimental/Animation/HierarchyTable/HierarchyTable.uplugin b/Engine/Plugins/Experimental/Animation/HierarchyTable/HierarchyTable.uplugin
--- a/Engine/Plugins/Experimental/Animation/HierarchyTable/HierarchyTable.uplugin
+++ b/Engine/Plugins/Experimental/Animation/HierarchyTable/HierarchyTable.uplugin
@@ -29,7 +29,7 @@
 	],
 	"Plugins": [
 		{
-			"Name": "TypedElementDataStorageEditorUI",
+			"Name": "EditorDataStorageFeatures",
 			"Enabled": true
 		}
 	]
diff --git a/Engine/Plugins/Experimental/Animation/HierarchyTable/Source/Editor/Private/HierarchyTableEditorToolkit.cpp b/Engine/Plugins/Experimental/Animation/HierarchyTable/Source/Editor/Private/HierarchyTableEditorToolkit.cpp
--- a/Engine/Plugins/Experimental/Animation/HierarchyTable/Source/Editor/Private/HierarchyTableEditorToolkit.cpp
+++ b/Engine/Plugins/Experimental/Animation/HierarchyTable/Source/Editor/Private/HierarchyTableEditorToolkit.cpp
@@ -11,7 +11,7 @@
 #include "Elements/Columns/TypedElementSelectionColumns.h"
 #include "Elements/Columns/TypedElementMiscColumns.h"
 #include "SceneOutlinerPublicTypes.h"
-#include "TypedElementOutlinerMode.h"
+#include "TedsOutlinerMode.h"
 #include "ReferenceSkeleton.h"
 #include "Animation/Skeleton.h"
 #include "HierarchyTable/Columns/OverrideColumn.h"
@@ -139,7 +139,7 @@ TSharedRef<SWidget> FHierarchyTableEditorToolkit::CreateTedsOutliner()
 	InitOptions.FilterBarOptions.bHasFilterBar = true;
 	InitOptions.OutlinerIdentifier = "HierarchyTableTedsOutliner";
 
-	FTypedElementOutlinerModeParams Params(nullptr);
+	FTedsOutlinerParams Params(nullptr);
 	{
 		TypedElementDataStorage::FQueryDescription RowQueryDescription =
 			Select()
@@ -149,7 +149,7 @@ TSharedRef<SWidget> FHierarchyTableEditorToolkit::CreateTedsOutliner()
 
 		Params.QueryDescription = RowQueryDescription;
 		Params.CellWidgetPurposes = TArray<FName>{ TEXT("General.Cell") };
-		Params.HierarchyData = FTypedElementOutlinerHierarchyData::GetDefaultHierarchyData();
+		Params.HierarchyData = FTedsOutlinerHierarchyData::GetDefaultHierarchyData();
 	}
 
 	FTedsOutlinerModule& TedsOutlinerModule = FModuleManager::GetModuleChecked<FTedsOutlinerModule>("TedsOutliner");
@@ -175,7 +175,7 @@ TSharedRef<SWidget> FHierarchyTableEditorToolkit::CreateTedsOutliner()
 		TypedElementRowHandle* ParentRow = EntryIndexToHandleMap.Find(Entry->Parent);
 		if (ParentRow)
 		{
-			DSI->AddColumn<FTypedElementParentColumn>(Row, { .Parent = *ParentRow });
+			DSI->AddColumn<FTableRowParentColumn>(Row, { .Parent = *ParentRow });
 		}
 
 		for (const UScriptStruct* Column : BaseHierarchyTableTypeColumns)
diff --git a/Engine/Plugins/Experimental/EditorDataStorage/Source/TedsCore/Private/Compatibility/TypedElementActorParentQueries.cpp b/Engine/Plugins/Experimental/EditorDataStorage/Source/TedsCore/Private/Compatibility/TypedElementActorParentQueries.cpp
--- a/Engine/Plugins/Experimental/EditorDataStorage/Source/TedsCore/Private/Compatibility/TypedElementActorParentQueries.cpp
+++ b/Engine/Plugins/Experimental/EditorDataStorage/Source/TedsCore/Private/Compatibility/TypedElementActorParentQueries.cpp
@@ -49,11 +49,11 @@ void UTypedElementActorParentFactory::RegisterAddParentColumn(ITypedElementDataS
 						RowHandle ParentRow = Context.FindIndexedRow(IdHash);
 						if (Context.IsRowAvailable(ParentRow))
 						{
-							Context.AddColumn(Row, FTypedElementParentColumn{ .Parent = ParentRow });
+							Context.AddColumn(Row, FTableRowParentColumn{ .Parent = ParentRow });
 						}
 						else
 						{
-							Context.AddColumn(Row, FTypedElementUnresolvedParentColumn{ .ParentIdHash = IdHash });
+							Context.AddColumn(Row, FUnresolvedTableRowParentColumn{ .ParentIdHash = IdHash });
 						}
 					}
 				}
@@ -61,7 +61,7 @@ void UTypedElementActorParentFactory::RegisterAddParentColumn(ITypedElementDataS
 		)
 		.Where()
 			.All<FTypedElementSyncFromWorldTag, FTypedElementActorTag>()
-			.None<FTypedElementParentColumn, FTypedElementUnresolvedParentColumn>()
+			.None<FTableRowParentColumn, FUnresolvedTableRowParentColumn>()
 		.Compile()
 	);
 }
@@ -76,7 +76,7 @@ void UTypedElementActorParentFactory::RegisterUpdateOrRemoveParentColumn(ITypedE
 			TEXT("Sync actor's parent to column"),
 			FProcessor(EQueryTickPhase::PrePhysics, DataStorage.GetQueryTickGroupName(EQueryTickGroups::SyncExternalToDataStorage))
 				.ForceToGameThread(true),
-			[](IQueryContext& Context, TypedElementRowHandle Row, const FTypedElementUObjectColumn& Actor, FTypedElementParentColumn& Parent)
+			[](IQueryContext& Context, TypedElementRowHandle Row, const FTypedElementUObjectColumn& Actor, FTableRowParentColumn& Parent)
 			{
 				if (const AActor* ActorInstance = Cast<AActor>(Actor.Object))
 				{
@@ -93,14 +93,14 @@ void UTypedElementActorParentFactory::RegisterUpdateOrRemoveParentColumn(ITypedE
 							}
 							else
 							{
-								Context.RemoveColumns<FTypedElementParentColumn>(Row);
-								Context.AddColumn(Row, FTypedElementUnresolvedParentColumn{ .ParentIdHash = IdHash });
+								Context.RemoveColumns<FTableRowParentColumn>(Row);
+								Context.AddColumn(Row, FUnresolvedTableRowParentColumn{ .ParentIdHash = IdHash });
 							}
 						}
 						return;
 					}
 				}
-				Context.RemoveColumns<FTypedElementParentColumn>(Row);
+				Context.RemoveColumns<FTableRowParentColumn>(Row);
 			}
 		)
 		.Where()
diff --git a/Engine/Plugins/Experimental/EditorDataStorage/Source/TedsCore/Private/Queries/TypedElementAlertQueries.cpp b/Engine/Plugins/Experimental/EditorDataStorage/Source/TedsCore/Private/Queries/TypedElementAlertQueries.cpp
--- a/Engine/Plugins/Experimental/EditorDataStorage/Source/TedsCore/Private/Queries/TypedElementAlertQueries.cpp
+++ b/Engine/Plugins/Experimental/EditorDataStorage/Source/TedsCore/Private/Queries/TypedElementAlertQueries.cpp
@@ -168,7 +168,7 @@ void UTypedElementAlertQueriesFactory::RegisterSubQueries(ITypedElementDataStora
 
 	ParentReadOnlyQuery = DataStorage.RegisterQuery(
 		Select()
-			.ReadOnly<FTypedElementParentColumn>()
+			.ReadOnly<FTableRowParentColumn>()
 		.Compile());
 }
 
@@ -181,7 +181,7 @@ void UTypedElementAlertQueriesFactory::RegisterParentUpdatesQueries(ITypedElemen
 		Select(
 			TEXT("Trigger alert update if alert's parent changed"),
 			FProcessor(EQueryTickPhase::FrameEnd, DataStorage.GetQueryTickGroupName(EQueryTickGroups::Default)),
-			[](IQueryContext& Context, FTypedElementAlertColumn& Alert, const FTypedElementParentColumn& Parent)
+			[](IQueryContext& Context, FTypedElementAlertColumn& Alert, const FTableRowParentColumn& Parent)
 			{
 				if (Alert.CachedParent != Parent.Parent)
 				{
@@ -197,7 +197,7 @@ void UTypedElementAlertQueriesFactory::RegisterParentUpdatesQueries(ITypedElemen
 		Select(
 			TEXT("Trigger alert update if child alert's parent changed"),
 			FProcessor(EQueryTickPhase::FrameEnd, DataStorage.GetQueryTickGroupName(EQueryTickGroups::Default)),
-			[](IQueryContext& Context, FTypedElementChildAlertColumn& ChildAlert, const FTypedElementParentColumn& Parent)
+			[](IQueryContext& Context, FTypedElementChildAlertColumn& ChildAlert, const FTableRowParentColumn& Parent)
 			{
 				if (ChildAlert.CachedParent != Parent.Parent)
 				{
@@ -220,7 +220,7 @@ void UTypedElementAlertQueriesFactory::RegisterChildAlertUpdatesQueries(ITypedEl
 			TEXT("Add missing child alerts"),
 			FPhaseAmble(FPhaseAmble::ELocation::Preamble, EQueryTickPhase::FrameEnd)
 				.MakeActivatable(AlertConditionName),
-			[](IQueryContext& Context, RowHandle Row, FTypedElementAlertColumn& Alert, const FTypedElementParentColumn& Parent)
+			[](IQueryContext& Context, RowHandle Row, FTypedElementAlertColumn& Alert, const FTableRowParentColumn& Parent)
 			{
 				AddChildAlertsToHierarchy(Context, Parent.Parent, 0);
 			})
@@ -285,13 +285,13 @@ void UTypedElementAlertQueriesFactory::RegisterOnAddQueries(ITypedElementDataSto
 				Context.ActivateQueries(AlertConditionName);
 			})
 		.Where()
-			.All<FTypedElementParentColumn>() // Only need to do an update pass if there are parents.
+			.All<FTableRowParentColumn>() // Only need to do an update pass if there are parents.
 		.Compile());
 
 	DataStorage.RegisterQuery(
 		Select(
 			TEXT("Register alert with parent on parent add"),
-			FObserver::OnAdd<FTypedElementParentColumn>(),
+			FObserver::OnAdd<FTableRowParentColumn>(),
 			[](IQueryContext& Context, RowHandle Row)
 			{
 				Context.ActivateQueries(AlertConditionName);
@@ -319,7 +319,7 @@ void UTypedElementAlertQueriesFactory::RegisterOnRemoveQueries(ITypedElementData
 	DataStorage.RegisterQuery(
 		Select(
 			TEXT("Update alert upon parent removal"),
-			FObserver::OnRemove<FTypedElementParentColumn>(),
+			FObserver::OnRemove<FTableRowParentColumn>(),
 			[](IQueryContext& Context, TypedElementDataStorage::RowHandle Row)
 			{
 				Context.ActivateQueries(AlertConditionName);
@@ -388,7 +388,7 @@ bool UTypedElementAlertQueriesFactory::MoveToNextParent(
 	using namespace TypedElementQueryBuilder;
 	
 	FQueryResult Result = Context.RunSubquery(SubQueryIndex, Parent, CreateSubqueryCallbackBinding(
-		[&Parent](const FTypedElementParentColumn& NextParent)
+		[&Parent](const FTableRowParentColumn& NextParent)
 		{
 			Parent = NextParent.Parent;
 		}));
diff --git a/Engine/Plugins/Experimental/EditorDataStorage/Source/TedsCore/Private/Queries/TypedElementHierarchyQueries.cpp b/Engine/Plugins/Experimental/EditorDataStorage/Source/TedsCore/Private/Queries/TypedElementHierarchyQueries.cpp
--- a/Engine/Plugins/Experimental/EditorDataStorage/Source/TedsCore/Private/Queries/TypedElementHierarchyQueries.cpp
+++ b/Engine/Plugins/Experimental/EditorDataStorage/Source/TedsCore/Private/Queries/TypedElementHierarchyQueries.cpp
@@ -15,13 +15,13 @@ void UTypedElementHiearchyQueriesFactory::RegisterQueries(ITypedElementDataStora
 		Select(
 			TEXT("Resolve hierarchy rows"),
 			FProcessor(DSI::EQueryTickPhase::PrePhysics, DataStorage.GetQueryTickGroupName(DSI::EQueryTickGroups::Default)),
-			[](DSI::IQueryContext& Context, RowHandle Row, const FTypedElementUnresolvedParentColumn& UnresolvedParent)
+			[](DSI::IQueryContext& Context, RowHandle Row, const FUnresolvedTableRowParentColumn& UnresolvedParent)
 			{
 				RowHandle ParentRow = Context.FindIndexedRow(UnresolvedParent.ParentIdHash);
 				if (Context.IsRowAvailable(ParentRow))
 				{
-					Context.RemoveColumns<FTypedElementUnresolvedParentColumn>(Row);
-					Context.AddColumn(Row, FTypedElementParentColumn{ .Parent = ParentRow });
+					Context.RemoveColumns<FUnresolvedTableRowParentColumn>(Row);
+					Context.AddColumn(Row, FTableRowParentColumn{ .Parent = ParentRow });
 				}
 			}
 		)
diff --git a/Engine/Plugins/Experimental/EditorDataStorageFeatures/Source/TedsDebugger/Private/STedsDebugger.cpp b/Engine/Plugins/Experimental/EditorDataStorageFeatures/Source/TedsDebugger/Private/STedsDebugger.cpp
--- a/Engine/Plugins/Experimental/EditorDataStorageFeatures/Source/TedsDebugger/Private/STedsDebugger.cpp
+++ b/Engine/Plugins/Experimental/EditorDataStorageFeatures/Source/TedsDebugger/Private/STedsDebugger.cpp
@@ -12,7 +12,7 @@
 #include "Modules/ModuleManager.h"
 #include "QueryEditor/TedsQueryEditor.h"
 #include "TedsOutlinerModule.h"
-#include "TypedElementOutlinerMode.h"
+#include "TedsOutlinerMode.h"
 
 #define LOCTEXT_NAMESPACE "STedsDebugger"
 
@@ -185,10 +185,10 @@ TSharedRef<SDockTab> STedsDebugger::SpawnTableViewerTab(const FSpawnTabArgs& Arg
 	InitOptions.FilterBarOptions.bHasFilterBar = true;
 	InitOptions.OutlinerIdentifier = "TedsDebugger.TableViewer";
 
-	FTypedElementOutlinerModeParams Params(nullptr);
+	FTedsOutlinerParams Params(nullptr);
 	Params.QueryDescription = RowQueryDescription;
 	Params.bUseDefaultTedsFilters = true;
-	Params.HierarchyData = TOptional<FTypedElementOutlinerHierarchyData>(); // We don't want to show hierarchies in the debugger
+	Params.HierarchyData = TOptional<FTedsOutlinerHierarchyData>(); // We don't want to show hierarchies in the debugger
 	Params.CellWidgetPurposes = TArray<FName>{TEXT("General.Cell")};
 	
 	FTedsOutlinerModule& TedsOutlinerModule = FModuleManager::GetModuleChecked<FTedsOutlinerModule>("TedsOutliner");
diff --git a/Engine/Plugins/Experimental/EditorDataStorageFeatures/Source/TedsOutliner/Private/Compatibility/SceneOutlinerRowHandleColumn.cpp b/Engine/Plugins/Experimental/EditorDataStorageFeatures/Source/TedsOutliner/Private/Compatibility/SceneOutlinerRowHandleColumn.cpp
new file mode 100644
--- /dev/null
+++ b/Engine/Plugins/Experimental/EditorDataStorageFeatures/Source/TedsOutliner/Private/Compatibility/SceneOutlinerRowHandleColumn.cpp
@@ -0,0 +1,91 @@
+// Copyright Epic Games, Inc. All Rights Reserved.
+
+#include "Compatibility/SceneOutlinerRowHandleColumn.h"
+
+#include "SortHelper.h"
+#include "TedsOutlinerItem.h"
+#include "Elements/Framework/TypedElementRegistry.h"
+#include "TedsTableViewerColumn.h"
+
+#define LOCTEXT_NAMESPACE "SceneOutlinerRowHandleColumn"
+
+FSceneOutlinerRowHandleColumn::FSceneOutlinerRowHandleColumn(ISceneOutliner& SceneOutliner)
+	: WeakSceneOutliner(StaticCastSharedRef<ISceneOutliner>(SceneOutliner.AsShared()))
+{
+	auto AssignWidgetToColumn = [this](TUniquePtr<FTypedElementWidgetConstructor> Constructor, TConstArrayView<TWeakObjectPtr<const UScriptStruct>>)
+	{
+		TSharedPtr<FTypedElementWidgetConstructor> WidgetConstructor(Constructor.Release());
+		TableViewerColumn = MakeShared<UE::EditorDataStorage::FTedsTableViewerColumn>(TEXT("Row Handle"), WidgetConstructor);
+		return false;
+	};
+	
+	UTypedElementRegistry* Registry = UTypedElementRegistry::GetInstance();
+	checkf(Registry, TEXT("FTedsOutlinerRowHandleColumn created before UTypedElementRegistry is available."));
+	ITypedElementDataStorageUiInterface* StorageUi = Registry->GetMutableDataStorageUi();
+	checkf(StorageUi, TEXT("FTedsOutlinerRowHandleColumn created before data storage interfaces were initialized."))
+
+	StorageUi->CreateWidgetConstructors(TEXT("General.Cell.RowHandle"), TypedElementDataStorage::FMetaDataView(), AssignWidgetToColumn);
+}
+
+
+FName FSceneOutlinerRowHandleColumn::GetID()
+{
+	static const FName ID("Row Handle");
+	return ID;
+}
+
+FName FSceneOutlinerRowHandleColumn::GetColumnID()
+{
+	return GetID();
+}
+
+SHeaderRow::FColumn::FArguments FSceneOutlinerRowHandleColumn::ConstructHeaderRowColumn()
+{
+	return SHeaderRow::Column(GetID())
+	.FillWidth(2)
+	.HeaderComboVisibility(EHeaderComboVisibility::OnHover);
+}
+
+const TSharedRef<SWidget> FSceneOutlinerRowHandleColumn::ConstructRowWidget(FSceneOutlinerTreeItemRef TreeItem, const STableRow<FSceneOutlinerTreeItemPtr>& Row)
+{
+	auto SceneOutliner = WeakSceneOutliner.Pin();
+	check(SceneOutliner.IsValid());
+
+	if (const FTedsOutlinerTreeItem* OutlinerTreeItem = TreeItem->CastTo<FTedsOutlinerTreeItem>())
+	{
+		const TypedElementDataStorage::RowHandle RowHandle = OutlinerTreeItem->GetRowHandle();
+
+		if(TSharedPtr<SWidget> Widget = TableViewerColumn->ConstructRowWidget(RowHandle))
+		{
+			return Widget.ToSharedRef();
+		}
+	}
+	return SNullWidget::NullWidget;
+}
+
+void FSceneOutlinerRowHandleColumn::PopulateSearchStrings(const ISceneOutlinerTreeItem& Item, TArray<FString>& OutSearchStrings) const
+{
+	if (const FTedsOutlinerTreeItem* OutlinerTreeItem = Item.CastTo<FTedsOutlinerTreeItem>())
+	{
+		OutSearchStrings.Add(LexToString<FString>(OutlinerTreeItem->GetRowHandle()));
+	}
+
+}
+
+void FSceneOutlinerRowHandleColumn::SortItems(TArray<FSceneOutlinerTreeItemPtr>& OutItems, const EColumnSortMode::Type SortMode) const
+{
+	FSceneOutlinerSortHelper<TypedElementDataStorage::RowHandle>()
+		/** Sort by type first */
+		.Primary([this](const ISceneOutlinerTreeItem& Item)
+		{
+			if (const FTedsOutlinerTreeItem* OutlinerTreeItem = Item.CastTo<FTedsOutlinerTreeItem>())
+			{
+				return OutlinerTreeItem->GetRowHandle();
+			}
+
+			return TypedElementDataStorage::InvalidRowHandle;
+		}, SortMode)
+		.Sort(OutItems);
+}
+
+#undef LOCTEXT_NAMESPACE
diff --git a/Engine/Plugins/Experimental/EditorDataStorageFeatures/Source/TedsOutliner/Private/Compatibility/SceneOutlinerRowHandleColumn.h b/Engine/Plugins/Experimental/EditorDataStorageFeatures/Source/TedsOutliner/Private/Compatibility/SceneOutlinerRowHandleColumn.h
new file mode 100644
--- /dev/null
+++ b/Engine/Plugins/Experimental/EditorDataStorageFeatures/Source/TedsOutliner/Private/Compatibility/SceneOutlinerRowHandleColumn.h
@@ -0,0 +1,43 @@
+// Copyright Epic Games, Inc. All Rights Reserved.
+
+#pragma once
+
+#include "SceneOutlinerFwd.h"
+#include "ISceneOutliner.h"
+#include "ISceneOutlinerColumn.h"
+
+namespace UE::EditorDataStorage
+{
+	class FTedsTableViewerColumn;
+}
+
+/**
+ * This is a custom column for the TEDS Outliner/Table Viewer to display row handles for items. It is a special cased column instead of going through
+ * the TEDS UI layer because row handles are not information stored in a TEDS column
+ */
+class FSceneOutlinerRowHandleColumn : public ISceneOutlinerColumn
+{
+public:
+	FSceneOutlinerRowHandleColumn(ISceneOutliner& SceneOutliner);
+	
+	virtual ~FSceneOutlinerRowHandleColumn() {}
+
+	static FName GetID();
+	
+	//////////////////////////////////////////////////////////////////////////
+	// Begin ISceneOutlinerColumn Implementation
+	virtual FName GetColumnID() override;
+	virtual SHeaderRow::FColumn::FArguments ConstructHeaderRowColumn() override;
+	virtual const TSharedRef<SWidget> ConstructRowWidget(FSceneOutlinerTreeItemRef TreeItem, const STableRow<FSceneOutlinerTreeItemPtr>& Row) override;
+	virtual void PopulateSearchStrings( const ISceneOutlinerTreeItem& Item, TArray< FString >& OutSearchStrings ) const override;
+	virtual bool SupportsSorting() const override { return true; }
+	virtual void SortItems(TArray<FSceneOutlinerTreeItemPtr>& OutItems, const EColumnSortMode::Type SortMode) const override;
+	// End ISceneOutlinerColumn Implementation
+	//////////////////////////////////////////////////////////////////////////
+
+private:
+	TWeakPtr<ISceneOutliner> WeakSceneOutliner;
+
+	// The Table Viewer column we are going to use internally to create the widget
+	TSharedPtr<UE::EditorDataStorage::FTedsTableViewerColumn> TableViewerColumn;
+};
diff --git a/Engine/Plugins/Experimental/EditorDataStorageFeatures/Source/TedsOutliner/Private/Compatibility/SceneOutlinerTedsBridge.cpp b/Engine/Plugins/Experimental/EditorDataStorageFeatures/Source/TedsOutliner/Private/Compatibility/SceneOutlinerTedsBridge.cpp
new file mode 100644
--- /dev/null
+++ b/Engine/Plugins/Experimental/EditorDataStorageFeatures/Source/TedsOutliner/Private/Compatibility/SceneOutlinerTedsBridge.cpp
@@ -0,0 +1,599 @@
+// Copyright Epic Games, Inc. All Rights Reserved.
+
+#include "Compatibility/SceneOutlinerTedsBridge.h"
+
+#include "ActorTreeItem.h"
+#include "Elements/Columns/TypedElementLabelColumns.h"
+#include "Elements/Columns/TypedElementMiscColumns.h"
+#include "Elements/Columns/TypedElementPackageColumns.h"
+#include "Elements/Columns/TypedElementTransformColumns.h"
+#include "Elements/Columns/TypedElementTypeInfoColumns.h"
+#include "Elements/Framework/TypedElementRegistry.h"
+#include "Elements/Framework/TypedElementQueryBuilder.h"
+#include "Elements/Interfaces/TypedElementDataStorageCompatibilityInterface.h"
+#include "ILevelEditor.h"
+#include "ISceneOutliner.h"
+#include "ISceneOutlinerColumn.h"
+#include "ISceneOutlinerTreeItem.h"
+#include "LevelEditor.h"
+#include "SSceneOutliner.h"
+#include "SceneOutlinerFwd.h"
+#include "SceneOutlinerModule.h"
+#include "SceneOutlinerPublicTypes.h"
+#include "TedsTableViewerColumn.h"
+#include "TedsTableViewerUtils.h"
+#include "Elements/Interfaces/Capabilities/TypedElementUiTextCapability.h"
+#include "TedsOutlinerItem.h"
+#include "Modules/ModuleManager.h"
+
+#define LOCTEXT_NAMESPACE "SceneOutlinerTedsBridge"
+
+FAutoConsoleCommand BindColumnsToSceneOutlinerConsoleCommand(
+	TEXT("TEDS.UI.BindColumnsToSceneOutliner"),
+	TEXT("Bind one or more columns to the most recently used Scene Outliner. Several prebuild configurations are offered as well.")
+	TEXT("An example input to show a label column is 'TEDS.UI.BindColumnsToSceneOutliner /Script/TypedElementFramework.TypedElementLabelColumn'."),
+	FConsoleCommandWithArgsDelegate::CreateLambda([](const TArray<FString>& Args)
+		{
+			using namespace TypedElementQueryBuilder;
+			using DSI = ITypedElementDataStorageInterface;
+
+		    const FName WidgetPurposes[] = {TEXT("SceneOutliner.Cell"), TEXT("General.Cell")};
+
+			UTypedElementRegistry* Registry = UTypedElementRegistry::GetInstance();
+			if (ITypedElementDataStorageInterface* DataStorage = Registry->GetMutableDataStorage())
+			{
+				static TypedElementQueryHandle Queries[] =
+				{
+					DataStorage->RegisterQuery(Select().ReadWrite<FTypedElementLabelColumn>().Compile()),
+					DataStorage->RegisterQuery(Select().ReadOnly<FTypedElementLocalTransformColumn>().Compile()),
+					DataStorage->RegisterQuery(Select().ReadOnly<FTypedElementPackagePathColumn>().Compile()),
+					DataStorage->RegisterQuery(
+						Select()
+							.ReadWrite<FTypedElementLabelColumn>()
+							.ReadOnly<FTypedElementLocalTransformColumn>()
+						.Compile()),
+					DataStorage->RegisterQuery(
+						Select()
+							.ReadOnly<FTypedElementLabelColumn>()
+							.ReadOnly<FTypedElementLabelHashColumn>()
+						.Compile())
+				};
+
+				FSceneOutlinerTedsQueryBinder& Binder = FSceneOutlinerTedsQueryBinder::GetInstance();
+				const TWeakPtr<ILevelEditor> LevelEditor = FModuleManager::GetModuleChecked<FLevelEditorModule>(TEXT("LevelEditor")).GetLevelEditorInstance();
+				const TSharedPtr<ISceneOutliner> SceneOutliner = LevelEditor.IsValid() ? LevelEditor.Pin()->GetMostRecentlyUsedSceneOutliner() : nullptr;
+				if (SceneOutliner.IsValid())
+				{
+					if (!Args.IsEmpty())
+					{
+						if (Args[0].IsNumeric())
+						{
+							int32 QueryIndex = FCString::Atoi(*Args[0]);
+							if (QueryIndex < sizeof(Queries) / sizeof(TypedElementQueryHandle))
+							{
+								Binder.AssignQuery(Queries[QueryIndex], SceneOutliner, WidgetPurposes);
+								return;
+							}
+						}
+						else
+						{
+							uint32 AdditionCount = 0;
+							Select Query;
+							for (const FString& Arg : Args)
+							{
+								FTopLevelAssetPath Path;
+								// TrySetPath has an ensure that checks if the path starts with an '/' and otherwise throws
+								// an assert.
+								if (!Arg.IsEmpty() && Arg[0] == '/' && Path.TrySetPath(Arg))
+								{
+									const UScriptStruct* ColumnType = TypeOptional(Path);
+									if (ColumnType && ColumnType->IsChildOf(FTypedElementDataStorageColumn::StaticStruct()))
+									{
+										Query.ReadOnly(ColumnType);
+										++AdditionCount;
+									}
+								}
+							}
+							if (AdditionCount > 0)
+							{
+								static TypedElementQueryHandle CustomQuery = TypedElementInvalidQueryHandle;
+								if (CustomQuery != TypedElementInvalidQueryHandle)
+								{
+									DataStorage->UnregisterQuery(CustomQuery);
+								}
+								CustomQuery = DataStorage->RegisterQuery(Query.Compile());
+								Binder.AssignQuery(CustomQuery, SceneOutliner, WidgetPurposes);
+								return;
+							}
+						}
+					}
+					Binder.AssignQuery(TypedElementInvalidQueryHandle, SceneOutliner, WidgetPurposes);
+				}
+			}
+		}));
+
+class FSceneOutlinerTedsBridge
+{
+public:
+	~FSceneOutlinerTedsBridge();
+
+	void Initialize(
+		ITypedElementDataStorageInterface& InStorage,
+		ITypedElementDataStorageUiInterface& InStorageUi,
+		ITypedElementDataStorageCompatibilityInterface& InStorageCompatibility,
+		const TSharedPtr<ISceneOutliner>& InOutliner);
+
+	void AssignQuery(TypedElementQueryHandle Query, const TConstArrayView<FName> CellWidgetPurposes);
+	void RegisterDealiaser(const FTreeItemIDDealiaser& InDealiaser);
+private:
+	void ClearColumns(ISceneOutliner& InOutliner);
+
+	TArray<FName> AddedColumns;
+	TWeakPtr<ISceneOutliner> Outliner;
+	ITypedElementDataStorageInterface* Storage{ nullptr };
+	ITypedElementDataStorageUiInterface* StorageUi{ nullptr };
+	ITypedElementDataStorageCompatibilityInterface* StorageCompatibility{ nullptr };
+	FTreeItemIDDealiaser Dealiaser;
+	TArray<FName> CellWidgetPurposes;
+};
+
+class FOutlinerColumn : public ISceneOutlinerColumn
+{
+public:
+	FOutlinerColumn(
+		TypedElementDataStorage::QueryHandle InQuery,
+		ITypedElementDataStorageInterface& InStorage, 
+		ITypedElementDataStorageUiInterface& InStorageUi, 
+		ITypedElementDataStorageCompatibilityInterface& InStorageCompatibility,
+		FName InNameId, 
+		TArray<TWeakObjectPtr<const UScriptStruct>> InColumnTypes,
+		TSharedPtr<FTypedElementWidgetConstructor> InHeaderWidgetConstructor,
+		TSharedPtr<FTypedElementWidgetConstructor> InCellWidgetConstructor,
+		FName InFallbackColumnName,
+		TWeakPtr<ISceneOutliner> InOwningOutliner,
+		const FTreeItemIDDealiaser& InDealiaser)
+		: Storage(InStorage)
+		, StorageUi(InStorageUi)
+		, StorageCompatibility(InStorageCompatibility)
+		, QueryHandle(InQuery)
+		, NameId(InNameId)
+		, OwningOutliner(InOwningOutliner)
+		, Dealiaser(InDealiaser)
+	{
+		MetaData.AddOrSetMutableData(TEXT("Name"), NameId.ToString());
+
+		using namespace TypedElementDataStorage;
+
+		TableViewerColumnImpl = MakeUnique<UE::EditorDataStorage::FTedsTableViewerColumn>(NameId, InCellWidgetConstructor, InColumnTypes,
+			InHeaderWidgetConstructor, FComboMetaDataView(FGenericMetaDataView(MetaData)).Next(FQueryMetaDataView(Storage.GetQueryDescription(QueryHandle))));
+
+		TableViewerColumnImpl->SetIsRowVisibleDelegate(
+		UE::EditorDataStorage::FTedsTableViewerColumn::FIsRowVisible::CreateRaw(this, &FOutlinerColumn::IsRowVisible)
+		);
+		
+		// Try to find a fallback column from the regular item, for handling cases like folders which are not in TEDS but want to use TEDS columns
+		FSceneOutlinerModule& SceneOutlinerModule = FModuleManager::LoadModuleChecked<FSceneOutlinerModule>("SceneOutliner");
+		FallbackColumn = SceneOutlinerModule.FactoryColumn(InFallbackColumnName, *OwningOutliner.Pin());
+	};
+	
+	FName GetColumnID() override
+	{
+		return NameId;
+	}
+
+	virtual void Tick(double InCurrentTime, float InDeltaTime)
+	{
+		TableViewerColumnImpl->Tick();
+	}
+
+	bool IsRowVisible(const TypedElementDataStorage::RowHandle InRowHandle) const
+	{
+		TSharedPtr<ISceneOutliner> OutlinerPinned = OwningOutliner.Pin();
+
+		if(!OutlinerPinned)
+		{
+			return false;
+		}
+		
+		// Try to grab the TEDS Outliner item from the row handle
+		FSceneOutlinerTreeItemPtr Item = OutlinerPinned->GetTreeItem(InRowHandle);
+
+		// If it doesn't exist, this could be a legacy item that uses something other than the row id as the ID, so check if we have a dealiaser
+		if(!Item)
+		{
+			if(Dealiaser.IsBound())
+			{
+				Item = OutlinerPinned->GetTreeItem(Dealiaser.Execute(InRowHandle));
+			}
+		}
+
+		if(!Item)
+		{
+			return false;
+		}
+
+		// Check if the item is visible in the tree
+		return OutlinerPinned->GetTree().IsItemVisible(Item);
+	}
+
+	virtual SHeaderRow::FColumn::FArguments ConstructHeaderRowColumn() override
+	{
+		return TableViewerColumnImpl->ConstructHeaderRowColumn();
+	}
+
+	// TODO: Sorting is currently handled through the fallback column if it exists because we have no way to sort columns through TEDS
+	virtual void SortItems(TArray<FSceneOutlinerTreeItemPtr>& RootItems, const EColumnSortMode::Type SortMode) const override
+	{
+		if(FallbackColumn)
+		{
+			FallbackColumn->SortItems(RootItems, SortMode);
+		}
+	}
+
+	virtual bool SupportsSorting() const override
+	{
+		return FallbackColumn ? FallbackColumn->SupportsSorting() : false;
+	}
+
+	void SetHighlightText(SWidget& Widget)
+	{
+		TSharedPtr<ISceneOutliner> OutlinerPinned = OwningOutliner.Pin();
+
+		if(!OutlinerPinned)
+		{
+			return;
+		}
+
+		if (TSharedPtr<ITypedElementUiTextCapability> TextCapability = Widget.GetMetaData<ITypedElementUiTextCapability>())
+		{
+			TextCapability->SetHighlightText(OutlinerPinned->GetFilterHighlightText());
+		}
+	
+		if (FChildren* ChildWidgets = Widget.GetChildren())
+		{
+			ChildWidgets->ForEachWidget([this](SWidget& ChildWidget)
+				{
+					SetHighlightText(ChildWidget);
+				});
+		}
+	}
+	
+	const TSharedRef<SWidget> ConstructRowWidget(FSceneOutlinerTreeItemRef TreeItem, const STableRow<FSceneOutlinerTreeItemPtr>& Row) override
+	{
+		using namespace TypedElementDataStorage;
+		
+		RowHandle RowHandle = InvalidRowHandle;
+
+		TSharedPtr<SWidget> RowWidget;
+
+		if(const FTedsOutlinerTreeItem* TedsItem = TreeItem->CastTo<FTedsOutlinerTreeItem>())
+		{
+			RowHandle = TedsItem->GetRowHandle();
+			
+		}
+		else if (const FActorTreeItem* ActorItem = TreeItem->CastTo<FActorTreeItem>())
+		{
+			if (AActor* Actor = ActorItem->Actor.Get())
+			{
+				RowHandle = StorageCompatibility.FindRowWithCompatibleObject(Actor);
+			}
+		}
+		else if(FallbackColumn)
+		{
+			RowWidget = FallbackColumn->ConstructRowWidget(TreeItem, Row);
+		}
+
+		if(Storage.IsRowAssigned(RowHandle))
+		{
+			RowWidget = TableViewerColumnImpl->ConstructRowWidget(RowHandle);
+		}
+
+		if(RowWidget)
+		{
+			SetHighlightText(*RowWidget);
+			return RowWidget.ToSharedRef();
+		}
+
+		return SNullWidget::NullWidget;
+	}
+
+	virtual void PopulateSearchStrings( const ISceneOutlinerTreeItem& Item, TArray< FString >& OutSearchStrings ) const override
+	{
+		// TODO: We don't currently have a way to convert TEDS widgets into searchable strings, but we can rely on the fallback column if it exists
+		if(FallbackColumn)
+		{
+			FallbackColumn->PopulateSearchStrings(Item, OutSearchStrings);
+		}
+	}
+
+	// The table viewer implementation that we internally use to create our widgets
+	TUniquePtr<UE::EditorDataStorage::FTedsTableViewerColumn> TableViewerColumnImpl;
+	
+	ITypedElementDataStorageInterface& Storage;
+	ITypedElementDataStorageUiInterface& StorageUi;
+	ITypedElementDataStorageCompatibilityInterface& StorageCompatibility;
+	TypedElementDataStorage::QueryHandle QueryHandle;
+	TypedElementDataStorage::FMetaData MetaData;
+	FName NameId;
+	TSharedPtr<ISceneOutlinerColumn> FallbackColumn;
+	TWeakPtr<ISceneOutliner> OwningOutliner;
+	FTreeItemIDDealiaser Dealiaser;
+};
+
+
+
+
+//
+// USceneOutlinerTedsBridgeFactory
+// 
+
+void USceneOutlinerTedsBridgeFactory::RegisterWidgetPurposes(ITypedElementDataStorageUiInterface& DataStorageUi) const
+{
+	using PurposeType = ITypedElementDataStorageUiInterface::EPurposeType;
+
+	DataStorageUi.RegisterWidgetPurpose(FSceneOutlinerTedsQueryBinder::HeaderWidgetPurpose, PurposeType::UniqueByNameAndColumn,
+		LOCTEXT("HeaderWidgetPurpose", "Widgets for headers in any Scene Outliner for specific columns or column combinations."));
+	DataStorageUi.RegisterWidgetPurpose(FSceneOutlinerTedsQueryBinder::DefaultHeaderWidgetPurpose, PurposeType::UniqueByName,
+		LOCTEXT("DefaultHeaderWidgetPurpose", "The default widget to use in headers for the Scene Outliner."));
+	
+	DataStorageUi.RegisterWidgetPurpose(FSceneOutlinerTedsQueryBinder::CellWidgetPurpose, PurposeType::UniqueByNameAndColumn,
+		LOCTEXT("CellWidgetPurpose", "Widgets for cells in any Scene Outliner for specific columns or column combinations."));
+	DataStorageUi.RegisterWidgetPurpose(FSceneOutlinerTedsQueryBinder::DefaultCellWidgetPurpose, PurposeType::UniqueByName,
+		LOCTEXT("DefaultCellWidgetPurpose", "The default widget to use in cells for the Scene Outliner."));
+
+	DataStorageUi.RegisterWidgetPurpose(FSceneOutlinerTedsQueryBinder::ItemLabelCellWidgetPurpose, PurposeType::UniqueByNameAndColumn,
+		LOCTEXT("ItemCellWidgetPurpose", "Widgets for cells in any Scene Outliner that are specific to the Item label column."));
+	DataStorageUi.RegisterWidgetPurpose(FSceneOutlinerTedsQueryBinder::DefaultItemLabelCellWidgetPurpose, PurposeType::UniqueByName,
+		LOCTEXT("DefaultItemCellWidgetPurpose", "The default widget to use in cells for the Scene Outliner specific to the Item label column."));
+
+
+}
+
+
+
+//
+// FSceneOutlinerTedsQueryBinder
+// 
+
+const FName FSceneOutlinerTedsQueryBinder::CellWidgetTableName(TEXT("Editor_SceneOutlinerCellWidgetTable"));
+const FName FSceneOutlinerTedsQueryBinder::HeaderWidgetPurpose(TEXT("SceneOutliner.Header"));
+const FName FSceneOutlinerTedsQueryBinder::DefaultHeaderWidgetPurpose(TEXT("SceneOutliner.Header.Default"));
+const FName FSceneOutlinerTedsQueryBinder::CellWidgetPurpose(TEXT("SceneOutliner.Cell"));
+const FName FSceneOutlinerTedsQueryBinder::DefaultCellWidgetPurpose(TEXT("SceneOutliner.Cell.Default"));
+const FName FSceneOutlinerTedsQueryBinder::ItemLabelCellWidgetPurpose(TEXT("SceneOutliner.ItemLabel.Cell"));
+const FName FSceneOutlinerTedsQueryBinder::DefaultItemLabelCellWidgetPurpose(TEXT("SceneOutliner.ItemLabel.Cell.Default"));
+
+FSceneOutlinerTedsQueryBinder::FSceneOutlinerTedsQueryBinder()
+{
+	UTypedElementRegistry* Registry = UTypedElementRegistry::GetInstance();
+	checkf(Registry, TEXT("Unable to bind a Scene Outliner to a query before the Typed Elements are available."));
+	if (Registry)
+	{
+		Storage = Registry->GetMutableDataStorage();
+		StorageUi = Registry->GetMutableDataStorageUi();
+		StorageCompatibility = Registry->GetMutableDataStorageCompatibility();
+	}
+
+	SetupDefaultColumnMapping();
+}
+
+void FSceneOutlinerTedsQueryBinder::SetupDefaultColumnMapping()
+{
+	// Map the type column from the TEDS to the default Outliner type column, so we can show type info for objects not in TEDS
+	TEDSToOutlinerDefaultColumnMapping.Add(FTypedElementClassTypeInfoColumn::StaticStruct(), FSceneOutlinerBuiltInColumnTypes::ActorInfo());
+}
+
+FName FSceneOutlinerTedsQueryBinder::FindOutlinerColumnFromTEDSColumns(TConstArrayView<TWeakObjectPtr<const UScriptStruct>> TEDSColumns) const
+{
+	// Currently, the algorithm naively looks through the mapping and returns the first match
+	for(const TWeakObjectPtr<const UScriptStruct>& Column : TEDSColumns)
+	{
+		if(const FName* FoundDefaultColumn = TEDSToOutlinerDefaultColumnMapping.Find(Column))
+		{
+			return *FoundDefaultColumn;
+		}
+	}
+
+	return FName();
+}
+
+FSceneOutlinerTedsQueryBinder& FSceneOutlinerTedsQueryBinder::GetInstance()
+{
+	static FSceneOutlinerTedsQueryBinder Binder;
+	return Binder;
+}
+
+TSharedPtr<FSceneOutlinerTedsBridge>* FSceneOutlinerTedsQueryBinder::FindOrAddQueryMapping(const TSharedPtr<ISceneOutliner>& Outliner)
+{
+	TSharedPtr<FSceneOutlinerTedsBridge>* QueryMapping = SceneOutliners.Find(Outliner);
+	if (QueryMapping == nullptr)
+	{
+		QueryMapping = &SceneOutliners.Add(Outliner, MakeShared<FSceneOutlinerTedsBridge>());
+		(*QueryMapping)->Initialize(*Storage, *StorageUi, *StorageCompatibility, Outliner);
+	}
+
+	return QueryMapping;
+}
+
+void FSceneOutlinerTedsQueryBinder::AssignQuery(TypedElementQueryHandle Query, const TSharedPtr<ISceneOutliner>& Outliner, TConstArrayView<FName> CellWidgetPurposes)
+{
+	CleanupStaleOutliners();
+
+	TSharedPtr<FSceneOutlinerTedsBridge>* QueryMapping = FindOrAddQueryMapping(Outliner);
+	(*QueryMapping)->AssignQuery(Query, CellWidgetPurposes);
+}
+
+void FSceneOutlinerTedsQueryBinder::RegisterTreeItemIDDealiaser(const TSharedPtr<ISceneOutliner>& Outliner, const FTreeItemIDDealiaser& InDealiaser)
+{
+	TSharedPtr<FSceneOutlinerTedsBridge>* QueryMapping = FindOrAddQueryMapping(Outliner);
+	(*QueryMapping)->RegisterDealiaser(InDealiaser);
+}
+
+void FSceneOutlinerTedsQueryBinder::CleanupStaleOutliners()
+{
+	for (TMap<TWeakPtr<ISceneOutliner>, TSharedPtr<FSceneOutlinerTedsBridge>>::TIterator It(SceneOutliners); It; ++It)
+	{
+		// Remove any query mappings where the target Outliner doesn't exist anymore
+		if(!It.Key().IsValid())
+		{
+			It.RemoveCurrent();
+		}
+	}
+}
+
+//
+// FSceneOutlinerTedsBridge
+//
+
+FSceneOutlinerTedsBridge::~FSceneOutlinerTedsBridge()
+{
+	TSharedPtr<ISceneOutliner> OutlinerPinned = Outliner.Pin();
+	if (OutlinerPinned)
+	{
+		ClearColumns(*OutlinerPinned);
+	}
+}
+
+void FSceneOutlinerTedsBridge::Initialize(
+	ITypedElementDataStorageInterface& InStorage,
+	ITypedElementDataStorageUiInterface& InStorageUi,
+	ITypedElementDataStorageCompatibilityInterface& InStorageCompatibility,
+	const TSharedPtr<ISceneOutliner>& InOutliner)
+{
+	Storage = &InStorage;
+	StorageUi = &InStorageUi;
+	StorageCompatibility = &InStorageCompatibility;
+	Outliner = InOutliner;
+}
+
+void FSceneOutlinerTedsBridge::RegisterDealiaser(const FTreeItemIDDealiaser& InDealiaser)
+{
+	Dealiaser = InDealiaser;
+}
+
+void FSceneOutlinerTedsBridge::AssignQuery(TypedElementQueryHandle Query, const TConstArrayView<FName> InCellWidgetPurposes)
+{
+	using MatchApproach = ITypedElementDataStorageUiInterface::EMatchApproach;
+	constexpr int32 DefaultPriorityIndex = 100;
+	FSceneOutlinerTedsQueryBinder& Binder = FSceneOutlinerTedsQueryBinder::GetInstance();
+	CellWidgetPurposes = InCellWidgetPurposes;
+
+
+	if (TSharedPtr<ISceneOutliner> OutlinerPinned = Outliner.Pin())
+	{
+		const ITypedElementDataStorageInterface::FQueryDescription& Description = Storage->GetQueryDescription(Query);
+		TypedElementDataStorage::FQueryMetaDataView MetaDataView(Description);
+
+		ClearColumns(*OutlinerPinned);
+
+		if (Description.Action == ITypedElementDataStorageInterface::FQueryDescription::EActionType::Select)
+		{
+			int32 SelectionCount = Description.SelectionTypes.Num();
+			AddedColumns.Reset(SelectionCount);
+
+			TArray<TWeakObjectPtr<const UScriptStruct>> ColumnTypes = UE::EditorDataStorage::TableViewerUtils::CreateVerifiedColumnTypeArray(Description.SelectionTypes);
+
+			int32 IndexOffset = 0;
+			auto ColumnConstructor = [this, Query, MetaDataView, &IndexOffset, &OutlinerPinned, Binder](
+				TUniquePtr<FTypedElementWidgetConstructor> Constructor, TConstArrayView<TWeakObjectPtr<const UScriptStruct>> ColumnTypes)
+				{
+					TSharedPtr<FTypedElementWidgetConstructor> CellConstructor(Constructor.Release());
+
+					/* If we have a fallback column for this query, remove it, take over it's priority and 
+					 * replace it with the TEDS column. But also allow the TEDS-Outliner column to fallback to it for
+					 * data not in TEDS yet.
+				 	 */
+					FName FallbackColumn = Binder.FindOutlinerColumnFromTEDSColumns(ColumnTypes);
+					const FSceneOutlinerColumnInfo* FallbackColumnInfo = OutlinerPinned->GetSharedData().ColumnMap.Find(FallbackColumn);
+					int32 ColumnPriority = FallbackColumnInfo ? FallbackColumnInfo->PriorityIndex : DefaultPriorityIndex + IndexOffset;
+
+					OutlinerPinned->RemoveColumn(FallbackColumn);
+
+					FName NameId = UE::EditorDataStorage::TableViewerUtils::FindLongestMatchingName(ColumnTypes, IndexOffset);
+					AddedColumns.Add(NameId);
+					OutlinerPinned->AddColumn(NameId,
+						FSceneOutlinerColumnInfo(ESceneOutlinerColumnVisibility::Visible, ColumnPriority,
+							FCreateSceneOutlinerColumn::CreateLambda(
+								[this, Query, MetaDataView, NameId, &ColumnTypes, CellConstructor, &OutlinerPinned, FallbackColumn](ISceneOutliner&)
+								{
+									TSharedPtr<FTypedElementWidgetConstructor> HeaderConstructor = 
+										UE::EditorDataStorage::TableViewerUtils::CreateHeaderWidgetConstructor(*StorageUi, MetaDataView, ColumnTypes, CellWidgetPurposes);
+									return MakeShared<FOutlinerColumn>(
+										Query, *Storage, *StorageUi, *StorageCompatibility, NameId,
+										TArray<TWeakObjectPtr<const UScriptStruct>>(ColumnTypes.GetData(), ColumnTypes.Num()), 
+										MoveTemp(HeaderConstructor), CellConstructor, FallbackColumn, OutlinerPinned, Dealiaser);
+
+								})
+						)
+					);
+					++IndexOffset;
+					return true;
+				};
+
+			for(const FName& WidgetPurpose : CellWidgetPurposes)
+			{
+				StorageUi->CreateWidgetConstructors(WidgetPurpose, MatchApproach::LongestMatch, ColumnTypes, 
+				MetaDataView, ColumnConstructor);
+			}
+
+			for (TWeakObjectPtr<const UScriptStruct>& ColumnType : ColumnTypes)
+			{
+				FName FallbackColumn = Binder.FindOutlinerColumnFromTEDSColumns({ColumnType});
+
+				auto AssignWidgetToColumn = [this, Query, MetaDataView, ColumnType, &IndexOffset, &OutlinerPinned, FallbackColumn](
+					TUniquePtr<FTypedElementWidgetConstructor> Constructor, TConstArrayView<TWeakObjectPtr<const UScriptStruct>>)
+				{
+					TSharedPtr<FTypedElementWidgetConstructor> CellConstructor(Constructor.Release());
+					FName NameId = FName(ColumnType->GetDisplayNameText().ToString());
+					AddedColumns.Add(NameId);
+					OutlinerPinned->AddColumn(NameId,
+						FSceneOutlinerColumnInfo(ESceneOutlinerColumnVisibility::Visible, DefaultPriorityIndex + IndexOffset,
+							FCreateSceneOutlinerColumn::CreateLambda(
+								[this, Query, MetaDataView, NameId, ColumnType, CellConstructor, &OutlinerPinned, FallbackColumn](ISceneOutliner&)
+								{
+									TArray<TWeakObjectPtr<const UScriptStruct>> ColumnTypesStored;
+									ColumnTypesStored.Add(ColumnType);
+									TSharedPtr<FTypedElementWidgetConstructor> HeaderConstructor =
+										UE::EditorDataStorage::TableViewerUtils::CreateHeaderWidgetConstructor(*StorageUi, MetaDataView, { ColumnType }, CellWidgetPurposes);
+									return MakeShared<FOutlinerColumn>(
+										Query, *Storage, *StorageUi, *StorageCompatibility, NameId, MoveTemp(ColumnTypesStored),
+										HeaderConstructor, CellConstructor, FallbackColumn, OutlinerPinned, Dealiaser);
+
+								})
+						)
+					);
+					++IndexOffset;
+					return false;
+				};
+
+				const int32 BeforeIndexOffset = IndexOffset;
+				for(const FName& WidgetPurpose : CellWidgetPurposes)
+				{
+					const FName DefaultWidgetPurpose(WidgetPurpose.ToString() + TEXT(".Default"));
+
+					StorageUi->CreateWidgetConstructors(DefaultWidgetPurpose, MetaDataView, AssignWidgetToColumn);
+					
+					if (BeforeIndexOffset != IndexOffset)
+					{
+						break;
+					}
+				}
+
+				if (BeforeIndexOffset == IndexOffset)
+				{
+					++IndexOffset;
+				}
+			}
+		}
+	}
+}
+
+void FSceneOutlinerTedsBridge::ClearColumns(ISceneOutliner& InOutliner)
+{
+	for (FName ColumnName : AddedColumns)
+	{
+		InOutliner.RemoveColumn(ColumnName);
+	}
+}
+
+#undef LOCTEXT_NAMESPACE
diff --git a/Engine/Plugins/Experimental/EditorDataStorageFeatures/Source/TedsOutliner/Private/Compatibility/TedsCompatibilityUtils.cpp b/Engine/Plugins/Experimental/EditorDataStorageFeatures/Source/TedsOutliner/Private/Compatibility/TedsCompatibilityUtils.cpp
deleted file mode 100644
diff --git a/Engine/Plugins/Experimental/EditorDataStorageFeatures/Source/TedsOutliner/Private/TedsOutlinerFilter.cpp b/Engine/Plugins/Experimental/EditorDataStorageFeatures/Source/TedsOutliner/Private/TedsOutlinerFilter.cpp
new file mode 100644
--- /dev/null
+++ b/Engine/Plugins/Experimental/EditorDataStorageFeatures/Source/TedsOutliner/Private/TedsOutlinerFilter.cpp
@@ -0,0 +1,87 @@
+// Copyright Epic Games, Inc. All Rights Reserved.
+
+#include "TedsOutlinerFilter.h"
+
+#include "TedsOutlinerImpl.h"
+
+FTedsOutlinerFilter::FTedsOutlinerFilter(const FName& InFilterName, const FText& InFilterDisplayName,
+	TSharedPtr<FFilterCategory> InCategory, TSharedRef<FTedsOutlinerImpl> InTedsOutlinerImpl,
+	const TypedElementDataStorage::FQueryDescription& InFilterQuery)
+	: FFilterBase(InCategory)
+	, FilterName(InFilterName)
+	, FilterDisplayName(InFilterDisplayName)
+	, TedsOutlinerImpl(InTedsOutlinerImpl)
+	, FilterQuery(InFilterQuery)
+{
+	
+}
+
+FString FTedsOutlinerFilter::GetName() const
+{
+	return FilterName.ToString();
+}
+
+FText FTedsOutlinerFilter::GetDisplayName() const
+{
+	return FilterDisplayName;
+}
+
+FText FTedsOutlinerFilter::GetToolTipText() const
+{
+	return FText::FromName(FilterName);
+}
+
+FLinearColor FTedsOutlinerFilter::GetColor() const
+{
+	return FLinearColor();	
+}
+
+FName FTedsOutlinerFilter::GetIconName() const
+{
+	return FName();
+}
+
+bool FTedsOutlinerFilter::IsInverseFilter() const
+{
+	return false;
+}
+
+void FTedsOutlinerFilter::ActiveStateChanged(bool bActive)
+{
+	if(bActive)
+	{
+		TedsOutlinerImpl->AddExternalQuery(FilterName, FilterQuery);
+	}
+	else
+	{
+		TedsOutlinerImpl->RemoveExternalQuery(FilterName);
+	}
+}
+
+void FTedsOutlinerFilter::ModifyContextMenu(FMenuBuilder& MenuBuilder)
+{
+	
+}
+
+void FTedsOutlinerFilter::SaveSettings(const FString& IniFilename, const FString& IniSection, const FString& SettingsString) const
+{
+	
+}
+
+void FTedsOutlinerFilter::LoadSettings(const FString& IniFilename, const FString& IniSection, const FString& SettingsString)
+{
+	
+}
+
+bool FTedsOutlinerFilter::PassesFilter(SceneOutliner::FilterBarType InItem) const
+{
+	// If this item is not compatible with the owning Table Viewer - it does not pass any filter queries
+	// If it is compatible, this is simply a dummy filter for the UI while the actual filter is applied through the TEDS query
+	if(TedsOutlinerImpl->IsItemCompatible().IsBound())
+	{
+		return TedsOutlinerImpl->IsItemCompatible().Execute(InItem);
+	}
+
+	// The filter is applied through a TEDS query and this is just a dummy to activate it, so we can simply return true otherwise
+	return false;
+}
diff --git a/Engine/Plugins/Experimental/EditorDataStorageFeatures/Source/TedsOutliner/Private/TedsOutlinerHierarchy.cpp b/Engine/Plugins/Experimental/EditorDataStorageFeatures/Source/TedsOutliner/Private/TedsOutlinerHierarchy.cpp
new file mode 100644
--- /dev/null
+++ b/Engine/Plugins/Experimental/EditorDataStorageFeatures/Source/TedsOutliner/Private/TedsOutlinerHierarchy.cpp
@@ -0,0 +1,75 @@
+// Copyright Epic Games, Inc. All Rights Reserved.
+
+#include "TedsOutlinerHierarchy.h"
+
+#include "Elements/Framework/TypedElementQueryBuilder.h"
+#include "Elements/Interfaces/TypedElementDataStorageInterface.h"
+#include "Elements/Interfaces/TypedElementQueryStorageInterfaces.h"
+#include "TedsOutlinerImpl.h"
+#include "TedsOutlinerItem.h"
+#include "TedsOutlinerMode.h"
+#include "Elements/Columns/TypedElementCompatibilityColumns.h"
+#include "Elements/Columns/TypedElementHiearchyColumns.h"
+#include "Elements/Columns/TypedElementLabelColumns.h"
+#include "Elements/Columns/TypedElementMiscColumns.h"
+#include "Elements/Columns/TypedElementTypeInfoColumns.h"
+
+FTedsOutlinerHierarchy::FTedsOutlinerHierarchy(FTedsOutlinerMode* InMode,
+	const TSharedRef<FTedsOutlinerImpl>& InTedsOutlinerImpl)
+	: ISceneOutlinerHierarchy(InMode)
+	, TedsOutlinerImpl(InTedsOutlinerImpl)
+{
+	HierarchyChangedHandle = TedsOutlinerImpl->OnHierarchyChanged().AddLambda([this](FSceneOutlinerHierarchyChangedData EventData)
+	{
+		HierarchyChangedEvent.Broadcast(EventData);
+	});
+
+	TedsOutlinerImpl->RecompileQueries();
+}
+
+FTedsOutlinerHierarchy::~FTedsOutlinerHierarchy()
+{
+	TedsOutlinerImpl->OnHierarchyChanged().Remove(HierarchyChangedHandle);
+}
+
+void FTedsOutlinerHierarchy::CreateItems(TArray<FSceneOutlinerTreeItemPtr>& OutItems) const
+{
+	TedsOutlinerImpl->CreateItemsFromQuery(OutItems, Mode);
+}
+
+void FTedsOutlinerHierarchy::CreateChildren(const FSceneOutlinerTreeItemPtr& Item,
+	TArray<FSceneOutlinerTreeItemPtr>& OutChildren) const
+{
+	TedsOutlinerImpl->CreateChildren(Item, OutChildren);
+}
+
+FSceneOutlinerTreeItemPtr FTedsOutlinerHierarchy::FindOrCreateParentItem(const ISceneOutlinerTreeItem& Item,
+	const TMap<FSceneOutlinerTreeItemID, FSceneOutlinerTreeItemPtr>& Items, bool bCreate)
+{
+	const FTedsOutlinerTreeItem* TedsTreeItem = Item.CastTo<FTedsOutlinerTreeItem>();
+	const ITypedElementDataStorageInterface* Storage = TedsOutlinerImpl->GetStorage();
+	
+	// If this item is not a TEDS item, we are not handling it
+	if(!TedsTreeItem)
+	{
+		return nullptr;
+	}
+	
+	const TypedElementRowHandle ParentRowHandle = TedsOutlinerImpl->GetParentRow(TedsTreeItem->GetRowHandle());
+
+	if(!Storage->IsRowAvailable(ParentRowHandle))
+	{
+		return nullptr;
+	}
+	
+	if (const FSceneOutlinerTreeItemPtr* ParentItem = Items.Find(ParentRowHandle))
+	{
+		return *ParentItem;
+	}
+	else if(bCreate)
+	{
+		return Mode->CreateItemFor<FTedsOutlinerTreeItem>(FTedsOutlinerTreeItem(ParentRowHandle, TedsOutlinerImpl), true);
+	}
+	
+	return nullptr;
+}
\ No newline at end of file
diff --git a/Engine/Plugins/Experimental/EditorDataStorageFeatures/Source/TedsOutliner/Private/TedsOutlinerHierarchy.h b/Engine/Plugins/Experimental/EditorDataStorageFeatures/Source/TedsOutliner/Private/TedsOutlinerHierarchy.h
new file mode 100644
--- /dev/null
+++ b/Engine/Plugins/Experimental/EditorDataStorageFeatures/Source/TedsOutliner/Private/TedsOutlinerHierarchy.h
@@ -0,0 +1,38 @@
+// Copyright Epic Games, Inc. All Rights Reserved.
+
+#pragma once
+
+#include "ISceneOutlinerHierarchy.h"
+#include "Elements/Interfaces/TypedElementDataStorageInterface.h"
+
+class FTedsOutlinerMode;
+class FTedsOutlinerImpl;
+
+/*
+ * Class that keeps track of hierarchy data and creates items using the given TEDS queries.
+ * See CreateGenericTEDSOutliner()
+ * Inherits from ISceneOutlinerHierarchy, which is responsible for creating the items you want to populate the Outliner
+ * with and establishing hierarchical relationships between the items
+ */
+class FTedsOutlinerHierarchy : public ISceneOutlinerHierarchy
+{
+public:
+	FTedsOutlinerHierarchy(FTedsOutlinerMode* InMode, const TSharedRef<FTedsOutlinerImpl>& InTedsOutlinerImpl);
+	virtual ~FTedsOutlinerHierarchy();
+	
+	/** Create a linearization of all applicable items in the hierarchy */
+	virtual void CreateItems(TArray<FSceneOutlinerTreeItemPtr>& OutItems) const override;
+	/** Create a linearization of all direct and indirect children of a given item in the hierarchy */
+	virtual void CreateChildren(const FSceneOutlinerTreeItemPtr& Item, TArray<FSceneOutlinerTreeItemPtr>& OutChildren) const override;
+	/** Find or optionally create a parent item for a given tree item */
+	virtual FSceneOutlinerTreeItemPtr FindOrCreateParentItem(const ISceneOutlinerTreeItem& Item, const TMap<FSceneOutlinerTreeItemID,
+		FSceneOutlinerTreeItemPtr>& Items, bool bCreate = false) override;
+
+protected:
+
+	// The actual model for the TEDS-Outliner
+	TSharedRef<FTedsOutlinerImpl> TedsOutlinerImpl;
+
+	// Delegate called by TedsOutlinerImpl when the hierarchy changes
+	FDelegateHandle HierarchyChangedHandle;
+};
diff --git a/Engine/Plugins/Experimental/EditorDataStorageFeatures/Source/TedsOutliner/Private/TedsOutlinerImpl.cpp b/Engine/Plugins/Experimental/EditorDataStorageFeatures/Source/TedsOutliner/Private/TedsOutlinerImpl.cpp
new file mode 100644
--- /dev/null
+++ b/Engine/Plugins/Experimental/EditorDataStorageFeatures/Source/TedsOutliner/Private/TedsOutlinerImpl.cpp
@@ -0,0 +1,794 @@
+// Copyright Epic Games, Inc. All Rights Reserved.
+
+#include "TedsOutlinerImpl.h"
+
+#include "SSceneOutliner.h"
+#include "TedsTableViewerUtils.h"
+#include "Elements/Columns/TypedElementLabelColumns.h"
+#include "Elements/Columns/TypedElementMiscColumns.h"
+#include "Elements/Columns/TypedElementTypeInfoColumns.h"
+#include "Elements/Interfaces/TypedElementDataStorageInterface.h"
+#include "Elements/Interfaces/TypedElementDataStorageUiInterface.h"
+#include "Elements/Framework/TypedElementQueryBuilder.h"
+#include "Elements/Framework/TypedElementRegistry.h"
+#include "Compatibility/SceneOutlinerRowHandleColumn.h"
+#include "TedsOutlinerFilter.h"
+#include "TedsOutlinerItem.h"
+#include "Columns/TedsOutlinerColumns.h"
+#include "Elements/Columns/TypedElementSelectionColumns.h"
+#include "Elements/Columns/TypedElementSlateWidgetColumns.h"
+#include "Filters/FilterBase.h"
+#include "Widgets/SWidget.h"
+
+#define LOCTEXT_NAMESPACE "TedsOutliner"
+
+FTedsOutlinerImpl::FTedsOutlinerImpl(const FTedsOutlinerParams& InParams, ISceneOutlinerMode* InMode)
+	: CreationParams(InParams)
+	, CellWidgetPurposes(InParams.CellWidgetPurposes)
+	, InitialQueryDescription(InParams.QueryDescription)
+	, HierarchyData(InParams.HierarchyData)
+	, SelectionSetName(InParams.SelectionSetOverride)
+	, SceneOutlinerMode(InMode)
+	, SceneOutliner(InParams.SceneOutliner)
+{
+	// Initialize the TEDS constructs
+	UTypedElementRegistry* Registry = UTypedElementRegistry::GetInstance();
+	checkf(Registry, TEXT("Unable to initialize the Typed Elements Outliner before TEDS is initialized."));
+	if (Registry)
+	{
+		Storage = Registry->GetMutableDataStorage();
+		StorageUi = Registry->GetMutableDataStorageUi();
+		StorageCompatibility = Registry->GetMutableDataStorageCompatibility();
+	}
+
+}
+
+void FTedsOutlinerImpl::CreateLabelWidgetConstructors()
+{
+	using namespace TypedElementQueryBuilder;
+
+	// Create and assign the widget constructors for the label column
+	
+	auto CreateWidgetConstructorForQuery = [this](TypedElementDataStorage::FQueryDescription InQueryDescription) -> TSharedPtr<FTypedElementWidgetConstructor>
+	{
+		// Create the Widget Constructor for the item label column
+		using MatchApproach = ITypedElementDataStorageUiInterface::EMatchApproach;
+
+		// Make a copy of the columns because CreateWidgetConstructors can modify it
+		TArray<TWeakObjectPtr<const UScriptStruct>> ColumnTypes(InQueryDescription.SelectionTypes);
+
+		TSharedPtr<FTypedElementWidgetConstructor> OutWidgetConstructorPtr;
+
+		bool bFoundWidget = false;
+
+		// We also want to look at the ItemLabel purpose for the label
+		TArray<FName> ItemLabelCellWidgetPurposes{TEXT("SceneOutliner.ItemLabel.Cell")};
+		ItemLabelCellWidgetPurposes.Append(CellWidgetPurposes);
+		
+		for(const FName& WidgetPurpose : ItemLabelCellWidgetPurposes)
+		{
+			StorageUi->CreateWidgetConstructors(WidgetPurpose, MatchApproach::ExactMatch, ColumnTypes, {},
+				[&OutWidgetConstructorPtr, ColumnTypes, &bFoundWidget](
+				TUniquePtr<FTypedElementWidgetConstructor> CreatedConstructor, 
+				TConstArrayView<TWeakObjectPtr<const UScriptStruct>> MatchedColumnTypes)
+				{
+					if (ColumnTypes.Num() == MatchedColumnTypes.Num())
+					{
+						OutWidgetConstructorPtr = TSharedPtr<FTypedElementWidgetConstructor>(CreatedConstructor.Release());
+						bFoundWidget = true;
+					}
+					// Either this was the exact match so no need to search further or the longest possible chain didn't match so the next ones will 
+					// always be shorter in both cases just return.
+					return false;
+				});
+
+			if (bFoundWidget)
+			{
+				break;
+			}
+		}
+
+		return OutWidgetConstructorPtr;
+	};
+	
+
+	TypedElementDataStorage::QueryHandle TypeColumnQueryHandle = Storage->RegisterQuery(
+																	Select()
+																		.ReadOnly<FTypedElementClassTypeInfoColumn>()
+																	.Compile()
+																	);
+
+
+	if (TSharedPtr<FTypedElementWidgetConstructor> TypeColumnWidgetConstructor = CreateWidgetConstructorForQuery(Storage->GetQueryDescription(TypeColumnQueryHandle)))
+	{
+		QueryToWidgetConstructorMap.Emplace(TypeColumnQueryHandle, TypeColumnWidgetConstructor);
+	}
+
+	TypedElementDataStorage::QueryHandle LabelColumnQueryHandle = Storage->RegisterQuery(
+																	Select()
+																		.ReadWrite<FTypedElementLabelColumn>()
+																	.Compile()
+																	);
+
+	if (TSharedPtr<FTypedElementWidgetConstructor> LabelColumnWidgetConstructor = CreateWidgetConstructorForQuery(Storage->GetQueryDescription(LabelColumnQueryHandle)))
+	{
+		QueryToWidgetConstructorMap.Emplace(LabelColumnQueryHandle, LabelColumnWidgetConstructor);
+	}
+}
+
+void FTedsOutlinerImpl::CreateFilterQueries()
+{
+	using namespace TypedElementQueryBuilder;
+
+	if (CreationParams.bUseDefaultTedsFilters)
+	{
+		// Create separate categories for columns and tags
+		TSharedRef<FFilterCategory> TedsColumnFilterCategory = MakeShared<FFilterCategory>(LOCTEXT("TedsColumnFilters", "TEDS Columns"), LOCTEXT("TedsColumnFiltersTooltip", "Filter by TEDS columns"));
+		TSharedRef<FFilterCategory> TedsTagFilterCategory = MakeShared<FFilterCategory>(LOCTEXT("TedsTagFilters", "TEDS Tags"), LOCTEXT("TedsTagFiltersTooltip", "Filter by TEDS Tags"));
+
+		const UStruct* TedsColumn = FTypedElementDataStorageColumn::StaticStruct();
+		const UStruct* TedsTag = FTypedElementDataStorageTag::StaticStruct();
+
+		// Grab all UStruct types to see if they derive from FTypedElementDataStorageColumn or FTypedElementDataStorageTag
+		ForEachObjectOfClass(UScriptStruct::StaticClass(), [&](UObject* Obj)
+		{
+			if (UScriptStruct* Struct = Cast<UScriptStruct>(Obj))
+			{
+				if (Struct->IsChildOf(TedsColumn) || Struct->IsChildOf(TedsTag))
+				{
+					// Create a query description to filter for this tag/column
+					TypedElementDataStorage::FQueryDescription FilterQueryDesc =
+					Select()
+					.Where()
+						.All(Struct)
+					.Compile();
+
+					// Create the filter
+					TSharedRef<FTedsOutlinerFilter> TedsFilter = MakeShared<FTedsOutlinerFilter>(Struct->GetFName(), Struct->GetDisplayNameText(),
+						Struct->IsChildOf(TedsColumn) ? TedsColumnFilterCategory : TedsTagFilterCategory, AsShared(), FilterQueryDesc);
+					SceneOutliner->AddFilterToFilterBar(TedsFilter);
+				}
+			}
+		});
+	}
+
+	// Custom filters input by the user
+	TSharedRef<FFilterCategory> CustomFiltersCategory = MakeShared<FFilterCategory>(LOCTEXT("TedsFilters", "TEDS Custom Filters"), LOCTEXT("TedsFiltersTooltip", "Filter by custom TEDS queries"));
+
+	for(const TPair<FName, const TypedElementDataStorage::FQueryDescription>& FilterQuery : CreationParams.FilterQueries)
+	{
+		// TEDS-Outliner TODO: Custom filters need a localizable display name instead of using the FName, but we need to change how they are added first
+		// to see if it can be consolidated with the SFilterBar API
+		TSharedRef<FTedsOutlinerFilter> TedsFilter = MakeShared<FTedsOutlinerFilter>(FilterQuery.Key, FText::FromName(FilterQuery.Key), CustomFiltersCategory, AsShared(), FilterQuery.Value);
+		SceneOutliner->AddFilterToFilterBar(TedsFilter);
+	
+	}
+}
+
+void FTedsOutlinerImpl::Init()
+{
+	CreateLabelWidgetConstructors();
+	CreateFilterQueries();
+	
+	if (SelectionSetName.IsSet())
+	{
+		// Ticker for selection updates so we don't fire the delegate multiple times in one frame for multi select
+		TickerHandle = FTSTicker::GetCoreTicker().AddTicker(FTickerDelegate::CreateLambda(
+				[this](float DeltaTimeInSeconds)
+				{
+					Tick();
+					return true;
+				}));
+	}
+}
+
+FTedsOutlinerImpl::~FTedsOutlinerImpl()
+{
+	UnregisterQueries();
+
+	// This is done outside of unregister queries because that is called when the internal query changes (i.e filter) but we don't want to unregister the
+	// label widget queries until shutdown
+	for(const TPair<TypedElementDataStorage::QueryHandle, TSharedPtr<FTypedElementWidgetConstructor>>& QueryConstructorPair : QueryToWidgetConstructorMap)
+	{
+		Storage->UnregisterQuery(QueryConstructorPair.Key);
+	}
+	
+	FTSTicker::GetCoreTicker().RemoveTicker(TickerHandle);
+}
+
+FTedsOutlinerImpl::FIsItemCompatible& FTedsOutlinerImpl::IsItemCompatible()
+{
+	return IsItemCompatibleWithTeds;
+}
+
+void FTedsOutlinerImpl::SetSelection(const TArray<TypedElementDataStorage::RowHandle>& InSelectedRows)
+{
+	if (!SelectionSetName.IsSet())
+	{
+		return;
+	}
+	
+	ClearSelection();
+
+	for(TypedElementDataStorage::RowHandle Row : InSelectedRows)
+	{
+		Storage->AddColumn(Row, FTypedElementSelectionColumn{ .SelectionSet = SelectionSetName.GetValue() });
+	}
+}
+
+TSharedRef<SWidget> FTedsOutlinerImpl::CreateLabelWidgetForItem(TypedElementRowHandle InRowHandle) const
+{
+	auto CreateWidgetForQuery = [InRowHandle, this](const TPair<TypedElementDataStorage::QueryHandle, TSharedPtr<FTypedElementWidgetConstructor>>& QueryConstructorPair) -> TSharedPtr<SWidget>
+	{
+		TypedElementDataStorage::FQueryDescription QueryDescription = Storage->GetQueryDescription(QueryConstructorPair.Key);
+		
+		// Create a generic metadata view for the Type Widget
+		TypedElementDataStorage::FMetaData QueryWideMetaData;
+		QueryWideMetaData.AddImmutableData("TypeInfoWidget_bUseIcon", true);
+		TypedElementDataStorage::FGenericMetaDataView GenericMetaDataView(QueryWideMetaData);
+
+		// Create metadata for the query itself
+		TypedElementDataStorage::FQueryMetaDataView QueryMetaDataView(QueryDescription);
+
+		// Combine the two metadata
+		TypedElementDataStorage::FComboMetaDataView MetaDataArgs(GenericMetaDataView, QueryMetaDataView);
+
+		TArray<TWeakObjectPtr<const UScriptStruct>> ColumnTypes(QueryDescription.SelectionTypes);
+		TSharedPtr<FTypedElementWidgetConstructor> CellWidgetConstructor = QueryConstructorPair.Value;
+
+		TypedElementRowHandle UiRowHandle = Storage->AddRow(Storage->FindTable(UE::EditorDataStorage::TableViewerUtils::GetWidgetTableName()));
+
+		if (FTypedElementRowReferenceColumn* RowReference = Storage->GetColumn<FTypedElementRowReferenceColumn>(UiRowHandle))
+		{
+			RowReference->Row = InRowHandle;
+		}
+
+		Storage->AddColumn(UiRowHandle, FTedsOutlinerColumn{.Outliner = StaticCastSharedRef<ISceneOutliner>(SceneOutliner->AsShared())});
+		
+		return StorageUi->ConstructWidget(UiRowHandle, *CellWidgetConstructor, MetaDataArgs);
+	};
+	
+	TSharedRef<SHorizontalBox> CombinedWidget = SNew(SHorizontalBox);
+
+	for(const TPair<TypedElementDataStorage::QueryHandle, TSharedPtr<FTypedElementWidgetConstructor>>& QueryConstructorPair : QueryToWidgetConstructorMap)
+	{
+		if (TSharedPtr<SWidget> WidgetForQuery = CreateWidgetForQuery(QueryConstructorPair))
+		{
+			CombinedWidget->AddSlot()
+					.AutoWidth()
+					.HAlign(HAlign_Left)
+					.VAlign(VAlign_Center)
+					.Padding(2.0f, 0.0f, 4.0f, 0.0f)
+					[
+						WidgetForQuery.ToSharedRef()
+					];
+		}
+	}
+	
+	return CombinedWidget;
+}
+
+void FTedsOutlinerImpl::AppendQuery(TypedElementDataStorage::FQueryDescription& Query1, const TypedElementDataStorage::FQueryDescription& Query2)
+{
+	// TEDS-Outliner TODO: We simply discard duplicate types for now but we probably want a more robust system to detect duplicates and conflicting conditions
+	for(int32 i = 0; i < Query2.ConditionOperators.Num(); ++i)
+	{
+		// Make sure we don't add duplicate conditions
+		TypedElementDataStorage::FQueryDescription::FOperator* FoundCondition = Query1.ConditionOperators.FindByPredicate([&Query2, i](const TypedElementDataStorage::FQueryDescription::FOperator& Op)
+		{
+			return Op.Type == Query2.ConditionOperators[i].Type;
+		});
+
+		// We also can't have a duplicate selection type and condition
+		TWeakObjectPtr<const UScriptStruct>* FoundSelection = Query1.SelectionTypes.FindByPredicate([&Query2, i](const TWeakObjectPtr<const UScriptStruct>& Selection)
+		{
+			return Selection == Query2.ConditionOperators[i].Type;
+		});
+
+		if (!FoundCondition && !FoundSelection)
+		{
+			Query1.ConditionOperators.Add(Query2.ConditionOperators[i]);
+			Query1.ConditionTypes.Add(Query2.ConditionTypes[i]);
+		}
+	}
+}
+
+void FTedsOutlinerImpl::AddExternalQuery(FName QueryName, const TypedElementDataStorage::FQueryDescription& InQueryDescription)
+{
+	ExternalQueries.Emplace(QueryName, InQueryDescription);
+
+	RecompileQueries();
+}
+
+void FTedsOutlinerImpl::RemoveExternalQuery(FName QueryName)
+{
+	ExternalQueries.Remove(QueryName);
+}
+
+void FTedsOutlinerImpl::AppendExternalQueries(TypedElementDataStorage::FQueryDescription& OutQuery)
+{
+	for(const TPair<FName, TypedElementDataStorage::FQueryDescription>& ExternalQuery : ExternalQueries)
+	{
+		AppendQuery(OutQuery, ExternalQuery.Value);
+	}
+}
+
+bool FTedsOutlinerImpl::HasItemParentChanged(TypedElementDataStorage::RowHandle InRowHandle, TypedElementDataStorage::RowHandle ParentRowHandle) const
+{
+	const FSceneOutlinerTreeItemPtr Item = SceneOutliner->GetTreeItem(InRowHandle, true);
+
+	// If the item doesn't exist, it doesn't make sense to say its parent changed
+	if (!Item)
+	{
+		return false;
+	}
+	
+	const FSceneOutlinerTreeItemPtr ParentItem = Item->GetParent();
+
+	// If the item doesn't have a parent, but ParentRowHandle is valid: The item just got added a parent so we want to dirty it
+	if (!ParentItem)
+	{
+		return Storage->IsRowAvailable(ParentRowHandle);
+	}
+	
+	const FTedsOutlinerTreeItem* TedsParentItem = ParentItem->CastTo<FTedsOutlinerTreeItem>();
+
+	if (TedsParentItem)
+	{
+		// return true if the row handle of the parent item doesn't match what we are given, i.e the parent has changed
+		return TedsParentItem->GetRowHandle() != ParentRowHandle;
+	}
+
+	return false;
+}
+
+bool FTedsOutlinerImpl::CanDisplayRow(TypedElementDataStorage::RowHandle ItemRowHandle) const
+{
+	/*
+	 * Don't display widgets that are created for rows in this table viewer. Widgets are only created for rows that are currently visible, so if we
+	 * display the rows for them we are now adding/removing rows to the table viewer based on currently visible rows. But adding rows can cause
+	 * scrolling and change the currently visible rows which in turn again adds/removes widget rows. This chain keeps continuing which can cause
+	 * flickering/scrolling issues in the table viewer.
+	 */
+	if (Storage->HasColumns<FTypedElementSlateWidgetReferenceColumn>(ItemRowHandle))
+	{
+		// Check if this widget row belongs to the same table viewer it is being displayed in
+		if (const FTedsOutlinerColumn* TedsOutlinerColumn = Storage->GetColumn<FTedsOutlinerColumn>(ItemRowHandle))
+		{
+			if (const TSharedPtr<ISceneOutliner> TableViewer = TedsOutlinerColumn->Outliner.Pin())
+			{
+				return SceneOutliner != TableViewer.Get();
+			}
+		}
+	}
+	return true;
+}
+
+void FTedsOutlinerImpl::CreateItemsFromQuery(TArray<FSceneOutlinerTreeItemPtr>& OutItems, ISceneOutlinerMode* InMode) const
+{
+	using namespace TypedElementQueryBuilder;
+	using DSI = ITypedElementDataStorageInterface;
+	
+	TArray<TypedElementDataStorage::RowHandle> Rows;
+	
+	TypedElementDataStorage::DirectQueryCallback RowCollector = CreateDirectQueryCallbackBinding(
+		[&Rows](DSI::IDirectQueryContext& Context)
+		{
+			TConstArrayView<TypedElementRowHandle> ContextRows = Context.GetRowHandles();
+			Rows.Append(ContextRows);
+		});
+
+	Storage->RunQuery(RowHandleQuery, RowCollector);
+	
+	for (const TypedElementRowHandle& Row : Rows)
+	{
+		if (!CanDisplayRow(Row))
+		{
+			continue;
+		}
+		
+		if (FSceneOutlinerTreeItemPtr TreeItem = InMode->CreateItemFor<FTedsOutlinerTreeItem>(FTedsOutlinerTreeItem(Row, AsShared()), false))
+		{
+			OutItems.Add(TreeItem);
+		}
+	}
+
+}
+
+void FTedsOutlinerImpl::CreateChildren(const FSceneOutlinerTreeItemPtr& Item, TArray<FSceneOutlinerTreeItemPtr>& OutChildren) const
+{
+	/* TEDS-Outliner TODO: This can probably be improved or optimized in the future
+	 * 
+	 * TEDS currently only supports one way lookup for parents, so to get the children
+	 * for a given row we currently have to go through every row (that matches our populate query) with a parent column to check if the parent
+	 * is our row.
+	 * This has to be done recursively to grab our children, grandchildren and so on...
+	 */
+
+	// If there's no hierarchy data, there is no need to create children
+	if (!HierarchyData.IsSet())
+	{
+		return;
+	}
+
+	using namespace TypedElementQueryBuilder;
+	using DSI = ITypedElementDataStorageInterface;
+	
+	const FTedsOutlinerTreeItem* TedsTreeItem = Item->CastTo<FTedsOutlinerTreeItem>();
+
+	// If this item is not a TEDS item, we are not handling it
+	if (!TedsTreeItem)
+	{
+		return;
+	}
+		
+	TypedElementRowHandle ItemRowHandle = TedsTreeItem->GetRowHandle();
+
+	if(!Storage->IsRowAssigned(ItemRowHandle))
+	{
+		return;
+	}
+
+
+	TSet<TypedElementDataStorage::RowHandle> MatchedRowsWithParentColumn;
+	
+	// Collect all entities that are owned by our entity
+	TypedElementDataStorage::DirectQueryCallback ChildRowCollector = CreateDirectQueryCallbackBinding(
+	[&MatchedRowsWithParentColumn] (const DSI::IDirectQueryContext& Context)
+	{
+		MatchedRowsWithParentColumn.Append(Context.GetRowHandles());
+	});
+
+	Storage->RunQuery(ChildRowHandleQuery, ChildRowCollector);
+
+	TArray<TypedElementDataStorage::RowHandle> ChildItems;
+
+	// Recursively get the children for each entity
+	TFunction<void(TypedElementRowHandle)> GetChildrenRecursive = [&ChildItems, &MatchedRowsWithParentColumn, DataStorage = Storage, &GetChildrenRecursive, InHierarchyData = HierarchyData]
+	(TypedElementRowHandle EntityRowHandle) -> void
+	{
+		for(TypedElementRowHandle ChildEntityRowHandle : MatchedRowsWithParentColumn)
+		{
+			void* ParentColumnData = DataStorage->GetColumnData(ChildEntityRowHandle, InHierarchyData.GetValue().HierarchyColumn);
+
+			if (ensureMsgf(ParentColumnData, TEXT("We should always the a parent column since we only grabbed rows with those ")))
+			{
+				// Get the parent row handle
+				const TypedElementDataStorage::RowHandle ParentRowHandle = InHierarchyData.GetValue().GetParent.Execute(ParentColumnData);
+				
+				// Check if this entity is owned by the entity we are looking children for
+				if (ParentRowHandle == EntityRowHandle)
+				{
+					ChildItems.Add(ChildEntityRowHandle);
+
+					// Recursively look for children of this item
+					GetChildrenRecursive(ChildEntityRowHandle);
+				}
+
+			}
+			
+		}
+	};
+		
+	GetChildrenRecursive(ItemRowHandle);
+
+	// Actually create the items for the child entities 
+	for (TypedElementRowHandle ChildItemRowHandle : ChildItems)
+	{
+		if (!CanDisplayRow(ChildItemRowHandle))
+		{
+			continue;
+		}
+		
+		if (FSceneOutlinerTreeItemPtr ChildActorItem = SceneOutlinerMode->CreateItemFor<FTedsOutlinerTreeItem>(FTedsOutlinerTreeItem(ChildItemRowHandle, AsShared())))
+		{
+			OutChildren.Add(ChildActorItem);
+		}
+	}
+}
+
+TypedElementDataStorage::RowHandle FTedsOutlinerImpl::GetParentRow(TypedElementDataStorage::RowHandle InRowHandle)
+{
+	// No parent if there is no hierarchy data specified
+	if (!HierarchyData.IsSet())
+	{
+		return TypedElementDataStorage::InvalidRowHandle;
+	}
+	
+	// If this entity does not have a parent entity, return InvalidRowHandle
+	void* ParentColumnData = Storage->GetColumnData(InRowHandle, HierarchyData.GetValue().HierarchyColumn);
+	
+	if (!ParentColumnData)
+	{
+		return TypedElementDataStorage::InvalidRowHandle;
+	}
+
+	// If the parent is invalid for some reason, return InvalidRowHandle
+	const TypedElementRowHandle ParentRowHandle = HierarchyData.GetValue().GetParent.Execute(ParentColumnData);
+	
+	if (!Storage->IsRowAvailable(ParentRowHandle))
+	{
+		return TypedElementDataStorage::InvalidRowHandle;
+	}
+	
+	if (!CanDisplayRow(ParentRowHandle))
+	{
+		return TypedElementDataStorage::InvalidRowHandle;
+	}
+
+	return ParentRowHandle;
+}
+
+void FTedsOutlinerImpl::OnItemAdded(TypedElementDataStorage::RowHandle ItemRowHandle)
+{
+	if (!CanDisplayRow(ItemRowHandle))
+	{
+		return;
+	}
+	
+	FSceneOutlinerHierarchyChangedData EventData;
+	EventData.Type = FSceneOutlinerHierarchyChangedData::Added;
+	EventData.Items.Add(SceneOutlinerMode->CreateItemFor<FTedsOutlinerTreeItem>(FTedsOutlinerTreeItem(ItemRowHandle, AsShared())));
+	HierarchyChangedEvent.Broadcast(EventData);
+}
+
+void FTedsOutlinerImpl::OnItemRemoved(TypedElementDataStorage::RowHandle ItemRowHandle)
+{
+	FSceneOutlinerHierarchyChangedData EventData;
+	EventData.Type = FSceneOutlinerHierarchyChangedData::Removed;
+	EventData.ItemIDs.Add(ItemRowHandle);
+	HierarchyChangedEvent.Broadcast(EventData);
+}
+
+void FTedsOutlinerImpl::OnItemMoved(TypedElementDataStorage::RowHandle ItemRowHandle)
+{
+	if (!CanDisplayRow(ItemRowHandle))
+	{
+		return;
+	}
+	
+	FSceneOutlinerHierarchyChangedData EventData;
+	EventData.Type = FSceneOutlinerHierarchyChangedData::Moved;
+	EventData.ItemIDs.Add(ItemRowHandle);
+	HierarchyChangedEvent.Broadcast(EventData);
+}
+
+void FTedsOutlinerImpl::RecompileQueries()
+{
+	using namespace TypedElementQueryBuilder;
+	using namespace TypedElementDataStorage;
+
+	UnregisterQueries();
+
+	if (!InitialQueryDescription.IsSet())
+	{
+		return;
+	}
+
+	// Our final query to collect rows to populate the Outliner - currently the same as the initial query the user provided
+	FQueryDescription FinalQueryDescription(InitialQueryDescription.Get());
+
+	// Add the filters the user has active to the query
+	AppendExternalQueries(FinalQueryDescription);
+
+	// Row to track addition of rows to the Outliner
+	FQueryDescription RowAdditionQueryDescription =
+		Select(
+				TEXT("Add Row to Outliner"),
+				FObserver::OnAdd<FTypedElementLabelColumn>().ForceToGameThread(true),
+				[this](IQueryContext& Context, TypedElementRowHandle Row)
+				{
+					OnItemAdded(Row);
+				})
+			.Compile();
+
+	// Add the conditions from FinalQueryDescription to ensure we are tracking addition of the rows the user requested
+	AppendQuery(RowAdditionQueryDescription, FinalQueryDescription);
+	
+	// Row to track removal of rows from the Outliner
+	FQueryDescription RowRemovalQueryDescription =
+		Select(
+				TEXT("Remove Row from Outliner"),
+				FObserver::OnRemove<FTypedElementLabelColumn>().ForceToGameThread(true),
+				[this](IQueryContext& Context, TypedElementRowHandle Row)
+				{
+					OnItemRemoved(Row);
+				})
+			.Compile();
+
+	// Add the conditions from FinalQueryDescription to ensure we are tracking removal of the rows the user requested
+	AppendQuery(RowRemovalQueryDescription, FinalQueryDescription);
+
+	// Queries to track parent info, only required if we have hierarchy data
+	if (HierarchyData.IsSet())
+	{
+		const UScriptStruct* ParentColumnType = HierarchyData.GetValue().HierarchyColumn;
+		
+		// Query to get all rows that match our conditions with a parent column (i.e all child rows)
+		FQueryDescription ChildHandleQueryDescription =
+							Select()
+							.Where()
+								.All(ParentColumnType)
+							.Compile();
+
+		// Add the conditions from FinalQueryDescription to ensure we are tracking removal of the rows the user requested
+		AppendQuery(ChildHandleQueryDescription, FinalQueryDescription);
+
+		FQueryDescription UpdateParentQueryDescription =
+			Select(
+			TEXT("Update item parent"),
+			FProcessor(EQueryTickPhase::DuringPhysics, Storage->GetQueryTickGroupName(EQueryTickGroups::Update))
+				.ForceToGameThread(true),
+			[this](IQueryContext& Context, TypedElementDataStorage::RowHandle Row)
+			{
+				TypedElementDataStorage::RowHandle ParentRowHandle = InvalidRowHandle;
+
+				if (const FTableRowParentColumn* ParentColumn = Context.GetColumn<FTableRowParentColumn>())
+				{
+					ParentRowHandle = ParentColumn->Parent;
+				}
+				
+				if (HasItemParentChanged(Row, ParentRowHandle))
+				{
+					OnItemMoved(Row);
+				}
+			})
+			.ReadOnly<FTableRowParentColumn>(EOptional::Yes)
+		.Where()
+			.All<FTypedElementSyncFromWorldTag>()
+		.Compile();
+		
+		// Add the conditions from FinalQueryDescription to ensure we are the rows the user requested
+		AppendQuery(UpdateParentQueryDescription, FinalQueryDescription);
+
+		ChildRowHandleQuery = Storage->RegisterQuery(MoveTemp(ChildHandleQueryDescription));
+		UpdateParentQuery = Storage->RegisterQuery(MoveTemp(UpdateParentQueryDescription));
+	}
+
+	if (SelectionSetName.IsSet())
+	{
+		// Query to grab all selected rows
+		FQueryDescription SelectedRowsQueryDescription =
+						Select()
+							.Where()
+								.All<FTypedElementSelectionColumn>()
+							.Compile();
+	
+		// Query to track when a row gets selected
+		FQueryDescription SelectionAddedQueryDescription =
+							Select(
+							TEXT("Row selected"),
+							FObserver::OnAdd<FTypedElementSelectionColumn>().ForceToGameThread(true),
+							[this](IQueryContext& Context, TypedElementRowHandle Row)
+							{
+								bSelectionDirty = true;
+							})
+							.Compile();
+
+		// Add the conditions from FinalQueryDescription to ensure we are the rows the user requested
+		AppendQuery(SelectionAddedQueryDescription, FinalQueryDescription);
+
+		// Query to track when a row gets deselected
+		FQueryDescription SelectionRemovedQueryDescription =
+							Select(
+							TEXT("Row deselected"),
+							FObserver::OnRemove<FTypedElementSelectionColumn>().ForceToGameThread(true),
+							[this](IQueryContext& Context, TypedElementRowHandle Row)
+							{
+								bSelectionDirty = true;
+							})
+							.Compile();
+
+		// Add the conditions from FinalQueryDescription to ensure we are the rows the user requested
+		AppendQuery(SelectionRemovedQueryDescription, FinalQueryDescription);
+
+		SelectedRowsQuery = Storage->RegisterQuery(MoveTemp(SelectedRowsQueryDescription));
+		SelectionAddedQuery = Storage->RegisterQuery(MoveTemp(SelectionAddedQueryDescription));
+		SelectionRemovedQuery = Storage->RegisterQuery(MoveTemp(SelectionRemovedQueryDescription));
+	}
+	
+	RowHandleQuery = Storage->RegisterQuery(MoveTemp(FinalQueryDescription));
+	RowAdditionQuery = Storage->RegisterQuery(MoveTemp(RowAdditionQueryDescription));
+	RowRemovalQuery = Storage->RegisterQuery(MoveTemp(RowRemovalQueryDescription));
+
+}
+
+void FTedsOutlinerImpl::UnregisterQueries() const
+{
+	if (Storage)
+	{
+		Storage->UnregisterQuery(RowHandleQuery);
+		Storage->UnregisterQuery(RowAdditionQuery);
+		Storage->UnregisterQuery(RowRemovalQuery);
+		Storage->UnregisterQuery(ChildRowHandleQuery);
+		Storage->UnregisterQuery(UpdateParentQuery);
+		Storage->UnregisterQuery(SelectedRowsQuery);
+		Storage->UnregisterQuery(SelectionAddedQuery);
+		Storage->UnregisterQuery(SelectionRemovedQuery);
+	}
+}
+
+void FTedsOutlinerImpl::ClearSelection() const
+{
+	if (!SelectionSetName.IsSet())
+	{
+		return;
+	}
+	
+	using namespace TypedElementQueryBuilder;
+	using DSI = ITypedElementDataStorageInterface;
+
+	TArray<TypedElementDataStorage::RowHandle> RowsToRemoveSelectionColumn;
+
+	// Query to remove the selection column from all rows that belong to this selection set
+	TypedElementDataStorage::DirectQueryCallback RowCollector = CreateDirectQueryCallbackBinding(
+	[this, &RowsToRemoveSelectionColumn](const DSI::IDirectQueryContext& Context, const TypedElementDataStorage::RowHandle* RowHandles)
+	{
+		const TConstArrayView<TypedElementDataStorage::RowHandle> Rows(RowHandles, Context.GetRowCount());
+
+		for(const TypedElementDataStorage::RowHandle RowHandle : Rows)
+		{
+			if (const FTypedElementSelectionColumn* SelectionColumn = Storage->GetColumn<FTypedElementSelectionColumn>(RowHandle))
+			{
+				if (SelectionColumn->SelectionSet == SelectionSetName)
+				{
+					RowsToRemoveSelectionColumn.Add(RowHandle);
+				}
+			}
+		}
+	});
+
+	Storage->RunQuery(SelectedRowsQuery, RowCollector);
+
+	for(const TypedElementDataStorage::RowHandle RowHandle : RowsToRemoveSelectionColumn)
+	{
+		Storage->RemoveColumn<FTypedElementSelectionColumn>(RowHandle);
+	}
+
+}
+
+void FTedsOutlinerImpl::Tick()
+{
+	if (bSelectionDirty)
+	{
+		OnTedsOutlinerSelectionChanged.Broadcast();
+		bSelectionDirty = false;
+	}
+}
+
+ITypedElementDataStorageInterface* FTedsOutlinerImpl::GetStorage() const
+{
+	return Storage;
+}
+
+ITypedElementDataStorageUiInterface* FTedsOutlinerImpl::GetStorageUI() const
+{
+	return StorageUi;
+}
+
+ITypedElementDataStorageCompatibilityInterface* FTedsOutlinerImpl::GetStorageCompatibility() const
+{
+	return StorageCompatibility;
+}
+
+TOptional<FName> FTedsOutlinerImpl::GetSelectionSetName() const
+{
+	return SelectionSetName;
+}
+
+FTedsOutlinerImpl::FOnTedsOutlinerSelectionChanged& FTedsOutlinerImpl::OnSelectionChanged()
+{
+	return OnTedsOutlinerSelectionChanged;
+}
+
+ISceneOutlinerHierarchy::FHierarchyChangedEvent& FTedsOutlinerImpl::OnHierarchyChanged()
+{
+	return HierarchyChangedEvent;
+}
+
+const TOptional<FTedsOutlinerHierarchyData>& FTedsOutlinerImpl::GetHierarchyData()
+{
+	return HierarchyData;
+}
+
+#undef LOCTEXT_NAMESPACE
diff --git a/Engine/Plugins/Experimental/EditorDataStorageFeatures/Source/TedsOutliner/Private/TedsOutlinerItem.cpp b/Engine/Plugins/Experimental/EditorDataStorageFeatures/Source/TedsOutliner/Private/TedsOutlinerItem.cpp
new file mode 100644
--- /dev/null
+++ b/Engine/Plugins/Experimental/EditorDataStorageFeatures/Source/TedsOutliner/Private/TedsOutlinerItem.cpp
@@ -0,0 +1,73 @@
+// Copyright Epic Games, Inc. All Rights Reserved.
+
+#include "TedsOutlinerItem.h"
+
+#include "Elements/Columns/TypedElementLabelColumns.h"
+#include "HAL/PlatformApplicationMisc.h"
+#include "TedsOutlinerImpl.h"
+
+#define LOCTEXT_NAMESPACE "TedsOutliner"
+
+const FSceneOutlinerTreeItemType FTedsOutlinerTreeItem::Type(&ISceneOutlinerTreeItem::Type);
+
+FTedsOutlinerTreeItem::FTedsOutlinerTreeItem(const TypedElementRowHandle& InRowHandle,
+	const TSharedRef<const FTedsOutlinerImpl>& InTedsOutlinerImpl)
+	: ISceneOutlinerTreeItem(Type)
+	, RowHandle(InRowHandle)
+	, TedsOutlinerImpl(InTedsOutlinerImpl)
+{
+	
+}
+
+bool FTedsOutlinerTreeItem::IsValid() const
+{
+	return true; // TEDS-Outliner TODO: check with TEDS if the item is valid?
+}
+
+FSceneOutlinerTreeItemID FTedsOutlinerTreeItem::GetID() const
+{
+	return FSceneOutlinerTreeItemID(RowHandle);
+}
+
+FString FTedsOutlinerTreeItem::GetDisplayString() const
+{
+	return TEXT("TEDS Item"); // TEDS-Outliner TODO: Used for searching by name, how to get this from TEDS
+}
+
+bool FTedsOutlinerTreeItem::CanInteract() const
+{
+	return true; // TEDS-Outliner TODO: check item constness from TEDS maybe?
+}
+
+TSharedRef<SWidget> FTedsOutlinerTreeItem::GenerateLabelWidget(ISceneOutliner& Outliner,
+	const STableRow<FSceneOutlinerTreeItemPtr>& InRow)
+{
+	return TedsOutlinerImpl->CreateLabelWidgetForItem(RowHandle);
+}
+
+void FTedsOutlinerTreeItem::GenerateContextMenu(UToolMenu* Menu, SSceneOutliner& Outliner)
+{
+	FToolMenuSection& Section = Menu->AddSection("Copy", LOCTEXT("CopySection", "Copy"));
+
+	Section.AddMenuEntry(
+		"CopyRowHandle",
+		LOCTEXT("CopyRowHandle_Title", "Copy row handle"),
+		LOCTEXT("CopyRowHandle_Tooltip", "Copy the row handle of this row to the clipboard."),
+		FSlateIcon(),
+		FUIAction(
+			FExecuteAction::CreateLambda([this]()
+			{
+				const FString ClipboardString(LexToString<FString>(RowHandle));
+				FPlatformApplicationMisc::ClipboardCopy(*ClipboardString);
+			}),
+			FCanExecuteAction()
+		)
+	);
+}
+
+TypedElementRowHandle FTedsOutlinerTreeItem::GetRowHandle() const
+{
+	return RowHandle;
+}
+
+#undef LOCTEXT_NAMESPACE
\ No newline at end of file
diff --git a/Engine/Plugins/Experimental/EditorDataStorageFeatures/Source/TedsOutliner/Private/TedsOutlinerMode.cpp b/Engine/Plugins/Experimental/EditorDataStorageFeatures/Source/TedsOutliner/Private/TedsOutlinerMode.cpp
new file mode 100644
--- /dev/null
+++ b/Engine/Plugins/Experimental/EditorDataStorageFeatures/Source/TedsOutliner/Private/TedsOutlinerMode.cpp
@@ -0,0 +1,302 @@
+// Copyright Epic Games, Inc. All Rights Reserved.
+
+#include "TedsOutlinerMode.h"
+
+#include "TedsOutlinerFilter.h"
+#include "Elements/Framework/TypedElementRegistry.h"
+#include "TedsOutlinerHierarchy.h"
+#include "TedsOutlinerItem.h"
+#include "Elements/Columns/TypedElementCompatibilityColumns.h"
+#include "Elements/Columns/TypedElementHiearchyColumns.h"
+#include "Elements/Columns/TypedElementMiscColumns.h"
+#include "Elements/Columns/TypedElementSelectionColumns.h"
+#include "FolderTreeItem.h"
+#include "Elements/Columns/TypedElementLabelColumns.h"
+#include "Elements/Columns/TypedElementTypeInfoColumns.h"
+#include "ToolMenus.h"
+
+#define LOCTEXT_NAMESPACE "TedsOutlinerMode"
+
+namespace UE::EditorDataStorage::Outliner::Private
+{
+	// Drag drop currently disabled as we are missing data marshaling for hierarchies from TEDS to the world
+	static bool TedsOutlinerDragDropEnabled = false;
+	static FAutoConsoleVariableRef TedsOutlinerDragDropEnabledCvar(TEXT("TEDS.UI.EnableTEDSOutlinerDragDrop"), TedsOutlinerDragDropEnabled, TEXT("Enable drag/drop for the generic TEDS Outliner."));
+	static FName ContextMenuName("TedsOutlinerContextMenu");
+}
+
+FTedsOutlinerMode::FTedsOutlinerMode(const FTedsOutlinerParams& InParams)
+	: ISceneOutlinerMode(InParams.SceneOutliner)
+{
+	using namespace TypedElementQueryBuilder;
+	
+	TedsOutlinerImpl = MakeShared<FTedsOutlinerImpl>(InParams, this);
+	TedsOutlinerImpl->Init();
+	TedsOutlinerImpl->OnSelectionChanged().AddRaw(this, &FTedsOutlinerMode::OnSelectionChanged);
+	
+	TedsOutlinerImpl->IsItemCompatible().BindLambda([](const ISceneOutlinerTreeItem& Item)
+	{
+		return Item.IsA<FTedsOutlinerTreeItem>();
+	});
+}
+
+FTedsOutlinerMode::~FTedsOutlinerMode()
+{
+	
+}
+
+void FTedsOutlinerMode::Rebuild()
+{
+	Hierarchy = CreateHierarchy();
+}
+
+void FTedsOutlinerMode::OnSelectionChanged()
+{
+	TOptional<FName> SelectionSetName = TedsOutlinerImpl->GetSelectionSetName();
+	ITypedElementDataStorageInterface* Storage = TedsOutlinerImpl->GetStorage();
+	
+	// The selection in TEDS was changed, update the outliner to respond
+	SceneOutliner->SetSelection([SelectionSetName, Storage](ISceneOutlinerTreeItem& InItem) -> bool
+	{
+		if(const FTedsOutlinerTreeItem* TedsItem = InItem.CastTo<FTedsOutlinerTreeItem>())
+		{
+			const TypedElementDataStorage::RowHandle RowHandle = TedsItem->GetRowHandle();
+
+			if(const FTypedElementSelectionColumn* SelectionColumn = Storage->GetColumn<FTypedElementSelectionColumn>(RowHandle))
+			{
+				return SelectionColumn->SelectionSet == SelectionSetName;
+			}
+		}
+		return false;
+	});
+}
+
+void FTedsOutlinerMode::SynchronizeSelection()
+{
+	OnSelectionChanged();
+}
+
+void FTedsOutlinerMode::OnItemSelectionChanged(FSceneOutlinerTreeItemPtr Item, ESelectInfo::Type SelectionType, const FSceneOutlinerItemSelection& Selection)
+{
+	if(SelectionType == ESelectInfo::Direct)
+	{
+		return; // Direct selection means we selected from outside the Outliner i.e through TEDS, so we don't need to redo the column addition
+	}
+
+	TArray<TypedElementDataStorage::RowHandle> RowHandles;
+	
+	// The selection in the Outliner changed, update TEDS
+	Selection.ForEachItem([&RowHandles](FSceneOutlinerTreeItemPtr& Item)
+	{
+		if(FTedsOutlinerTreeItem* TedsItem = Item->CastTo<FTedsOutlinerTreeItem>())
+		{
+			RowHandles.Add(TedsItem->GetRowHandle());
+		}
+	});
+
+	TedsOutlinerImpl->SetSelection(RowHandles);
+}
+
+TSharedPtr<FDragDropOperation> FTedsOutlinerMode::CreateDragDropOperation(const FPointerEvent& MouseEvent, const TArray<FSceneOutlinerTreeItemPtr>& InTreeItems) const
+{
+	const TOptional<FTedsOutlinerHierarchyData>& HierarchyData = TedsOutlinerImpl->GetHierarchyData();
+
+	// We don't want drag/drop if this TEDS Outliner isn't showing any hierarchy data
+	if(!HierarchyData.IsSet())
+	{
+		return nullptr;
+	}
+	
+	TArray<TypedElementDataStorage::RowHandle> DraggedRowHandles;
+
+	for(const FSceneOutlinerTreeItemPtr& Item :InTreeItems)
+	{
+		const FTedsOutlinerTreeItem* TedsItem = Item->CastTo<FTedsOutlinerTreeItem>();
+		if(ensureMsgf(TedsItem, TEXT("We should only have TEDS items in the TEDS Outliner")))
+		{
+			DraggedRowHandles.Add(TedsItem->GetRowHandle());
+		}
+	}
+
+	return FTedsRowDragDropOp::New(DraggedRowHandles);
+}
+
+bool FTedsOutlinerMode::ParseDragDrop(FSceneOutlinerDragDropPayload& OutPayload, const FDragDropOperation& Operation) const
+{
+	if (Operation.IsOfType<FTedsRowDragDropOp>())
+	{
+		const FTedsRowDragDropOp& TedsOp = static_cast<const FTedsRowDragDropOp&>(Operation);
+
+		for(TypedElementDataStorage::RowHandle RowHandle : TedsOp.DraggedRows)
+		{
+			OutPayload.DraggedItems.Add(SceneOutliner->GetTreeItem(RowHandle));
+		}
+		return true;
+	}
+	return false;
+}
+
+FSceneOutlinerDragValidationInfo FTedsOutlinerMode::ValidateDrop(const ISceneOutlinerTreeItem& DropTarget, const FSceneOutlinerDragDropPayload& Payload) const
+{
+	const TOptional<FTedsOutlinerHierarchyData>& HierarchyData = TedsOutlinerImpl->GetHierarchyData();
+	ITypedElementDataStorageInterface* Storage = TedsOutlinerImpl->GetStorage();
+
+	// We don't want drag/drop if this TEDS Outliner isn't showing any hierarchy data
+	if(!HierarchyData.IsSet())
+	{
+		return FSceneOutlinerDragValidationInfo(ESceneOutlinerDropCompatibility::IncompatibleGeneric,
+			LOCTEXT("DropDisabled", "Drag/Drop is disabled due to missing hierarchy data!"));
+	}
+	
+	TArray<TypedElementDataStorage::RowHandle> DraggedRowHandles;
+
+	Payload.ForEachItem<FTedsOutlinerTreeItem>([&DraggedRowHandles](FTedsOutlinerTreeItem& TedsItem)
+		{
+			DraggedRowHandles.Add(TedsItem.GetRowHandle());
+		});
+
+	// Dropping onto another item
+	// TEDS-Outliner TODO: Need better drag/drop validation and better place for this, TEDS-Outliner does not know about what types these rows are and all types that exist and what attachment is valid
+	if(const FTedsOutlinerTreeItem* TedsItem = DropTarget.CastTo<FTedsOutlinerTreeItem>())
+	{
+		TypedElementDataStorage::RowHandle DropTargetRowHandle = TedsItem->GetRowHandle();
+
+		FTypedElementClassTypeInfoColumn* DropTargetTypeInfoColumn = Storage->GetColumn<FTypedElementClassTypeInfoColumn>(DropTargetRowHandle);
+
+		// For now only allow attachment to same type
+		if(!DropTargetTypeInfoColumn || !DropTargetTypeInfoColumn->TypeInfo.Get())
+		{
+			return FSceneOutlinerDragValidationInfo(ESceneOutlinerDropCompatibility::IncompatibleGeneric, LOCTEXT("DropTargetInvalidType", "Invalid Drop target"));
+		}
+
+		
+		// TEDS-Outliner TODO: Currently we detect parent changes by removing the column and then adding the column back with the new parent 
+		for(TypedElementDataStorage::RowHandle RowHandle : DraggedRowHandles)
+		{
+			FTypedElementClassTypeInfoColumn* TypeInfoColumn = Storage->GetColumn<FTypedElementClassTypeInfoColumn>(RowHandle);
+
+			if(!TypeInfoColumn || !TypeInfoColumn->TypeInfo.Get())
+			{
+				return FSceneOutlinerDragValidationInfo(ESceneOutlinerDropCompatibility::IncompatibleGeneric, LOCTEXT("DragItemInvalidType", "Invalid Drag item"));
+			}
+
+			// TEDS-Outliner TOOD UE-205438: Proper drag/drop validation
+			if(TypeInfoColumn->TypeInfo.Get() != DropTargetTypeInfoColumn->TypeInfo.Get())
+			{
+				return FSceneOutlinerDragValidationInfo(ESceneOutlinerDropCompatibility::IncompatibleGeneric, FText::Format(LOCTEXT("DragDropTypeMismatch", "Cannot drag a {0} into a {1}"), FText::FromName(TypeInfoColumn->TypeInfo.Get()->GetFName()), FText::FromName(DropTargetTypeInfoColumn->TypeInfo.Get()->GetFName())));
+			}
+		}
+
+		return FSceneOutlinerDragValidationInfo(ESceneOutlinerDropCompatibility::CompatibleAttach, LOCTEXT("ValidDrop", "Valid Drop"));
+	}
+	// Dropping onto root, remove parent
+	else if (const FFolderTreeItem* FolderItem = DropTarget.CastTo<FFolderTreeItem>())
+	{
+		const FFolder DestinationPath = FolderItem->GetFolder();
+
+		if(DestinationPath.IsNone())
+		{
+			bool bValidDetach = false;
+			
+			for(TypedElementDataStorage::RowHandle RowHandle : DraggedRowHandles)
+			{
+				if(Storage->HasColumns(RowHandle, MakeArrayView({HierarchyData.GetValue().HierarchyColumn})))
+				{
+					bValidDetach = true;
+					break;
+				}
+			}
+
+			if(bValidDetach)
+			{
+				return FSceneOutlinerDragValidationInfo(ESceneOutlinerDropCompatibility::CompatibleDetach, LOCTEXT("MoveToRoot", "Move to root"));
+			}
+		}
+	}
+
+	return FSceneOutlinerDragValidationInfo(ESceneOutlinerDropCompatibility::IncompatibleGeneric, LOCTEXT("InvalidDrop", "Invalid Drop target"));
+}
+
+void FTedsOutlinerMode::OnDrop(ISceneOutlinerTreeItem& DropTarget, const FSceneOutlinerDragDropPayload& Payload, const FSceneOutlinerDragValidationInfo& ValidationInfo) const
+{
+	const TOptional<FTedsOutlinerHierarchyData>& HierarchyData = TedsOutlinerImpl->GetHierarchyData();
+	ITypedElementDataStorageInterface* Storage = TedsOutlinerImpl->GetStorage();
+
+	if(!UE::EditorDataStorage::Outliner::Private::TedsOutlinerDragDropEnabledCvar->GetBool() || !HierarchyData.IsSet())
+	{
+		return;
+	}
+	
+	TArray<TypedElementDataStorage::RowHandle> DraggedRowHandles;
+
+	Payload.ForEachItem<FTedsOutlinerTreeItem>([&DraggedRowHandles](FTedsOutlinerTreeItem& TedsItem)
+	{
+		DraggedRowHandles.Add(TedsItem.GetRowHandle());
+	});
+
+	if(ValidationInfo.CompatibilityType == ESceneOutlinerDropCompatibility::CompatibleDetach)
+	{
+		for(TypedElementDataStorage::RowHandle RowHandle : DraggedRowHandles)
+		{
+			Storage->RemoveColumn(RowHandle, HierarchyData.GetValue().HierarchyColumn);
+			Storage->AddColumn<FTypedElementSyncBackToWorldTag>(RowHandle);
+		}
+	}
+	
+	if(const FTedsOutlinerTreeItem* TedsItem = DropTarget.CastTo<FTedsOutlinerTreeItem>())
+	{
+		TypedElementDataStorage::RowHandle DropTargetRowHandle = TedsItem->GetRowHandle();
+		
+		for(TypedElementDataStorage::RowHandle RowHandle : DraggedRowHandles)
+		{
+			// Add the column
+			Storage->AddColumn(RowHandle, HierarchyData.GetValue().HierarchyColumn);
+
+			// Let the hierarchy data fill in the parent row into the column
+			HierarchyData.GetValue().SetParent.Execute(Storage->GetColumnData(RowHandle, HierarchyData.GetValue().HierarchyColumn), DropTargetRowHandle);
+			
+			Storage->AddColumn<FTypedElementSyncBackToWorldTag>(RowHandle);
+		}
+	}
+}
+
+TSharedPtr<SWidget> FTedsOutlinerMode::CreateContextMenu()
+{
+	UToolMenus* ToolMenus = UToolMenus::Get();
+
+	if (!ToolMenus->IsMenuRegistered(UE::EditorDataStorage::Outliner::Private::ContextMenuName))
+	{
+		UToolMenu* Menu = ToolMenus->RegisterMenu((UE::EditorDataStorage::Outliner::Private::ContextMenuName));
+		Menu->AddDynamicSection("DynamicHierarchySection", FNewToolMenuDelegate::CreateLambda([](UToolMenu* InMenu)
+		{
+			if(UTedsOutlinerMenuContext* TedsOutlinerMenuContext = InMenu->FindContext<UTedsOutlinerMenuContext>())
+			{
+				if(SSceneOutliner* SceneOutliner = TedsOutlinerMenuContext->OwningSceneOutliner)
+				{
+					TArray<FSceneOutlinerTreeItemPtr> Selection = SceneOutliner->GetTree().GetSelectedItems();
+
+					if(Selection.Num() == 1)
+					{
+						Selection[0]->GenerateContextMenu(InMenu, *SceneOutliner);
+					}
+				}
+			}
+			
+		}));
+	}
+
+	UTedsOutlinerMenuContext* TedsOutlinerMenuContext = NewObject<UTedsOutlinerMenuContext>();
+	TedsOutlinerMenuContext->OwningSceneOutliner = SceneOutliner;
+	
+	FToolMenuContext MenuContext;
+	MenuContext.AddObject(TedsOutlinerMenuContext);
+
+	return UToolMenus::Get()->GenerateWidget(UE::EditorDataStorage::Outliner::Private::ContextMenuName, MenuContext);
+}
+
+TUniquePtr<ISceneOutlinerHierarchy> FTedsOutlinerMode::CreateHierarchy()
+{
+	return MakeUnique<FTedsOutlinerHierarchy>(this, TedsOutlinerImpl.ToSharedRef());
+}
+
+#undef LOCTEXT_NAMESPACE
diff --git a/Engine/Plugins/Experimental/EditorDataStorageFeatures/Source/TedsOutliner/Private/TedsOutlinerModule.cpp b/Engine/Plugins/Experimental/EditorDataStorageFeatures/Source/TedsOutliner/Private/TedsOutlinerModule.cpp
--- a/Engine/Plugins/Experimental/EditorDataStorageFeatures/Source/TedsOutliner/Private/TedsOutlinerModule.cpp
+++ b/Engine/Plugins/Experimental/EditorDataStorageFeatures/Source/TedsOutliner/Private/TedsOutlinerModule.cpp
@@ -4,8 +4,7 @@
 
 #include "LevelEditor.h"
 #include "SceneOutlinerPublicTypes.h"
-#include "TypedElementOutlinerColumnIntegration.h"
-#include "TypedElementOutlinerMode.h"
+#include "TedsOutlinerMode.h"
 #include "WorkspaceMenuStructureModule.h"
 #include "WorkspaceMenuStructure.h"
 #include "Elements/Columns/TypedElementAlertColumns.h"
@@ -16,14 +15,15 @@
 #include "Elements/Columns/TypedElementTypeInfoColumns.h"
 #include "Elements/Framework/TypedElementRegistry.h"
 #include "Modules/ModuleManager.h"
-#include "Widgets/TedsOutlinerRowHandleColumn.h"
+#include "Compatibility/SceneOutlinerTedsBridge.h"
+#include "Compatibility/SceneOutlinerRowHandleColumn.h"
 #include "Widgets/Docking/SDockTab.h"
 
 #define LOCTEXT_NAMESPACE "TedsOutlinerModule"
 
-namespace UE::TedsOutliner
+namespace UE::EditorDataStorage::Outliner::Private
 {
-	static bool bUseNewSCCWidgets = false;
+	static bool bUseNewRevisionControlWidgets = false;
 
 	void RefreshLevelEditorTedsOutliner(bool bAlwaysInvoke)
 	{
@@ -38,9 +38,9 @@ namespace UE::TedsOutliner
 		}
 	}
 	
-	static FAutoConsoleVariableRef CVarUseNewSCCWidgets(
-		TEXT("TEDS.UI.UseNewSCCWidgets"),
-		UE::TedsOutliner::bUseNewSCCWidgets,
+	static FAutoConsoleVariableRef CVarUseNewRevisionControlWidgets(
+		TEXT("TEDS.UI.UseNewRevisionControlWidgets"),
+		bUseNewRevisionControlWidgets,
 		TEXT("Use new TEDS-based source control widgets in the Outliner (requires TEDS-Outliner to be enabled)")
 		, FConsoleVariableDelegate::CreateLambda([](IConsoleVariable*)
 		{
@@ -61,29 +61,29 @@ FTedsOutlinerModule::FTedsOutlinerModule()
 {
 }
 
-TSharedRef<ISceneOutliner> FTedsOutlinerModule::CreateTedsOutliner(const FSceneOutlinerInitializationOptions& InInitOptions, const FTypedElementOutlinerModeParams& InInitTedsOptions, TypedElementDataStorage::QueryHandle ColumnQuery) const
+TSharedRef<ISceneOutliner> FTedsOutlinerModule::CreateTedsOutliner(const FSceneOutlinerInitializationOptions& InInitOptions, const FTedsOutlinerParams& InInitTedsOptions, TypedElementDataStorage::QueryHandle ColumnQuery) const
 {
 	UTypedElementRegistry* Registry = UTypedElementRegistry::GetInstance();
 
 	ensureMsgf(Registry&& Registry->AreDataStorageInterfacesSet(), TEXT("Unable to initialize the Teds-Outliner before TEDS itself is initialized."));
 
 	FSceneOutlinerInitializationOptions InitOptions(InInitOptions);
-	FTypedElementOutlinerModeParams InitTedsOptions(InInitTedsOptions);
+	FTedsOutlinerParams InitTedsOptions(InInitTedsOptions);
 
 	InitOptions.ModeFactory = FCreateSceneOutlinerMode::CreateLambda([&InitTedsOptions](SSceneOutliner* Outliner)
 	{
 		using namespace TypedElementQueryBuilder;
 		InitTedsOptions.SceneOutliner = Outliner;
 		
-		return new FTypedElementOutlinerMode(InitTedsOptions);
+		return new FTedsOutlinerMode(InitTedsOptions);
 	});
 
 	// Add the custom column that displays row handles
-	InitOptions.ColumnMap.Add(FTedsOutlinerRowHandleColumn::GetID(),
+	InitOptions.ColumnMap.Add(FSceneOutlinerRowHandleColumn::GetID(),
 		FSceneOutlinerColumnInfo(ESceneOutlinerColumnVisibility::Visible, 2,
 			FCreateSceneOutlinerColumn::CreateLambda([](ISceneOutliner& InSceneOutliner)
 			{
-				return MakeShareable(new FTedsOutlinerRowHandleColumn(InSceneOutliner));
+				return MakeShareable(new FSceneOutlinerRowHandleColumn(InSceneOutliner));
 			})));
 	
 	InitOptions.ColumnMap.Add(FSceneOutlinerBuiltInColumnTypes::Label(), FSceneOutlinerColumnInfo(ESceneOutlinerColumnVisibility::Visible, 10));
@@ -91,7 +91,7 @@ TSharedRef<ISceneOutliner> FTedsOutlinerModule::CreateTedsOutliner(const FSceneO
 	
 	TSharedRef<ISceneOutliner> TedsOutlinerShared = SNew(SSceneOutliner, InitOptions);
 	
-	FTypedElementSceneOutlinerQueryBinder::GetInstance().AssignQuery(ColumnQuery, TedsOutlinerShared, InitTedsOptions.CellWidgetPurposes);
+	FSceneOutlinerTedsQueryBinder::GetInstance().AssignQuery(ColumnQuery, TedsOutlinerShared, InitTedsOptions.CellWidgetPurposes);
 	
 	return TedsOutlinerShared;
 }
@@ -122,13 +122,13 @@ TypedElementDataStorage::QueryHandle FTedsOutlinerModule::GetLevelEditorTedsOutl
 			.ReadOnly<FTypedElementClassTypeInfoColumn, FTypedElementAlertColumn, FTypedElementChildAlertColumn>()
 		.Compile());
 
-	// Query to also include SCC info
-	static TypedElementDataStorage::QueryHandle SCCQuery = Storage->RegisterQuery(
+	// Query to also include revision control info
+	static TypedElementDataStorage::QueryHandle RevisionControlQuery = Storage->RegisterQuery(
 		Select()
 			.ReadOnly<FTypedElementClassTypeInfoColumn, FTypedElementPackageReference, FTypedElementAlertColumn>()
 		.Compile());
 
-	return UE::TedsOutliner::bUseNewSCCWidgets ? SCCQuery : ColumnQuery;
+	return UE::EditorDataStorage::Outliner::Private::bUseNewRevisionControlWidgets ? RevisionControlQuery : ColumnQuery;
 	
 }
 
@@ -158,7 +158,7 @@ TSharedRef<SWidget> FTedsOutlinerModule::CreateLevelEditorTedsOutliner()
 	InitOptions.bShowTransient = true;
 	InitOptions.OutlinerIdentifier = "TEDSOutliner";
 
-	FTypedElementOutlinerModeParams Params(nullptr);
+	FTedsOutlinerParams Params(nullptr);
 	Params.QueryDescription = OutlinerQueryDescription;
 	Params.bUseDefaultTedsFilters = true;
 
diff --git a/Engine/Plugins/Experimental/EditorDataStorageFeatures/Source/TedsOutliner/Private/TypedElementOutlinerColumnIntegration.cpp b/Engine/Plugins/Experimental/EditorDataStorageFeatures/Source/TedsOutliner/Private/TypedElementOutlinerColumnIntegration.cpp
deleted file mode 100644
diff --git a/Engine/Plugins/Experimental/EditorDataStorageFeatures/Source/TedsOutliner/Private/TypedElementOutlinerFilter.cpp b/Engine/Plugins/Experimental/EditorDataStorageFeatures/Source/TedsOutliner/Private/TypedElementOutlinerFilter.cpp
deleted file mode 100644
diff --git a/Engine/Plugins/Experimental/EditorDataStorageFeatures/Source/TedsOutliner/Private/TypedElementOutlinerHierarchy.cpp b/Engine/Plugins/Experimental/EditorDataStorageFeatures/Source/TedsOutliner/Private/TypedElementOutlinerHierarchy.cpp
deleted file mode 100644
diff --git a/Engine/Plugins/Experimental/EditorDataStorageFeatures/Source/TedsOutliner/Private/TypedElementOutlinerHierarchy.h b/Engine/Plugins/Experimental/EditorDataStorageFeatures/Source/TedsOutliner/Private/TypedElementOutlinerHierarchy.h
deleted file mode 100644
diff --git a/Engine/Plugins/Experimental/EditorDataStorageFeatures/Source/TedsOutliner/Private/TypedElementOutlinerItem.cpp b/Engine/Plugins/Experimental/EditorDataStorageFeatures/Source/TedsOutliner/Private/TypedElementOutlinerItem.cpp
deleted file mode 100644
diff --git a/Engine/Plugins/Experimental/EditorDataStorageFeatures/Source/TedsOutliner/Private/TypedElementOutlinerMode.cpp b/Engine/Plugins/Experimental/EditorDataStorageFeatures/Source/TedsOutliner/Private/TypedElementOutlinerMode.cpp
deleted file mode 100644
diff --git a/Engine/Plugins/Experimental/EditorDataStorageFeatures/Source/TedsOutliner/Private/Widgets/AlertWidget.cpp b/Engine/Plugins/Experimental/EditorDataStorageFeatures/Source/TedsOutliner/Private/Widgets/AlertWidget.cpp
new file mode 100644
--- /dev/null
+++ b/Engine/Plugins/Experimental/EditorDataStorageFeatures/Source/TedsOutliner/Private/Widgets/AlertWidget.cpp
@@ -0,0 +1,430 @@
+// Copyright Epic Games, Inc. All Rights Reserved.
+
+#include "Widgets/AlertWidget.h"
+
+#include "Columns/UIPropertiesColumns.h"
+#include "Elements/Columns/TypedElementAlertColumns.h"
+#include "Elements/Columns/TypedElementMiscColumns.h"
+#include "Elements/Columns/TypedElementSlateWidgetColumns.h"
+#include "Elements/Framework/TypedElementRegistry.h"
+#include "Widgets/Images/SImage.h"
+#include "Widgets/Input/SButton.h"
+#include "Widgets/SOverlay.h"
+#include "Widgets/Text/STextBlock.h"
+
+#define LOCTEXT_NAMESPACE "TedsAlertWidget"
+
+namespace UE::EditorDataStorage::Widgets::Private
+{
+	void UpdateWidget(const TSharedPtr<SWidget>& Widget, const FText& Alert, bool bIsWarning, uint16 ErrorCount, uint16 WarningCount,
+		TypedElementDataStorage::RowHandle RowWithAlertAction)
+	{
+		if (Widget)
+		{
+			uint32 ChildCount = ErrorCount + WarningCount;
+			if (FChildren* Children = Widget->GetChildren())
+			{
+				SImage& Background = 
+					static_cast<SImage&>(*Children->GetSlotAt(FAlertWidgetConstructor::IconBackgroundSlot).GetWidget());
+				SImage& Badge = 
+					static_cast<SImage&>(*Children->GetSlotAt(FAlertWidgetConstructor::IconBadgeSlot).GetWidget());
+				STextBlock& CounterText =
+					static_cast<STextBlock&>(*Children->GetSlotAt(FAlertWidgetConstructor::CounterTextSlot).GetWidget());
+				SButton& ActionButton =
+					static_cast<SButton&>(*Children->GetSlotAt(FAlertWidgetConstructor::ActionButtonSlot).GetWidget());
+				
+				// Setup the background image
+				if (ChildCount > 0)
+				{
+					if (!Alert.IsEmpty())
+					{
+						Background.SetImage(bIsWarning
+							? FAppStyle::GetBrush("Icons.Warning.Background")
+							: FAppStyle::GetBrush("Icons.Error.Background"));
+					}
+					else
+					{
+						Background.SetImage(FAppStyle::GetBrush("Icons.Alert.Background"));
+					}
+				}
+				else
+				{
+					if (!Alert.IsEmpty())
+					{
+						Background.SetImage(bIsWarning
+							? FAppStyle::GetBrush("Icons.Warning.Solid")
+							: FAppStyle::GetBrush("Icons.Error.Solid"));
+					}
+					else
+					{
+						Background.SetImage(FAppStyle::GetBrush("Icons.Alert.Solid"));
+					}
+				}
+
+				// Set counter if needed, otherwise turn it off.
+				if (ChildCount == 0)
+				{
+					// If there are no children, don't show the badge and don't show a counter.
+					Badge.SetVisibility(EVisibility::Hidden);
+					CounterText.SetVisibility(EVisibility::Hidden);
+				}
+				else
+				{
+					// If there are children, also take into account if there's an alert as well.
+					uint32 TotalChildCount = ChildCount + (Alert.IsEmpty() ? 0 : 1);
+					if (TotalChildCount <= 9)
+					{
+						Badge.SetVisibility(EVisibility::HitTestInvisible);
+						CounterText.SetVisibility(EVisibility::HitTestInvisible);
+						CounterText.SetText(FText::AsNumber(TotalChildCount));
+						CounterText.SetFont(FCoreStyle::GetDefaultFontStyle("Regular", FAlertWidgetConstructor::BadgeFontSize));
+						CounterText.SetMargin(FMargin(
+							FAlertWidgetConstructor::BadgeHorizontalOffset, 
+							FAlertWidgetConstructor::BadgeVerticalOffset));
+					}
+					else
+					{
+						Badge.SetVisibility(EVisibility::HitTestInvisible);
+						CounterText.SetVisibility(EVisibility::HitTestInvisible);
+						CounterText.SetText(FText::FromString(TEXT("*")));
+						CounterText.SetFont(FCoreStyle::GetDefaultFontStyle("Regular", 14));
+						CounterText.SetMargin(FMargin(
+							FAlertWidgetConstructor::BadgeHorizontalOffset - 2.0f,
+							FAlertWidgetConstructor::BadgeVerticalOffset - 6.5f));
+					}
+				}
+
+				// Setup the tool tip text
+				if (!Alert.IsEmpty() && ChildCount > 0)
+				{
+					FText ToolTipText = FText::Format(LOCTEXT("ChildAlertCountWithMessage", "Errors: {0}\nWarnings: {1}\n\n{2}"),
+						FText::AsNumber(ErrorCount + (bIsWarning ? 0 : 1)),
+						FText::AsNumber(WarningCount + (bIsWarning ? 1 : 0)), Alert);
+					Background.SetToolTipText(ToolTipText);
+					ActionButton.SetToolTipText(MoveTemp(ToolTipText));
+				}
+				else if (!Alert.IsEmpty())
+				{
+					Background.SetToolTipText(Alert);
+					ActionButton.SetToolTipText(Alert);
+				}
+				else if (ChildCount > 0)
+				{
+					FText ToolTipText = FText::Format(LOCTEXT("ChildAlertCount", "Errors: {0}\nWarnings: {1}"),
+						FText::AsNumber(ErrorCount), FText::AsNumber(WarningCount));
+					Background.SetToolTipText(ToolTipText);
+					ActionButton.SetToolTipText(MoveTemp(ToolTipText));
+				}
+
+				// If there's an action to call, enable the invisible button, otherwise turn it off.
+				if (RowWithAlertAction != TypedElementDataStorage::InvalidRowHandle)
+				{
+					Background.SetVisibility(EVisibility::HitTestInvisible);
+					ActionButton.SetVisibility(EVisibility::Visible);
+					ActionButton.SetOnClicked(FOnClicked::CreateLambda([RowWithAlertAction]()
+						{
+							UTypedElementRegistry* Registry = UTypedElementRegistry::GetInstance();
+							const ITypedElementDataStorageInterface* DataStorage = Registry->GetDataStorage();
+							if (const FTypedElementAlertActionColumn* Action =
+								DataStorage->GetColumn<FTypedElementAlertActionColumn>(RowWithAlertAction))
+							{
+								Action->Action(RowWithAlertAction);
+							}
+							return FReply::Handled();
+						}));
+				}
+				else
+				{
+					Background.SetVisibility(EVisibility::Visible);
+					ActionButton.SetVisibility(EVisibility::Hidden);
+				}
+			}
+		}
+	}
+}
+
+//
+// UAlertWidgetFactory
+//
+
+void UAlertWidgetFactory::RegisterWidgetConstructors(
+	ITypedElementDataStorageInterface& DataStorage,
+	ITypedElementDataStorageUiInterface& DataStorageUi) const
+{
+	using namespace TypedElementDataStorage;
+
+	DataStorageUi.RegisterWidgetFactory<FAlertWidgetConstructor>(FName(TEXT("General.Cell")),
+		FColumn<FTypedElementAlertColumn>() || FColumn<FTypedElementChildAlertColumn>());
+	
+	DataStorageUi.RegisterWidgetFactory<FAlertHeaderWidgetConstructor>(FName(TEXT("General.Header")),
+		FColumn<FTypedElementAlertColumn>() || FColumn<FTypedElementChildAlertColumn>());
+}
+
+void UAlertWidgetFactory::RegisterQueries(ITypedElementDataStorageInterface& DataStorage)
+{
+	RegisterAlertQueries(DataStorage);
+	RegisterAlertHeaderQueries(DataStorage);
+}
+
+void UAlertWidgetFactory::RegisterAlertQueries(ITypedElementDataStorageInterface& DataStorage)
+{
+	using namespace TypedElementQueryBuilder;
+	using namespace TypedElementDataStorage;
+	using DSI = ITypedElementDataStorageInterface;
+	
+	TypedElementQueryHandle UpdateWidget_OnlyAlert = DataStorage.RegisterQuery(
+		Select()
+			.ReadOnly<FTypedElementAlertColumn>()
+		.Where()
+			.Any<FTypedElementSyncFromWorldTag, FTypedElementSyncBackToWorldTag>()
+			.None<FTypedElementChildAlertColumn>()
+		.Compile());
+
+	TypedElementQueryHandle UpdateWidget_OnlyChildAlert = DataStorage.RegisterQuery(
+		Select()
+			.ReadOnly<FTypedElementChildAlertColumn>()
+		.Where()
+			.Any<FTypedElementSyncFromWorldTag, FTypedElementSyncBackToWorldTag>()
+			.None<FTypedElementAlertColumn>()
+		.Compile());
+
+	TypedElementQueryHandle UpdateWidget_Both = DataStorage.RegisterQuery(
+		Select()
+			.ReadOnly<FTypedElementAlertColumn, FTypedElementChildAlertColumn>()
+		.Where()
+			.Any<FTypedElementSyncFromWorldTag, FTypedElementSyncBackToWorldTag>()
+		.Compile());
+	
+	DataStorage.RegisterQuery(
+		Select(TEXT("Sync Transform column to heads up display"),
+			FProcessor(DSI::EQueryTickPhase::FrameEnd, DataStorage.GetQueryTickGroupName(DSI::EQueryTickGroups::SyncWidgets))
+			.ForceToGameThread(true),
+			[](IQueryContext& Context,
+				FTypedElementSlateWidgetReferenceColumn& Widget,
+				const FTypedElementRowReferenceColumn& ReferenceColumn)
+			{
+				Context.RunSubquery(0, ReferenceColumn.Row, CreateSubqueryCallbackBinding(
+					[&Widget](ISubqueryContext& Context, RowHandle Row, const FTypedElementAlertColumn& Alert)
+					{
+						checkf(
+							Alert.AlertType == FTypedElementAlertColumnType::Warning ||
+							Alert.AlertType == FTypedElementAlertColumnType::Error,
+							TEXT("Alert column has unsupported type %i"), static_cast<int>(Alert.AlertType));
+						UE::EditorDataStorage::Widgets::Private::UpdateWidget(Widget.Widget.Pin(), Alert.Message,
+							Alert.AlertType == FTypedElementAlertColumnType::Warning, 0, 0,
+							Context.HasColumn<FTypedElementAlertActionColumn>() ? Row : InvalidRowHandle);
+					}));
+				Context.RunSubquery(1, ReferenceColumn.Row, CreateSubqueryCallbackBinding(
+					[&Widget](ISubqueryContext& Context, RowHandle Row, const FTypedElementChildAlertColumn& ChildAlert)
+					{
+						UE::EditorDataStorage::Widgets::Private::UpdateWidget(Widget.Widget.Pin(), FText::GetEmpty(), false,
+							ChildAlert.Counts[static_cast<size_t>(FTypedElementAlertColumnType::Error)],
+							ChildAlert.Counts[static_cast<size_t>(FTypedElementAlertColumnType::Warning)],
+							Context.HasColumn<FTypedElementAlertActionColumn>() ? Row : InvalidRowHandle);
+					}));
+				Context.RunSubquery(2, ReferenceColumn.Row, CreateSubqueryCallbackBinding(
+					[&Widget](
+						ISubqueryContext& Context, RowHandle Row, 
+						const FTypedElementAlertColumn& Alert, 
+						const FTypedElementChildAlertColumn& ChildAlert)
+					{
+						checkf(
+							Alert.AlertType == FTypedElementAlertColumnType::Warning ||
+							Alert.AlertType == FTypedElementAlertColumnType::Error,
+							TEXT("Alert column has unsupported type %i"), static_cast<int>(Alert.AlertType));
+						UE::EditorDataStorage::Widgets::Private::UpdateWidget(Widget.Widget.Pin(), Alert.Message,
+							Alert.AlertType == FTypedElementAlertColumnType::Warning,
+							ChildAlert.Counts[static_cast<size_t>(FTypedElementAlertColumnType::Error)],
+							ChildAlert.Counts[static_cast<size_t>(FTypedElementAlertColumnType::Warning)],
+							Context.HasColumn<FTypedElementAlertActionColumn>() ? Row : InvalidRowHandle);
+					}));
+			}
+		)
+		.Where()
+			.All<FAlertWidgetTag>()
+			.DependsOn()
+				.SubQuery(UpdateWidget_OnlyAlert)
+				.SubQuery(UpdateWidget_OnlyChildAlert)
+				.SubQuery(UpdateWidget_Both)
+		.Compile());
+}
+
+void UAlertWidgetFactory::RegisterAlertHeaderQueries(ITypedElementDataStorageInterface& DataStorage)
+{
+	using namespace TypedElementQueryBuilder;
+	using namespace TypedElementDataStorage;
+	using DSI = ITypedElementDataStorageInterface;
+	
+	TypedElementQueryHandle AlertCount = DataStorage.RegisterQuery(
+		Count()
+		.Where()
+			.Any<FTypedElementAlertColumn>()
+		.Compile());
+	
+	DataStorage.RegisterQuery(
+		Select(TEXT("Update alert header"),
+			FProcessor(DSI::EQueryTickPhase::FrameEnd, DataStorage.GetQueryTickGroupName(DSI::EQueryTickGroups::SyncWidgets))
+			.ForceToGameThread(true),
+			[](IQueryContext& Context, RowHandle Row, FTypedElementSlateWidgetReferenceColumn& Widget)
+			{
+				FQueryResult Result = Context.RunSubquery(0);
+				if (Result.Count > 0)
+				{
+					if (TSharedPtr<SWidget> WidgetPtr = Widget.Widget.Pin())
+					{
+						static_cast<SImage*>(WidgetPtr.Get())->SetImage(FAppStyle::GetBrush("Icons.Warning.Solid"));
+						Context.AddColumns<FAlertHeaderActiveWidgetTag>(Row);
+					}
+				}
+			}
+		)
+		.Where()
+			.All<FAlertHeaderWidgetTag>()
+			.None<FAlertHeaderActiveWidgetTag>()
+		.DependsOn()
+			.SubQuery(AlertCount)
+		.Compile());
+
+	DataStorage.RegisterQuery(
+		Select(TEXT("Update active alert header"),
+			FProcessor(DSI::EQueryTickPhase::FrameEnd, DataStorage.GetQueryTickGroupName(DSI::EQueryTickGroups::SyncWidgets))
+			.ForceToGameThread(true),
+			[](IQueryContext& Context, RowHandle Row, FTypedElementSlateWidgetReferenceColumn& Widget)
+			{
+				FQueryResult Result = Context.RunSubquery(0);
+				if (Result.Count == 0)
+				{
+					if (TSharedPtr<SWidget> WidgetPtr = Widget.Widget.Pin())
+					{
+						static_cast<SImage*>(WidgetPtr.Get())->SetImage(FAppStyle::GetBrush("Icons.Alert"));
+						Context.RemoveColumns<FAlertHeaderActiveWidgetTag>(Row);
+					}
+				}
+			}
+		)
+		.Where()
+			.All<FAlertHeaderWidgetTag, FAlertHeaderActiveWidgetTag>()
+		.DependsOn()
+			.SubQuery(AlertCount)
+		.Compile());
+}
+
+
+
+//
+// FAlertWidgetConstructor
+//
+
+FAlertWidgetConstructor::FAlertWidgetConstructor()
+	: Super(StaticStruct())
+{
+}
+
+TSharedPtr<SWidget> FAlertWidgetConstructor::CreateWidget(
+	const TypedElementDataStorage::FMetaDataView& Arguments)
+{
+	return SNew(SOverlay)
+		+SOverlay::Slot()
+			.VAlign(VAlign_Fill)
+			.HAlign(HAlign_Fill)
+		[
+			SNew(SButton)
+				.ButtonStyle(FAppStyle::Get(), "FlatButton")
+				.Text(FText::FromString(TEXT("X"))) // There needs to be at least some content otherwise nothing will show.
+				.ForegroundColor(FLinearColor::Transparent) // Then the color needs to be cleared so the X doesn't show.
+				.ContentPadding(FMargin(0.0f))
+		]
+		+SOverlay::Slot()
+			.VAlign(VAlign_Center)
+			.HAlign(HAlign_Center)
+		[
+			SNew(SImage)
+				.Image(FAppStyle::GetBrush("Icons.Alert.Solid"))
+				.DesiredSizeOverride(FVector2D(16.0f, 16.0f))
+		]
+		+SOverlay::Slot()
+			.VAlign(VAlign_Center)
+			.HAlign(HAlign_Center)
+		[
+			SNew(SImage)
+				.Image(FAppStyle::GetBrush("Icons.Alert.Badge"))
+				.DesiredSizeOverride(FVector2D(16.0f, 16.0f))
+		]
+		+ SOverlay::Slot()
+			.VAlign(VAlign_Bottom)
+			.HAlign(HAlign_Center)
+		[
+			SNew(STextBlock)
+				.Font(FCoreStyle::GetDefaultFontStyle("Regular", BadgeFontSize))
+				.ColorAndOpacity(FLinearColor::Black)
+				.Margin(FMargin(BadgeHorizontalOffset, BadgeVerticalOffset))
+		];
+}
+
+TConstArrayView<const UScriptStruct*> FAlertWidgetConstructor::GetAdditionalColumnsList() const
+{
+	static TTypedElementColumnTypeList<FTypedElementRowReferenceColumn, FAlertWidgetTag> Columns;
+	return Columns;
+}
+
+bool FAlertWidgetConstructor::FinalizeWidget(ITypedElementDataStorageInterface* DataStorage,
+	ITypedElementDataStorageUiInterface* DataStorageUi, TypedElementDataStorage::RowHandle Row, const TSharedPtr<SWidget>& Widget)
+{
+	using namespace TypedElementDataStorage;
+
+	RowHandle TargetRow = DataStorage->GetColumn<FTypedElementRowReferenceColumn>(Row)->Row;
+	const FTypedElementAlertColumn* Alert = DataStorage->GetColumn<FTypedElementAlertColumn>(TargetRow);
+	const FTypedElementChildAlertColumn* ChildAlert = DataStorage->GetColumn<FTypedElementChildAlertColumn>(TargetRow);
+	
+	uint16 ErrorCount = ChildAlert ? ChildAlert->Counts[static_cast<size_t>(FTypedElementAlertColumnType::Error)] : 0;
+	uint16 WarningCount = ChildAlert ? ChildAlert->Counts[static_cast<size_t>(FTypedElementAlertColumnType::Warning)] : 0;
+
+	UE::EditorDataStorage::Widgets::Private::UpdateWidget(
+		Widget, 
+		Alert ? Alert->Message : FText::GetEmpty(), 
+		Alert ? (Alert->AlertType == FTypedElementAlertColumnType::Warning) : false, 
+		ErrorCount, WarningCount,
+		DataStorage->HasColumns<FTypedElementAlertActionColumn>(TargetRow) ? TargetRow : InvalidRowHandle);
+	
+	return true;
+}
+
+
+
+//
+// FAlertHeaderWidgetConstructor
+//
+
+FAlertHeaderWidgetConstructor::FAlertHeaderWidgetConstructor()
+	: Super(StaticStruct())
+{
+}
+
+TSharedPtr<SWidget> FAlertHeaderWidgetConstructor::CreateWidget(
+	const TypedElementDataStorage::FMetaDataView& Arguments)
+{
+	return SNew(SImage)
+		.DesiredSizeOverride(FVector2D(16.f, 16.f))
+		.ColorAndOpacity(FSlateColor::UseForeground())
+		.Image(FAppStyle::GetBrush("Icons.Alert"))
+		.ToolTipText(FText(LOCTEXT("AlertColumnHeader", "Alerts")));
+}
+
+TConstArrayView<const UScriptStruct*> FAlertHeaderWidgetConstructor::GetAdditionalColumnsList() const
+{
+	static TTypedElementColumnTypeList<FAlertHeaderWidgetTag> Columns;
+	return Columns;
+}
+
+bool FAlertHeaderWidgetConstructor::FinalizeWidget(ITypedElementDataStorageInterface* DataStorage, 
+	ITypedElementDataStorageUiInterface* DataStorageUi, TypedElementDataStorage::RowHandle Row, const TSharedPtr<SWidget>& Widget)
+{
+	DataStorage->AddColumn(Row, FUIHeaderPropertiesColumn
+		{
+			.ColumnSizeMode = EColumnSizeMode::Fixed,
+			.Width = 24.0f
+		});
+	return true;
+}
+
+#undef LOCTEXT_NAMESPACE
diff --git a/Engine/Plugins/Experimental/EditorDataStorageFeatures/Source/TedsOutliner/Private/Widgets/OverrideWidget.cpp b/Engine/Plugins/Experimental/EditorDataStorageFeatures/Source/TedsOutliner/Private/Widgets/OverrideWidget.cpp
new file mode 100644
--- /dev/null
+++ b/Engine/Plugins/Experimental/EditorDataStorageFeatures/Source/TedsOutliner/Private/Widgets/OverrideWidget.cpp
@@ -0,0 +1,182 @@
+// Copyright Epic Games, Inc. All Rights Reserved.
+
+#include "Widgets/OverrideWidget.h"
+
+#include "Editor.h"
+#include "Elements/Columns/TypedElementMiscColumns.h"
+#include "Elements/Columns/TypedElementOverrideColumns.h"
+#include "Elements/Columns/TypedElementSlateWidgetColumns.h"
+#include "Elements/Columns/TypedElementTypeInfoColumns.h"
+#include "Elements/Framework/TypedElementQueryBuilder.h"
+#include "Elements/Framework/TypedElementRegistry.h"
+#include "Compatibility/SceneOutlinerTedsBridge.h"
+#include "Widgets/Images/SLayeredImage.h"
+
+#define LOCTEXT_NAMESPACE "TedsOverrideWidget"
+
+//
+// UOverrideWidgetFactory
+//
+
+void UOverrideWidgetFactory::RegisterWidgetConstructors(
+	ITypedElementDataStorageInterface& DataStorage, ITypedElementDataStorageUiInterface& DataStorageUi) const
+{
+	using namespace TypedElementDataStorage;
+
+	// The  widget is a specific widget for the Scene Outliner's item label column
+	DataStorageUi.RegisterWidgetFactory<FOverrideWidgetConstructor>(FSceneOutlinerTedsQueryBinder::ItemLabelCellWidgetPurpose,
+		FColumn<FTypedElementClassTypeInfoColumn>() || FColumn<FObjectOverrideColumn>() );
+}
+
+void UOverrideWidgetFactory::RegisterQueries(ITypedElementDataStorageInterface& DataStorage)
+{
+	using namespace TypedElementQueryBuilder;
+	using DSI = ITypedElementDataStorageInterface;
+
+	const TypedElementQueryHandle UpdateWidget = DataStorage.RegisterQuery(
+		Select()
+		.Where()
+			.Any<FTypedElementSyncFromWorldTag, FTypedElementSyncBackToWorldTag>()
+			.All<FObjectOverrideColumn>()
+		.Compile());
+	
+	DataStorage.RegisterQuery(
+		Select(
+			TEXT("Sync override status to widget"),
+			FProcessor(DSI::EQueryTickPhase::FrameEnd, DataStorage.GetQueryTickGroupName(DSI::EQueryTickGroups::SyncWidgets))
+				.ForceToGameThread(true),
+			[](
+				DSI::IQueryContext& Context, 
+				const FTypedElementSlateWidgetReferenceColumn& Widget,
+				const FTypedElementRowReferenceColumn& Target)
+				{
+					FOverrideWidgetConstructor::UpdateOverrideWidget(Widget.Widget, Target.Row);
+				})
+		.Where()
+			.All<FOverrideWidgetTag>()
+		.DependsOn()
+			.SubQuery(UpdateWidget)
+		.Compile()
+		);
+}
+
+//
+// FOverrideWidgetConstructor
+//
+
+FOverrideWidgetConstructor::FOverrideWidgetConstructor()
+: Super(FOverrideWidgetConstructor::StaticStruct())
+{
+}
+
+TSharedPtr<SWidget> FOverrideWidgetConstructor::CreateWidget(const TypedElementDataStorage::FMetaDataView& Arguments)
+{
+	return SNew(SLayeredImage)
+			.DesiredSizeOverride(FVector2D(16.f, 16.f))
+			.ColorAndOpacity(FSlateColor::UseForeground());
+}
+
+bool FOverrideWidgetConstructor::FinalizeWidget(ITypedElementDataStorageInterface* DataStorage,
+	ITypedElementDataStorageUiInterface* DataStorageUi, TypedElementRowHandle Row, const TSharedPtr<SWidget>& Widget)
+{
+	const TypedElementRowHandle TargetRow = DataStorage->GetColumn<FTypedElementRowReferenceColumn>(Row)->Row;
+
+	if (const FTypedElementClassTypeInfoColumn* TypeInfoColumn = DataStorage->GetColumn<FTypedElementClassTypeInfoColumn>(TargetRow))
+	{
+		// Add the icon showing the type of the widget as the base image
+		SLayeredImage* WidgetInstance = static_cast<SLayeredImage*>(Widget.Get());
+		WidgetInstance->SetImage(GetIconForRow(DataStorage, TargetRow, TypeInfoColumn));
+
+		// Add override info if applicable
+		if(const FObjectOverrideColumn* OverrideColumn = DataStorage->GetColumn<FObjectOverrideColumn>(TargetRow))
+		{
+			AddOverrideBadge(Widget, OverrideColumn->OverriddenState);
+		}
+
+		return true;
+	}
+
+	return false;
+
+}
+
+TConstArrayView<const UScriptStruct*> FOverrideWidgetConstructor::GetAdditionalColumnsList() const
+{
+	static TTypedElementColumnTypeList<FOverrideWidgetTag> Columns;
+	return Columns;
+}
+
+void FOverrideWidgetConstructor::AddOverrideBadge(const TWeakPtr<SWidget>& Widget, EOverriddenState OverriddenState)
+{
+	if(const TSharedPtr<SWidget> WidgetPtr = Widget.Pin())
+	{
+		const TSharedPtr<SLayeredImage> WidgetInstance = StaticCastSharedPtr<SLayeredImage>(WidgetPtr);
+
+		WidgetInstance->RemoveAllLayers();
+
+		FText Tooltip;
+
+		switch(OverriddenState)
+		{
+		case EOverriddenState::Added:
+			WidgetInstance->AddLayer(FAppStyle::GetBrush("SceneOutliner.OverrideAddedBase"));
+			WidgetInstance->AddLayer(FAppStyle::GetBrush("SceneOutliner.OverrideAdded"));
+
+			Tooltip = LOCTEXT("OverrideAddedTooltip", "This entity has been added.");
+			break;
+			
+		case EOverriddenState::AllOverridden:
+			// Not implemented yet
+			break;
+			
+		case EOverriddenState::HasOverrides:
+			WidgetInstance->AddLayer(FAppStyle::GetBrush("SceneOutliner.OverrideInsideBase"));
+			WidgetInstance->AddLayer(FAppStyle::GetBrush("SceneOutliner.OverrideInside"));
+			
+			Tooltip = LOCTEXT("OverrideInsideTooltip", "At least one property or child has an override.");
+			break;
+			
+		case EOverriddenState::NoOverrides:
+			// No icon for no overrides
+			break;
+			
+		case EOverriddenState::SubObjectsHasOverrides:
+			// Not implemented yet
+			break;
+		}
+		
+		// We only add the tooltip if there already isn't one, because entities are currently dirtied almost every frame, which ends up calling this function
+		// and invalidating the tooltip every frame - leading to it never get drawn
+		if(!WidgetInstance->GetToolTip())
+		{
+			WidgetInstance->SetToolTipText(Tooltip);
+		}
+		
+	}
+}
+
+void FOverrideWidgetConstructor::RemoveOverrideBadge(const TWeakPtr<SWidget>& Widget)
+{
+	if(const TSharedPtr<SWidget> WidgetPtr = Widget.Pin())
+	{
+		const TSharedPtr<SLayeredImage> WidgetInstance = StaticCastSharedPtr<SLayeredImage>(WidgetPtr);
+		WidgetInstance->RemoveAllLayers();
+		WidgetInstance->SetToolTipText(FText::GetEmpty());
+	}
+}
+
+void FOverrideWidgetConstructor::UpdateOverrideWidget(const TWeakPtr<SWidget>& Widget, const TypedElementRowHandle TargetRow)
+{
+	const ITypedElementDataStorageInterface* DataStorageInterface = UTypedElementRegistry::GetInstance()->GetDataStorage();
+
+	if(const FObjectOverrideColumn* OverrideColumn = DataStorageInterface->GetColumn<FObjectOverrideColumn>(TargetRow))
+	{
+		AddOverrideBadge(Widget, OverrideColumn->OverriddenState);
+	}
+	else
+	{
+		RemoveOverrideBadge(Widget);
+	}
+}
+
+#undef LOCTEXT_NAMESPACE
diff --git a/Engine/Plugins/Experimental/EditorDataStorageFeatures/Source/TedsOutliner/Private/Widgets/OverrideWidget.h b/Engine/Plugins/Experimental/EditorDataStorageFeatures/Source/TedsOutliner/Private/Widgets/OverrideWidget.h
new file mode 100644
--- /dev/null
+++ b/Engine/Plugins/Experimental/EditorDataStorageFeatures/Source/TedsOutliner/Private/Widgets/OverrideWidget.h
@@ -0,0 +1,57 @@
+// Copyright Epic Games, Inc. All Rights Reserved.
+
+#pragma once
+
+#include "Elements/Interfaces/TypedElementDataStorageFactory.h"
+#include "Elements/Interfaces/TypedElementDataStorageInterface.h"
+#include "Elements/Interfaces/TypedElementDataStorageUiInterface.h"
+#include "UObject/OverridableManager.h"
+#include "Widgets/TypeInfoWidget.h"
+
+#include "OverrideWidget.generated.h"
+
+/*
+ * Widget for the Outliner that shows the icon showing the type of the object alongside the override status as a badg
+ */
+UCLASS()
+class UOverrideWidgetFactory : public UTypedElementDataStorageFactory
+{
+	GENERATED_BODY()
+
+public:
+	~UOverrideWidgetFactory() override = default;
+
+	TEDSOUTLINER_API void RegisterWidgetConstructors(ITypedElementDataStorageInterface& DataStorage,
+		ITypedElementDataStorageUiInterface& DataStorageUi) const override;
+
+	TEDSOUTLINER_API void RegisterQueries(ITypedElementDataStorageInterface& DataStorage) override;
+
+};
+
+USTRUCT(meta = (DisplayName = "Override Widget"))
+struct FOverrideWidgetTag : public FTypedElementDataStorageTag
+{
+	GENERATED_BODY()
+};
+
+USTRUCT()
+struct FOverrideWidgetConstructor : public FTypeInfoWidgetConstructor
+{
+	GENERATED_BODY()
+
+public:
+	FOverrideWidgetConstructor();
+	~FOverrideWidgetConstructor() override = default;
+
+	virtual TConstArrayView<const UScriptStruct*> GetAdditionalColumnsList() const override;
+	
+	static void AddOverrideBadge(const TWeakPtr<SWidget>& Widget, EOverriddenState OverriddenState);
+	static void RemoveOverrideBadge(const TWeakPtr<SWidget>& Widget);
+
+	static void UpdateOverrideWidget(const TWeakPtr<SWidget>& Widget, const TypedElementRowHandle TargetRow);
+
+protected:
+	TEDSOUTLINER_API TSharedPtr<SWidget> CreateWidget(const TypedElementDataStorage::FMetaDataView& Arguments) override;
+	TEDSOUTLINER_API bool FinalizeWidget(ITypedElementDataStorageInterface* DataStorage, ITypedElementDataStorageUiInterface* DataStorageUi,
+		TypedElementRowHandle Row, const TSharedPtr<SWidget>& Widget) override;
+};
\ No newline at end of file
diff --git a/Engine/Plugins/Experimental/EditorDataStorageFeatures/Source/TedsOutliner/Private/Widgets/TedsOutlinerRowHandleColumn.cpp b/Engine/Plugins/Experimental/EditorDataStorageFeatures/Source/TedsOutliner/Private/Widgets/TedsOutlinerRowHandleColumn.cpp
deleted file mode 100644
diff --git a/Engine/Plugins/Experimental/EditorDataStorageFeatures/Source/TedsOutliner/Private/Widgets/TedsOutlinerRowHandleColumn.h b/Engine/Plugins/Experimental/EditorDataStorageFeatures/Source/TedsOutliner/Private/Widgets/TedsOutlinerRowHandleColumn.h
deleted file mode 100644
diff --git a/Engine/Plugins/Experimental/EditorDataStorageFeatures/Source/TedsOutliner/Private/Widgets/TransformHeadsUpWidget.cpp b/Engine/Plugins/Experimental/EditorDataStorageFeatures/Source/TedsOutliner/Private/Widgets/TransformHeadsUpWidget.cpp
new file mode 100644
--- /dev/null
+++ b/Engine/Plugins/Experimental/EditorDataStorageFeatures/Source/TedsOutliner/Private/Widgets/TransformHeadsUpWidget.cpp
@@ -0,0 +1,238 @@
+// Copyright Epic Games, Inc. All Rights Reserved.
+
+#include "Widgets/TransformHeadsUpWidget.h"
+
+#include "Elements/Columns/TypedElementMiscColumns.h"
+#include "Elements/Columns/TypedElementSlateWidgetColumns.h"
+#include "Elements/Columns/TypedElementTransformColumns.h"
+#include "Elements/Framework/TypedElementQueryBuilder.h"
+#include "Widgets/SCompoundWidget.h"
+#include "Widgets/SBoxPanel.h"
+#include "Widgets/Images/SImage.h"
+#include "Styling/AppStyle.h"
+#include "Widgets/Input/SNumericEntryBox.h"
+
+#define LOCTEXT_NAMESPACE "TedsTransformHeadsUpWidget"
+
+namespace UE::EditorDataStorage::Widgets::Private
+{
+	enum class EAbnormalTransformTypes : int32
+	{
+		None = 0x0000,
+		NonUniformScale = 0x0001,
+		NegativeXScale = 0x0002,
+		NegativeYScale = 0x0004,
+		NegativeZScale = 0x0008,
+		UnnormalizedRotation = 0x0010,
+	};
+	ENUM_CLASS_FLAGS(EAbnormalTransformTypes);
+
+	EAbnormalTransformTypes GetAbnormalTransformTypes(const FTransform& InTransform)
+	{
+		EAbnormalTransformTypes AbnormalTransformFlags = EAbnormalTransformTypes::None;
+		const FVector& TransformScale = InTransform.GetScale3D();
+
+		if (!TransformScale.GetAbs().AllComponentsEqual())
+		{
+			AbnormalTransformFlags |= EAbnormalTransformTypes::NonUniformScale;
+		}
+		if (TransformScale.X < 0.0)
+		{
+			AbnormalTransformFlags |= EAbnormalTransformTypes::NegativeXScale;
+		}
+		if (TransformScale.Y < 0.0)
+		{
+			AbnormalTransformFlags |= EAbnormalTransformTypes::NegativeYScale;
+		}
+		if (TransformScale.Z < 0.0)
+		{
+			AbnormalTransformFlags |= EAbnormalTransformTypes::NegativeZScale;
+		}
+		if (!InTransform.IsRotationNormalized())
+		{
+			AbnormalTransformFlags |= EAbnormalTransformTypes::UnnormalizedRotation;
+		}
+
+		return AbnormalTransformFlags;
+	}
+} // namespace UE::EditorDataStorage::Widgets::Private
+
+class STransformQuickDisplay : public SHorizontalBox
+{
+	SLATE_DECLARE_WIDGET(STransformQuickDisplay, SHorizontalBox)
+public:
+	SLATE_BEGIN_ARGS(STransformQuickDisplay) {}
+	SLATE_END_ARGS()
+
+	void Construct(const FArguments& InArgs)
+	{
+		const static FMargin IconPadding(1, 1, 0, 0);
+		using namespace UE::EditorDataStorage::Widgets::Private;
+
+		AddSlot().AutoWidth().Padding(IconPadding)
+		[
+			SNew(SImage)
+				.Image(FAppStyle::GetBrush("EditorViewport.ScaleGridSnap"))
+				.ToolTipText(LOCTEXT("NonUniformScaleTooltip", "Has Non-Uniform Scale"))
+				.Visibility_Lambda([this]() { return EnumHasAllFlags(AbnormalTransformFlags, EAbnormalTransformTypes::NonUniformScale) ? EVisibility::Visible : EVisibility::Collapsed; })
+		];
+		AddSlot().AutoWidth().Padding(IconPadding)
+		[
+			SNew(SImage)
+				.Image(FAppStyle::GetBrush("EditorViewport.ScaleMode"))
+				.ColorAndOpacity(SNumericEntryBox<double>::RedLabelBackgroundColor)
+				.ToolTipText(LOCTEXT("NegativeXScaleTooltip", "Has Negative X Scale"))
+				.Visibility_Lambda([this]() { return EnumHasAllFlags(AbnormalTransformFlags, EAbnormalTransformTypes::NegativeXScale) ? EVisibility::Visible : EVisibility::Collapsed; })
+		];
+		AddSlot().AutoWidth().Padding(IconPadding)
+		[
+			SNew(SImage)
+				.Image(FAppStyle::GetBrush("EditorViewport.ScaleMode"))
+				.ColorAndOpacity(SNumericEntryBox<double>::GreenLabelBackgroundColor)
+				.ToolTipText(LOCTEXT("NegativeYScaleTooltip", "Has Negative Y Scale"))
+				.Visibility_Lambda([this]() { return EnumHasAllFlags(AbnormalTransformFlags, EAbnormalTransformTypes::NegativeYScale) ? EVisibility::Visible : EVisibility::Collapsed; })
+		];
+		AddSlot().AutoWidth().Padding(IconPadding)
+		[
+			SNew(SImage)
+				.Image(FAppStyle::GetBrush("EditorViewport.ScaleMode"))
+				.ColorAndOpacity(SNumericEntryBox<double>::BlueLabelBackgroundColor)
+				.ToolTipText(LOCTEXT("NegativeZScaleTooltip", "Has Negative Z Scale"))
+				.Visibility_Lambda([this]() { return EnumHasAllFlags(AbnormalTransformFlags, EAbnormalTransformTypes::NegativeZScale) ? EVisibility::Visible : EVisibility::Collapsed; })
+		];
+		AddSlot().AutoWidth().Padding(IconPadding)
+		[
+			SNew(SImage)
+				.Image(FAppStyle::GetBrush("SurfaceDetails.AntiClockwiseRotation"))
+				.ToolTipText(LOCTEXT("UnnormalizedRotationTooltip", "Has Un-normalized Rotation"))
+				.Visibility_Lambda([this]() { return EnumHasAllFlags(AbnormalTransformFlags, EAbnormalTransformTypes::UnnormalizedRotation) ? EVisibility::Visible : EVisibility::Collapsed; })
+		];
+		AddSlot().AutoWidth().Padding(IconPadding)
+		[
+			SNew(SImage)
+				.Image(FAppStyle::GetBrush("Symbols.Check"))
+				.ToolTipText(LOCTEXT("NothingToReportTooltip", "No Abnormal Transform Data"))
+				.Visibility_Lambda([this]() { return (AbnormalTransformFlags == EAbnormalTransformTypes::None) ? EVisibility::Visible : EVisibility::Collapsed; })
+		];
+	}
+
+	void UpdateFromTransform(UE::EditorDataStorage::Widgets::Private::EAbnormalTransformTypes InAbnormalTransformFlags)
+	{
+		if (AbnormalTransformFlags != InAbnormalTransformFlags)
+		{
+			AbnormalTransformFlags = InAbnormalTransformFlags;
+			Invalidate(EInvalidateWidgetReason::Visibility);
+		}
+	}
+
+private:
+	UE::EditorDataStorage::Widgets::Private::EAbnormalTransformTypes AbnormalTransformFlags;
+};
+
+SLATE_IMPLEMENT_WIDGET(STransformQuickDisplay)
+void STransformQuickDisplay::PrivateRegisterAttributes(FSlateAttributeInitializer&) {}
+
+static void UpdateTransformHeadsUpDisplay(FTypedElementSlateWidgetReferenceColumn& Widget, UE::EditorDataStorage::Widgets::Private::EAbnormalTransformTypes AbnormalTransformFlags)
+{
+	TSharedPtr<SWidget> WidgetPointer = Widget.Widget.Pin();
+	checkf(WidgetPointer, TEXT("Referenced widget is not valid. A constructed widget may not have been cleaned up. This can "
+		"also happen if this processor is running in the same phase as the processors responsible for cleaning up old "
+		"references."));
+	checkf(WidgetPointer->GetType() == STransformQuickDisplay::StaticWidgetClass().GetWidgetType(),
+		TEXT("Stored widget with FTransformHeadsUpWidgetTag doesn't match type %s, but was a %s."),
+		*(STransformQuickDisplay::StaticWidgetClass().GetWidgetType().ToString()),
+		*(WidgetPointer->GetTypeAsString()));
+
+	STransformQuickDisplay* BoxWidget = static_cast<STransformQuickDisplay*>(WidgetPointer.Get());
+	BoxWidget->UpdateFromTransform(AbnormalTransformFlags);
+}
+
+//
+// UTransformHeadsUpWidgetFactory
+//
+
+void UTransformHeadsUpWidgetFactory::RegisterQueries(ITypedElementDataStorageInterface& DataStorage)
+{
+	using namespace TypedElementQueryBuilder;
+	using DSI = ITypedElementDataStorageInterface;
+	namespace DS = TypedElementDataStorage;
+		
+	TypedElementQueryHandle UpdateTransformWidget = DataStorage.RegisterQuery(
+		Select()
+			.ReadOnly<FTypedElementLocalTransformColumn>()
+		.Where()
+			.Any<FTypedElementSyncFromWorldTag, FTypedElementSyncBackToWorldTag>()
+		.Compile());
+
+	DataStorage.RegisterQuery(
+		Select(TEXT("Sync Transform column to heads up display"),
+		FProcessor(DSI::EQueryTickPhase::FrameEnd, DataStorage.GetQueryTickGroupName(DSI::EQueryTickGroups::SyncWidgets))
+			.ForceToGameThread(true),
+			[](DS::IQueryContext& Context,
+				FTypedElementSlateWidgetReferenceColumn& Widget,
+				const FTypedElementRowReferenceColumn& ReferenceColumn)
+			{
+				Context.RunSubquery(0, ReferenceColumn.Row, CreateSubqueryCallbackBinding(
+					[&Widget](const FTypedElementLocalTransformColumn& Transform)
+					{
+						UpdateTransformHeadsUpDisplay(Widget, UE::EditorDataStorage::Widgets::Private::GetAbnormalTransformTypes(Transform.Transform));
+					}));
+			}
+		)
+	.Where()
+		.All<FTransformHeadsUpWidgetTag>()
+	.DependsOn()
+		.SubQuery(UpdateTransformWidget)
+	.Compile());
+
+}
+
+void UTransformHeadsUpWidgetFactory::RegisterWidgetConstructors(ITypedElementDataStorageInterface& DataStorage,
+	ITypedElementDataStorageUiInterface& DataStorageUi) const
+{
+	using namespace TypedElementDataStorage;
+
+	DataStorageUi.RegisterWidgetFactory<FTransformHeadsUpWidgetConstructor>(FName(TEXT("SceneOutliner.Cell")),
+		FColumn<FTypedElementLocalTransformColumn>());
+}
+
+
+
+//
+// FTransformHeadsUpWidgetConstructor
+//
+
+FTransformHeadsUpWidgetConstructor::FTransformHeadsUpWidgetConstructor()
+	: Super(FTransformHeadsUpWidgetConstructor::StaticStruct())
+{
+}
+
+TConstArrayView<const UScriptStruct*> FTransformHeadsUpWidgetConstructor::GetAdditionalColumnsList() const
+{
+	static TTypedElementColumnTypeList<FTypedElementRowReferenceColumn,
+		FTransformHeadsUpWidgetTag> Columns;
+	return Columns;
+}
+
+TSharedPtr<SWidget> FTransformHeadsUpWidgetConstructor::CreateWidget(const TypedElementDataStorage::FMetaDataView& Arguments)
+{
+	return SNew(STransformQuickDisplay);
+}
+
+bool FTransformHeadsUpWidgetConstructor::FinalizeWidget(
+	ITypedElementDataStorageInterface* DataStorage,
+	ITypedElementDataStorageUiInterface* DataStorageUi,
+	TypedElementRowHandle Row,
+	const TSharedPtr<SWidget>& Widget)
+{
+	FTypedElementRowReferenceColumn& RefColumn = *DataStorage->GetColumn<FTypedElementRowReferenceColumn>(Row);
+	if (const FTypedElementLocalTransformColumn* TransformColumn = DataStorage->GetColumn<FTypedElementLocalTransformColumn>(RefColumn.Row))
+	{
+		UpdateTransformHeadsUpDisplay(*DataStorage->GetColumn<FTypedElementSlateWidgetReferenceColumn>(Row), 
+			UE::EditorDataStorage::Widgets::Private::GetAbnormalTransformTypes(TransformColumn->Transform));
+	}
+
+	return true;
+}
+
+#undef LOCTEXT_NAMESPACE
diff --git a/Engine/Plugins/Experimental/EditorDataStorageFeatures/Source/TedsOutliner/Private/Widgets/TypeInfoWidget.cpp b/Engine/Plugins/Experimental/EditorDataStorageFeatures/Source/TedsOutliner/Private/Widgets/TypeInfoWidget.cpp
new file mode 100644
--- /dev/null
+++ b/Engine/Plugins/Experimental/EditorDataStorageFeatures/Source/TedsOutliner/Private/Widgets/TypeInfoWidget.cpp
@@ -0,0 +1,170 @@
+// Copyright Epic Games, Inc. All Rights Reserved.
+
+#include "Widgets/TypeInfoWidget.h"
+
+#include "SceneOutlinerHelpers.h"
+#include "Elements/Columns/TypedElementCompatibilityColumns.h"
+#include "Elements/Columns/TypedElementMiscColumns.h"
+#include "Elements/Columns/TypedElementSlateWidgetColumns.h"
+#include "Elements/Columns/TypedElementTypeInfoColumns.h"
+#include "Elements/Interfaces/Capabilities/TypedElementUiTextCapability.h"
+#include "Styling/SlateIconFinder.h"
+#include "Widgets/Text/STextBlock.h"
+#include "Widgets/Images/SImage.h"
+
+void UTypeInfoWidgetFactory::RegisterWidgetConstructors(ITypedElementDataStorageInterface& DataStorage,
+	ITypedElementDataStorageUiInterface& DataStorageUi) const
+{
+	DataStorageUi.RegisterWidgetFactory<FTypeInfoWidgetConstructor>(FName(TEXT("General.Cell")),
+		TypedElementDataStorage::FColumn<FTypedElementClassTypeInfoColumn>());
+
+}
+
+TMap<FName, const FSlateBrush*> FTypeInfoWidgetConstructor::CachedIconMap;
+
+FTypeInfoWidgetConstructor::FTypeInfoWidgetConstructor()
+	: Super(FTypeInfoWidgetConstructor::StaticStruct())
+	, bUseIcon(false)
+{
+}
+
+FTypeInfoWidgetConstructor::FTypeInfoWidgetConstructor(const UScriptStruct* InTypeInfo)
+	: Super(InTypeInfo)
+	, bUseIcon(false)
+{
+	
+}
+
+TSharedPtr<SWidget> FTypeInfoWidgetConstructor::CreateWidget(
+	const TypedElementDataStorage::FMetaDataView& Arguments	)
+{
+	bUseIcon = false;
+	
+	// Check if the caller provided metadata to use an icon widget
+	TypedElementDataStorage::FMetaDataEntryView MetaDataEntryView = Arguments.FindGeneric("TypeInfoWidget_bUseIcon");
+	if(MetaDataEntryView.IsSet())
+	{
+		check(MetaDataEntryView.IsType<bool>());
+
+		bUseIcon = *MetaDataEntryView.TryGetExact<bool>();
+	}
+
+	if(bUseIcon)
+	{
+		return SNew(SImage)
+					.DesiredSizeOverride(FVector2D(16.f, 16.f))
+					.ColorAndOpacity(FSlateColor::UseForeground());
+	}
+	else
+	{
+		return SNew(SHorizontalBox);
+	}
+}
+
+bool FTypeInfoWidgetConstructor::FinalizeWidget(ITypedElementDataStorageInterface* DataStorage,
+	ITypedElementDataStorageUiInterface* DataStorageUi, TypedElementRowHandle Row, const TSharedPtr<SWidget>& Widget)
+{
+	checkf(Widget, TEXT("Referenced widget is not valid. A constructed widget may not have been cleaned up. This can "
+		"also happen if this processor is running in the same phase as the processors responsible for cleaning up old "
+		"references."));
+
+	TypedElementRowHandle TargetRow = DataStorage->GetColumn<FTypedElementRowReferenceColumn>(Row)->Row;
+
+	if (const FTypedElementClassTypeInfoColumn* TypeInfoColumn = DataStorage->GetColumn<FTypedElementClassTypeInfoColumn>(TargetRow))
+	{
+		if(bUseIcon)
+		{
+			checkf(Widget->GetType() == SImage::StaticWidgetClass().GetWidgetType(),
+				TEXT("Stored widget with FTypedElementLabelWidgetConstructor doesn't match type %s, but was a %s."),
+				*(SImage::StaticWidgetClass().GetWidgetType().ToString()),
+				*(Widget->GetTypeAsString()));
+
+			SImage* WidgetInstance = static_cast<SImage*>(Widget.Get());
+			
+			WidgetInstance->SetImage(GetIconForRow(DataStorage, Row, TypeInfoColumn));
+		}
+		else
+		{
+			checkf(Widget->GetType() == SHorizontalBox::StaticWidgetClass().GetWidgetType(),
+				TEXT("Stored widget with FTypedElementLabelWidgetConstructor doesn't match type %s, but was a %s."),
+				*(SHorizontalBox::StaticWidgetClass().GetWidgetType().ToString()),
+				*(Widget->GetTypeAsString()));
+				
+			SHorizontalBox* WidgetInstance = static_cast<SHorizontalBox*>(Widget.Get());
+
+			TSharedPtr<SWidget> ActualWidget;
+
+			// Check if we have a hyperlink for this object
+			if (const FTypedElementUObjectColumn* ObjectColumn = DataStorage->GetColumn<FTypedElementUObjectColumn>(TargetRow))
+			{
+				ActualWidget = SceneOutliner::FSceneOutlinerHelpers::GetClassHyperlink(ObjectColumn->Object.Get());
+			}
+
+			// If not, we simply show a text block with the type
+			if(!ActualWidget)
+			{
+				TSharedPtr<STextBlock> TextBlock = SNew(STextBlock)
+								.ColorAndOpacity(FSlateColor::UseSubduedForeground())
+								.Text(FText::FromString(TypeInfoColumn->TypeInfo.Get()->GetName()));
+
+				TextBlock->AddMetadata(MakeShared<TTypedElementUiTextCapability<STextBlock>>(*TextBlock));
+				ActualWidget = TextBlock;
+			}
+
+			WidgetInstance->AddSlot()
+						.AutoWidth()
+						.VAlign(VAlign_Center)
+						.Padding(8, 0, 0, 0)
+						[
+							ActualWidget.ToSharedRef()
+						];
+		}
+	}
+
+	return true;
+}
+
+const FSlateBrush* FTypeInfoWidgetConstructor::GetIconForRow(ITypedElementDataStorageInterface* DataStorage, TypedElementRowHandle Row, const FTypedElementClassTypeInfoColumn* TypeInfoColumn)
+{
+	/* The logic here is very similar to SActorTreeLabel::GetIcon in ActorTreeItem.cpp which allows the actor to specify
+	 * an override for the icon, and has a fallback to the class icon if not.
+	 */
+	
+	const UClass* Type = TypeInfoColumn->TypeInfo.Get();
+	FName IconName;
+
+	// Allow the actor the first chance to provide an icon override
+	if (DataStorage->HasColumns<FTypedElementActorTag>(Row))
+	{
+		if (FTypedElementUObjectColumn* ActorStore = DataStorage->GetColumn<FTypedElementUObjectColumn>(Row))
+		{
+			if (const AActor* Actor = Cast<AActor>(ActorStore->Object))
+			{
+				IconName = Actor->GetCustomIconName();
+			}
+		}
+	}
+
+	if(IconName == NAME_None)
+	{
+		IconName = Type->GetFName();
+	}
+
+	// Check the cache if we already found an icon for this class
+	if(const FSlateBrush** CachedBrush = CachedIconMap.Find(IconName))
+	{
+		if(*CachedBrush)
+		{
+			return *CachedBrush;
+		}
+	}
+
+	if(const FSlateBrush* FoundSlateBrush = FSlateIconFinder::FindIconForClass(Type).GetOptionalIcon())
+	{
+		CachedIconMap.Add(IconName, FoundSlateBrush);
+		return FoundSlateBrush;
+	}
+
+	// Fallback to the regular actor icon if we haven't found any specific icon
+	return FSlateIconFinder::FindIconForClass(AActor::StaticClass()).GetOptionalIcon();;
+}
diff --git a/Engine/Plugins/Experimental/EditorDataStorageFeatures/Source/TedsOutliner/Private/Widgets/TypedElementAlertWidget.cpp b/Engine/Plugins/Experimental/EditorDataStorageFeatures/Source/TedsOutliner/Private/Widgets/TypedElementAlertWidget.cpp
deleted file mode 100644
diff --git a/Engine/Plugins/Experimental/EditorDataStorageFeatures/Source/TedsOutliner/Private/Widgets/TypedElementOverrideWidget.cpp b/Engine/Plugins/Experimental/EditorDataStorageFeatures/Source/TedsOutliner/Private/Widgets/TypedElementOverrideWidget.cpp
deleted file mode 100644
diff --git a/Engine/Plugins/Experimental/EditorDataStorageFeatures/Source/TedsOutliner/Private/Widgets/TypedElementOverrideWidget.h b/Engine/Plugins/Experimental/EditorDataStorageFeatures/Source/TedsOutliner/Private/Widgets/TypedElementOverrideWidget.h
deleted file mode 100644
diff --git a/Engine/Plugins/Experimental/EditorDataStorageFeatures/Source/TedsOutliner/Private/Widgets/TypedElementTransformHeadsUpWidget.cpp b/Engine/Plugins/Experimental/EditorDataStorageFeatures/Source/TedsOutliner/Private/Widgets/TypedElementTransformHeadsUpWidget.cpp
deleted file mode 100644
diff --git a/Engine/Plugins/Experimental/EditorDataStorageFeatures/Source/TedsOutliner/Private/Widgets/TypedElementTypeInfoWidget.cpp b/Engine/Plugins/Experimental/EditorDataStorageFeatures/Source/TedsOutliner/Private/Widgets/TypedElementTypeInfoWidget.cpp
deleted file mode 100644
diff --git a/Engine/Plugins/Experimental/EditorDataStorageFeatures/Source/TedsOutliner/Public/Compatibility/SceneOutlinerTedsBridge.h b/Engine/Plugins/Experimental/EditorDataStorageFeatures/Source/TedsOutliner/Public/Compatibility/SceneOutlinerTedsBridge.h
new file mode 100644
--- /dev/null
+++ b/Engine/Plugins/Experimental/EditorDataStorageFeatures/Source/TedsOutliner/Public/Compatibility/SceneOutlinerTedsBridge.h
@@ -0,0 +1,75 @@
+// Copyright Epic Games, Inc. All Rights Reserved.
+
+#pragma once
+
+#include "SceneOutlinerStandaloneTypes.h"
+#include "Containers/ContainersFwd.h"
+#include "Elements/Interfaces/TypedElementDataStorageFactory.h"
+#include "Elements/Interfaces/TypedElementDataStorageInterface.h"
+#include "Elements/Interfaces/TypedElementDataStorageUiInterface.h"
+#include "Templates/SharedPointer.h"
+#include "UObject/ObjectMacros.h"
+
+#include "SceneOutlinerTedsBridge.generated.h"
+
+class ISceneOutliner;
+class ITypedElementDataStorageInterface;
+class ITypedElementDataStorageUiInterface;
+class ITypedElementDataStorageCompatibilityInterface;
+
+DECLARE_DELEGATE_RetVal_OneParam(FSceneOutlinerTreeItemID, FTreeItemIDDealiaser, TypedElementDataStorage::RowHandle);
+
+class FSceneOutlinerTedsBridge;
+
+/**
+ * Utility class to bind Typed Elements Data Storage queries to a Scene Outliner. The provided query is expected to be a select query
+ * and will be used to populate the Scene Outliner in addition to already existing data.
+ */
+class TEDSOUTLINER_API FSceneOutlinerTedsQueryBinder
+{
+public:
+	static const FName CellWidgetTableName;
+	static const FName HeaderWidgetPurpose;
+	static const FName DefaultHeaderWidgetPurpose;
+	static const FName CellWidgetPurpose;
+	static const FName DefaultCellWidgetPurpose;
+	static const FName ItemLabelCellWidgetPurpose;
+	static const FName DefaultItemLabelCellWidgetPurpose;
+
+
+	static FSceneOutlinerTedsQueryBinder& GetInstance();
+
+	void AssignQuery(TypedElementQueryHandle Query, const TSharedPtr<ISceneOutliner>& Widget, TConstArrayView<FName> InCellWidgetPurposes);
+
+	// Register a dealiser for a specific TEDS-Outliner to convert a row handle to an FSceneOutlinerTreeItemID
+	void RegisterTreeItemIDDealiaser(const TSharedPtr<ISceneOutliner>& Widget, const FTreeItemIDDealiaser& InDealiaser);
+
+	// Get the name of the Outliner column corresponding to the given TEDS column (if any)
+	FName FindOutlinerColumnFromTEDSColumns(TConstArrayView<TWeakObjectPtr<const UScriptStruct>> TEDSColumns) const;
+
+private:
+	FSceneOutlinerTedsQueryBinder();
+	void SetupDefaultColumnMapping();
+	void CleanupStaleOutliners();
+	
+	TSharedPtr<FSceneOutlinerTedsBridge>* FindOrAddQueryMapping(const TSharedPtr<ISceneOutliner>& Widget);
+	
+	TMap<TWeakPtr<ISceneOutliner>, TSharedPtr<FSceneOutlinerTedsBridge>> SceneOutliners;
+
+	ITypedElementDataStorageInterface* Storage{ nullptr };
+	ITypedElementDataStorageUiInterface* StorageUi{ nullptr };
+	ITypedElementDataStorageCompatibilityInterface* StorageCompatibility{ nullptr };
+
+	TMap<TWeakObjectPtr<const UScriptStruct>, FName> TEDSToOutlinerDefaultColumnMapping;
+};
+
+UCLASS()
+class USceneOutlinerTedsBridgeFactory : public UTypedElementDataStorageFactory
+{
+	GENERATED_BODY()
+
+public:
+	~USceneOutlinerTedsBridgeFactory() override = default;
+
+	void RegisterWidgetPurposes(ITypedElementDataStorageUiInterface& DataStorageUi) const override;
+};
diff --git a/Engine/Plugins/Experimental/EditorDataStorageFeatures/Source/TedsOutliner/Public/Compatibility/TedsCompatibilityUtils.h b/Engine/Plugins/Experimental/EditorDataStorageFeatures/Source/TedsOutliner/Public/Compatibility/TedsCompatibilityUtils.h
deleted file mode 100644
diff --git a/Engine/Plugins/Experimental/EditorDataStorageFeatures/Source/TedsOutliner/Public/TedsOutlinerFilter.h b/Engine/Plugins/Experimental/EditorDataStorageFeatures/Source/TedsOutliner/Public/TedsOutlinerFilter.h
new file mode 100644
--- /dev/null
+++ b/Engine/Plugins/Experimental/EditorDataStorageFeatures/Source/TedsOutliner/Public/TedsOutlinerFilter.h
@@ -0,0 +1,61 @@
+// Copyright Epic Games, Inc. All Rights Reserved.
+
+#pragma once
+
+#include "Filters/FilterBase.h"
+#include "SceneOutlinerPublicTypes.h"
+#include "Elements/Common/TypedElementQueryDescription.h"
+#include "Elements/Framework/TypedElementMetaData.h"
+
+class FTedsOutlinerImpl;
+
+// TEDS-Outliner TODO: Lots of minor missing functionality for icon, tooltip, color etc
+class FTedsOutlinerFilter : public FFilterBase<SceneOutliner::FilterBarType>
+{
+public:
+	FTedsOutlinerFilter(const FName& InFilterName, const FText& InFilterDisplayName, TSharedPtr<FFilterCategory> InCategory,
+		TSharedRef<FTedsOutlinerImpl> InTedsOutlinerImpl, const TypedElementDataStorage::FQueryDescription& InFilterQuery);
+
+	/** Returns the system name for this filter */
+	virtual FString GetName() const override;
+
+	/** Returns the human readable name for this filter */
+	virtual FText GetDisplayName() const override;
+
+	/** Returns the tooltip for this filter, shown in the filters menu */
+	virtual FText GetToolTipText() const override;
+
+	/** Returns the color this filter button will be when displayed as a button */
+	virtual FLinearColor GetColor() const override;
+
+	/** Returns the name of the icon to use in menu entries */
+	virtual FName GetIconName() const override;
+
+	/** If true, the filter will be active in the FilterBar when it is inactive in the UI (i.e the filter pill is grayed out)
+	 * @See: FFrontendFilter_ShowOtherDevelopers in Content Browser
+	 */
+	virtual bool IsInverseFilter() const override;
+
+	/** Notification that the filter became active or inactive */
+	virtual void ActiveStateChanged(bool bActive) override;
+
+	/** Called when the right-click context menu is being built for this filter */
+	virtual void ModifyContextMenu(FMenuBuilder& MenuBuilder) override;
+	
+	/** Can be overriden for custom FilterBar subclasses to save settings, currently not implemented in any generic Filter Bar */
+	virtual void SaveSettings(const FString& IniFilename, const FString& IniSection, const FString& SettingsString) const override;
+
+	/** Can be overridden for custom FilterBar subclasses to load settings, currently not implemented in any generic Filter Bar */
+	virtual void LoadSettings(const FString& IniFilename, const FString& IniSection, const FString& SettingsString) override;
+	
+	/** Returns whether the specified Item passes the Filter's restrictions */
+	virtual bool PassesFilter( SceneOutliner::FilterBarType InItem ) const override;
+
+protected:
+
+	FName FilterName;
+	FText FilterDisplayName;
+
+	TSharedRef<FTedsOutlinerImpl> TedsOutlinerImpl;
+	const TypedElementDataStorage::FQueryDescription FilterQuery;
+};
\ No newline at end of file
diff --git a/Engine/Plugins/Experimental/EditorDataStorageFeatures/Source/TedsOutliner/Public/TedsOutlinerImpl.h b/Engine/Plugins/Experimental/EditorDataStorageFeatures/Source/TedsOutliner/Public/TedsOutlinerImpl.h
new file mode 100644
--- /dev/null
+++ b/Engine/Plugins/Experimental/EditorDataStorageFeatures/Source/TedsOutliner/Public/TedsOutlinerImpl.h
@@ -0,0 +1,239 @@
+// Copyright Epic Games, Inc. All Rights Reserved.
+
+#pragma once
+
+#include "Elements/Interfaces/TypedElementDataStorageInterface.h"
+#include "ISceneOutlinerHierarchy.h"
+#include "ISceneOutlinerMode.h"
+#include "Elements/Columns/TypedElementHiearchyColumns.h"
+
+class ITypedElementDataStorageUiInterface;
+class ITypedElementDataStorageCompatibilityInterface;
+struct FTypedElementWidgetConstructor;
+class SWidget;
+
+// Struct storing information on how hierarchies are handled in the TEDS Outliner
+struct FTedsOutlinerHierarchyData
+{
+	/** A delegate used to get the parent row handle for a given row */
+	DECLARE_DELEGATE_RetVal_OneParam(TypedElementDataStorage::RowHandle, FGetParentRowHandle, void* /* InColumnData */);
+	
+	/** A delegate used to set the parent row handle for a given row */
+	DECLARE_DELEGATE_TwoParams(FSetParentRowHandle, void* /* InColumnData */, TypedElementDataStorage::RowHandle /* InParentRowHandle */);
+
+	FTedsOutlinerHierarchyData(const UScriptStruct* InHierarchyColumn, const FGetParentRowHandle& InGetParent, const FSetParentRowHandle& InSetParent)
+		: HierarchyColumn(InHierarchyColumn)
+		, GetParent(InGetParent)
+		, SetParent(InSetParent)
+	{
+	
+	}
+
+	// The column that contains the parent row handle for rows
+	const UScriptStruct* HierarchyColumn;
+
+	// Function to get parent row handle
+	FGetParentRowHandle GetParent;
+
+	// Function to set the parent row handle
+	FSetParentRowHandle SetParent;
+	
+	// Get the default hierarchy data for the TEDS Outliner that uses FTableRowParentColumn to get the parent
+	static FTedsOutlinerHierarchyData GetDefaultHierarchyData()
+	{
+		const FGetParentRowHandle RowHandleGetter = FGetParentRowHandle::CreateLambda([](void* InColumnData)
+			{
+				if(const FTableRowParentColumn* ParentColumn = static_cast<FTableRowParentColumn *>(InColumnData))
+				{
+					return ParentColumn->Parent;
+				}
+
+				return TypedElementDataStorage::InvalidRowHandle;
+			});
+
+		const FSetParentRowHandle RowHandleSetter = FSetParentRowHandle::CreateLambda([](void* InColumnData,
+			TypedElementDataStorage::RowHandle InRowHandle)
+			{
+				if(FTableRowParentColumn* ParentColumn = static_cast<FTableRowParentColumn *>(InColumnData))
+				{
+					ParentColumn->Parent = InRowHandle;
+				}
+
+			});
+		
+		return FTedsOutlinerHierarchyData(FTableRowParentColumn::StaticStruct(), RowHandleGetter, RowHandleSetter);
+	}
+};
+
+struct FTedsOutlinerParams
+{
+	FTedsOutlinerParams(SSceneOutliner* InSceneOutliner)
+	: SceneOutliner(InSceneOutliner)
+	, QueryDescription()
+	, bUseDefaultTedsFilters(false)
+	, HierarchyData(FTedsOutlinerHierarchyData::GetDefaultHierarchyData())
+	, CellWidgetPurposes{TEXT("SceneOutliner.Cell"), TEXT("General.Cell")}
+	{}
+
+	SSceneOutliner* SceneOutliner;
+
+	// The query description that will be used to populate rows in the TEDS-Outliner
+	TAttribute<TypedElementDataStorage::FQueryDescription> QueryDescription;
+	
+	// TEDS queries that will be used to create filters in this Outliner
+	// TEDS-Outliner TODO: Can we consolidate this with the SceneOutliner API to create filters? Currently has to be separate because FTEDSOutlinerFilter
+	// needs a reference to the mode which is not possible since filters with the Outliner API are added before the mode is init
+	TMap<FName, const TypedElementDataStorage::FQueryDescription> FilterQueries;
+
+	// If true, this Outliner will automatically add all TEDS tags and columns as filters
+	bool bUseDefaultTedsFilters;
+
+	// If specified, this is how the TEDS Outliner will handle hierarchies. If not specified - there will be no hierarchies shown as a
+	// parent-child relation in the tree view
+	TOptional<FTedsOutlinerHierarchyData> HierarchyData;
+
+	// The selection set to use for this Outliner, unset = don't propagate tree selection to the TEDS column
+	TOptional<FName> SelectionSetOverride;
+
+	// The purposes to use when generating widgets for the columns through TEDS UI
+	TArray<FName> CellWidgetPurposes;
+};
+
+
+// This class is meant to be a model to hold functionality to create a "table viewer" in TEDS that can be
+// attached to any view/UI.
+// TEDS-Outliner TODO: This class still has a few outliner implementation details leaking in that should be removed
+class TEDSOUTLINER_API FTedsOutlinerImpl : public TSharedFromThis<FTedsOutlinerImpl>
+{
+
+public:
+
+	FTedsOutlinerImpl(const FTedsOutlinerParams& InParams, ISceneOutlinerMode* InMode);
+	virtual ~FTedsOutlinerImpl();
+
+	void Init();
+
+	// TEDS construct getters
+	ITypedElementDataStorageInterface* GetStorage() const;
+	ITypedElementDataStorageUiInterface* GetStorageUI() const;
+	ITypedElementDataStorageCompatibilityInterface* GetStorageCompatibility() const;
+
+	TOptional<FName> GetSelectionSetName() const;
+
+	// Delegate fired when the selection in TEDS changes, only if SelectionSetName is set
+	DECLARE_MULTICAST_DELEGATE(FOnTedsOutlinerSelectionChanged)
+	FOnTedsOutlinerSelectionChanged& OnSelectionChanged();
+
+	// Delegate fired when the hierarchy changes due to item addition/removal/move
+	ISceneOutlinerHierarchy::FHierarchyChangedEvent& OnHierarchyChanged();
+
+	// Delegate to check if a certain outliner item is compatible with this TEDS Outliner Impl - set by the system using FTedsOutlinerImpl
+	DECLARE_DELEGATE_RetVal_OneParam(bool, FIsItemCompatible, const ISceneOutlinerTreeItem&)
+	FIsItemCompatible& IsItemCompatible();
+
+	// Update the selection in TEDS to the input rows, only if SelectionSetName is set
+	void SetSelection(const TArray<TypedElementDataStorage::RowHandle>& InSelectedRows);
+
+	// Helper function to create a label widget for a given row
+	TSharedRef<SWidget> CreateLabelWidgetForItem(TypedElementRowHandle InRowHandle) const;
+
+	// Get the hierarchy data associated with this table viewer
+	const TOptional<FTedsOutlinerHierarchyData>& GetHierarchyData();
+	
+	// Add an external query to the Outliner
+	void AddExternalQuery(FName QueryName, const TypedElementDataStorage::FQueryDescription& InQueryDescription);
+	void RemoveExternalQuery(FName QueryName);
+
+	// Append all external queries into the given query
+	void AppendExternalQueries(TypedElementDataStorage::FQueryDescription& OutQuery);
+
+	// TEDS-Outliner TODO: This should live in TEDS long term
+	// Funtion to combine 2 queries (adds to second query to the first)
+	static void AppendQuery(TypedElementDataStorage::FQueryDescription& Query1, const TypedElementDataStorage::FQueryDescription& Query2);
+
+	// Check if the given item's parent has changed (i.e ParentRowHandle does not match what the Outliner reports as the parent)
+	bool HasItemParentChanged(TypedElementDataStorage::RowHandle ItemRowHandle, TypedElementDataStorage::RowHandle ParentRowHandle) const;
+
+	// Outliner specific functionality
+	void CreateItemsFromQuery(TArray<FSceneOutlinerTreeItemPtr>& OutItems, ISceneOutlinerMode* InMode) const;
+	void CreateChildren(const FSceneOutlinerTreeItemPtr& Item, TArray<FSceneOutlinerTreeItemPtr>& OutChildren) const;
+
+	// Get the parent row for a given row
+	TypedElementDataStorage::RowHandle GetParentRow(TypedElementDataStorage::RowHandle InRowHandle);
+
+	// Recompile all queries used by this table viewer
+	void RecompileQueries();
+
+protected:
+	
+	void OnItemAdded(TypedElementDataStorage::RowHandle ItemRowHandle);
+	void OnItemRemoved(TypedElementDataStorage::RowHandle ItemRowHandle);
+	void OnItemMoved(TypedElementDataStorage::RowHandle ItemRowHandle);
+
+	void UnregisterQueries() const;
+	void ClearSelection() const;
+	void Tick();
+
+	void CreateLabelWidgetConstructors();
+	void CreateFilterQueries();
+
+	// Check if this row can be displayed in this table viewer
+	bool CanDisplayRow(TypedElementDataStorage::RowHandle ItemRowHandle) const;
+	
+protected:
+	// TEDS Storage Constructs
+	ITypedElementDataStorageInterface* Storage{ nullptr };
+	ITypedElementDataStorageUiInterface* StorageUi{ nullptr };
+	ITypedElementDataStorageCompatibilityInterface* StorageCompatibility{ nullptr };
+
+	FTedsOutlinerParams CreationParams;
+
+	// Widget constructor to create the label widget
+	TArray<TPair<TypedElementDataStorage::QueryHandle, TSharedPtr<FTypedElementWidgetConstructor>>> QueryToWidgetConstructorMap;
+
+	// Widget purposes this table viewer supports
+	TArray<FName> CellWidgetPurposes;
+	
+	// Initial query provided by user
+	TAttribute<TypedElementDataStorage::FQueryDescription> InitialQueryDescription;
+
+	// External queries that are currently active (e.g Filters)
+	TMap<FName, TypedElementDataStorage::FQueryDescription> ExternalQueries;
+
+	// Optional Hierarchy Data
+	TOptional<FTedsOutlinerHierarchyData> HierarchyData;
+
+	// Querys to track row handle collection, addition and removal
+	TypedElementDataStorage::QueryHandle RowHandleQuery = TypedElementDataStorage::InvalidQueryHandle;
+	TypedElementDataStorage::QueryHandle RowAdditionQuery = TypedElementDataStorage::InvalidQueryHandle;
+	TypedElementDataStorage::QueryHandle RowRemovalQuery = TypedElementDataStorage::InvalidQueryHandle;
+
+	// Query to get all child rows
+	TypedElementDataStorage::QueryHandle ChildRowHandleQuery = TypedElementDataStorage::InvalidQueryHandle;
+
+	// Query to track when a row's parent gets changed
+	TypedElementDataStorage::QueryHandle UpdateParentQuery = TypedElementDataStorage::InvalidQueryHandle;
+	
+	// Query to get all selected rows, track selection added, track selection removed
+	TypedElementDataStorage::QueryHandle SelectedRowsQuery = TypedElementDataStorage::InvalidQueryHandle;
+	TypedElementDataStorage::QueryHandle SelectionAddedQuery = TypedElementDataStorage::InvalidQueryHandle;
+	TypedElementDataStorage::QueryHandle SelectionRemovedQuery = TypedElementDataStorage::InvalidQueryHandle;
+	
+	TOptional<FName> SelectionSetName;
+	bool bSelectionDirty = false;
+	
+	// Ticker for selection updates so we don't fire the delegate multiple times in one frame for multi select
+	FTSTicker::FDelegateHandle TickerHandle;
+	
+	FOnTedsOutlinerSelectionChanged OnTedsOutlinerSelectionChanged;
+
+	// Scene Outliner specific constructors
+	ISceneOutlinerMode* SceneOutlinerMode;
+	SSceneOutliner* SceneOutliner;
+
+	// Event fired when the hierarchy changes (addition/removal/move)
+	ISceneOutlinerHierarchy::FHierarchyChangedEvent HierarchyChangedEvent;
+
+	// Delegate to check if an item is compatible with this table viewer
+	FIsItemCompatible IsItemCompatibleWithTeds;
+};
diff --git a/Engine/Plugins/Experimental/EditorDataStorageFeatures/Source/TedsOutliner/Public/TedsOutlinerItem.h b/Engine/Plugins/Experimental/EditorDataStorageFeatures/Source/TedsOutliner/Public/TedsOutlinerItem.h
new file mode 100644
--- /dev/null
+++ b/Engine/Plugins/Experimental/EditorDataStorageFeatures/Source/TedsOutliner/Public/TedsOutlinerItem.h
@@ -0,0 +1,46 @@
+// Copyright Epic Games, Inc. All Rights Reserved.
+
+#pragma once
+
+#include "Elements/Interfaces/TypedElementDataStorageInterface.h"
+#include "TedsOutlinerMode.h"
+
+class FTedsOutlinerImpl;
+
+/*
+ * A generic item in the TEDS driven Outliner, that uses a TypedElementRowHandle to uniquely identify the object it is
+ * looking at. Functionality should be added through TEDS queries instead of having a different TreeItem type for each
+ * type of object you are looking at (i.e. Actor vs Folder)
+ * Inherits from ISceneOutlinerItem - which determines what type of item you are looking at. E.G FActorTreeItem for actors
+ */
+struct FTedsOutlinerTreeItem : ISceneOutlinerTreeItem
+{
+public:
+	
+	DECLARE_DELEGATE_RetVal_OneParam(bool, FFilterPredicate, const TypedElementDataStorage::RowHandle);
+
+	TEDSOUTLINER_API bool Filter(FFilterPredicate Pred) const
+	{
+		return Pred.Execute(RowHandle);
+	}
+
+	TEDSOUTLINER_API FTedsOutlinerTreeItem(const TypedElementRowHandle& InRowHandle, const TSharedRef<const FTedsOutlinerImpl>& InTedsOutlinerImpl);
+
+	/* Begin ISceneOutlinerTreeItem Implementation */
+	TEDSOUTLINER_API virtual bool IsValid() const override;
+	TEDSOUTLINER_API virtual FSceneOutlinerTreeItemID GetID() const override;
+	TEDSOUTLINER_API virtual FString GetDisplayString() const override;
+	TEDSOUTLINER_API virtual bool CanInteract() const override;
+	TEDSOUTLINER_API virtual TSharedRef<SWidget> GenerateLabelWidget(ISceneOutliner& Outliner, const STableRow<FSceneOutlinerTreeItemPtr>& InRow) override;
+	/** Generate a context menu for this item. Only called if *only* this item is selected. */
+	TEDSOUTLINER_API virtual void GenerateContextMenu(UToolMenu* Menu, SSceneOutliner& Outliner);
+	/* End ISceneOutlinerTreeItem Implementation */
+
+	TEDSOUTLINER_API static const FSceneOutlinerTreeItemType Type;
+
+	TEDSOUTLINER_API TypedElementRowHandle GetRowHandle() const;
+
+private:
+	const TypedElementRowHandle RowHandle;
+	const TSharedRef<const FTedsOutlinerImpl> TedsOutlinerImpl;
+};
\ No newline at end of file
diff --git a/Engine/Plugins/Experimental/EditorDataStorageFeatures/Source/TedsOutliner/Public/TedsOutlinerMode.h b/Engine/Plugins/Experimental/EditorDataStorageFeatures/Source/TedsOutliner/Public/TedsOutlinerMode.h
new file mode 100644
--- /dev/null
+++ b/Engine/Plugins/Experimental/EditorDataStorageFeatures/Source/TedsOutliner/Public/TedsOutlinerMode.h
@@ -0,0 +1,82 @@
+// Copyright Epic Games, Inc. All Rights Reserved.
+
+#pragma once
+
+#include "ISceneOutlinerMode.h"
+#include "TedsOutlinerImpl.h"
+#include "Elements/Interfaces/TypedElementDataStorageInterface.h"
+#include "Elements/Columns/TypedElementHiearchyColumns.h"
+
+#include "TedsOutlinerMode.generated.h"
+
+// TEDS-Outliner TODO: This can probably be moved to a more generic location for all TEDS related drag drops?
+class FTedsRowDragDropOp : public FDecoratedDragDropOp
+{
+public:
+	DRAG_DROP_OPERATOR_TYPE(FTedsRowDragDropOp, FDecoratedDragDropOp)
+
+	/** Rows we are dragging */
+	TArray<TypedElementDataStorage::RowHandle> DraggedRows;
+
+	void Init(const TArray<TypedElementDataStorage::RowHandle>& InRowHandles)
+	{
+		DraggedRows = InRowHandles;
+	}
+
+	static TSharedRef<FTedsRowDragDropOp> New(const TArray<TypedElementDataStorage::RowHandle>& InRowHandles)
+	{
+		TSharedRef<FTedsRowDragDropOp> Operation = MakeShareable(new FTedsRowDragDropOp);
+		
+		Operation->Init(InRowHandles);
+		Operation->SetupDefaults();
+		Operation->Construct();
+		
+		return Operation;
+	}
+};
+
+// Class to hold the owning scene outliner for a menu
+// TEDS-Outliner TODO: Once menus go through TEDS UI this can be done using the FTableViewerColumn on the widget row instead
+UCLASS()
+class UTedsOutlinerMenuContext : public UObject
+{
+	GENERATED_BODY()
+public:
+	SSceneOutliner* OwningSceneOutliner = nullptr;
+};
+
+/*
+ * TEDS driven Outliner mode where the Outliner is populated using the results of the RowHandleQueries passed in.
+ * See CreateGenericTEDSOutliner() for example usage
+ * Inherits from ISceneOutlinerMode - which contains all actions that depend on the type of item you are viewing in the Outliner
+ */
+class TEDSOUTLINER_API FTedsOutlinerMode : public ISceneOutlinerMode
+{
+public:
+	explicit FTedsOutlinerMode(const FTedsOutlinerParams& InParams);
+	virtual ~FTedsOutlinerMode() override;
+
+	/* ISceneOutlinerMode interface */
+	virtual void Rebuild() override;
+	virtual void SynchronizeSelection() override;
+	virtual void OnItemSelectionChanged(FSceneOutlinerTreeItemPtr Item, ESelectInfo::Type SelectionType, const FSceneOutlinerItemSelection& Selection) override;
+	virtual ESelectionMode::Type GetSelectionMode() const override { return ESelectionMode::Multi; }
+	virtual bool CanSupportDragAndDrop() const override { return true; } // TODO: Can we check this from TEDS somehow (if the user requests a drag column?)
+	virtual TSharedPtr<FDragDropOperation> CreateDragDropOperation(const FPointerEvent& MouseEvent, const TArray<FSceneOutlinerTreeItemPtr>& InTreeItems) const override;
+	virtual bool ParseDragDrop(FSceneOutlinerDragDropPayload& OutPayload, const FDragDropOperation& Operation) const override;
+	virtual FSceneOutlinerDragValidationInfo ValidateDrop(const ISceneOutlinerTreeItem& DropTarget, const FSceneOutlinerDragDropPayload& Payload) const override;
+	virtual void OnDrop(ISceneOutlinerTreeItem& DropTarget, const FSceneOutlinerDragDropPayload& Payload, const FSceneOutlinerDragValidationInfo& ValidationInfo) const override;
+	virtual TSharedPtr<SWidget> CreateContextMenu() override;
+	/* end ISceneOutlinerMode interface */
+
+protected:
+	
+	virtual TUniquePtr<ISceneOutlinerHierarchy> CreateHierarchy() override;
+	// Called by TedsOutlinerImpl when the selection in TEDS changes
+	void OnSelectionChanged();
+
+protected:
+
+	// The actual model for the TEDS Outliner
+	TSharedPtr<FTedsOutlinerImpl> TedsOutlinerImpl;
+};
\ No newline at end of file
diff --git a/Engine/Plugins/Experimental/EditorDataStorageFeatures/Source/TedsOutliner/Public/TedsOutlinerModule.h b/Engine/Plugins/Experimental/EditorDataStorageFeatures/Source/TedsOutliner/Public/TedsOutlinerModule.h
--- a/Engine/Plugins/Experimental/EditorDataStorageFeatures/Source/TedsOutliner/Public/TedsOutlinerModule.h
+++ b/Engine/Plugins/Experimental/EditorDataStorageFeatures/Source/TedsOutliner/Public/TedsOutlinerModule.h
@@ -9,7 +9,7 @@
 
 class ISceneOutliner;
 struct FSceneOutlinerInitializationOptions;
-struct FTypedElementOutlinerModeParams;
+struct FTedsOutlinerParams;
 class SDockTab;
 class SWidget;
 class FSpawnTabArgs;
@@ -34,7 +34,7 @@ public:
 	 * @return	New scene outliner widget
 	 */
 	virtual TSharedRef<ISceneOutliner> CreateTedsOutliner(
-		const FSceneOutlinerInitializationOptions& InInitOptions, const FTypedElementOutlinerModeParams& InInitTedsOptions, TypedElementDataStorage::QueryHandle ColumnQuery) const;
+		const FSceneOutlinerInitializationOptions& InInitOptions, const FTedsOutlinerParams& InInitTedsOptions, TypedElementDataStorage::QueryHandle ColumnQuery) const;
 
 	// IModuleInterface interface
 	virtual void StartupModule() override;
diff --git a/Engine/Plugins/Experimental/EditorDataStorageFeatures/Source/TedsOutliner/Public/TypedElementOutlinerColumnIntegration.h b/Engine/Plugins/Experimental/EditorDataStorageFeatures/Source/TedsOutliner/Public/TypedElementOutlinerColumnIntegration.h
deleted file mode 100644
diff --git a/Engine/Plugins/Experimental/EditorDataStorageFeatures/Source/TedsOutliner/Public/TypedElementOutlinerFilter.h b/Engine/Plugins/Experimental/EditorDataStorageFeatures/Source/TedsOutliner/Public/TypedElementOutlinerFilter.h
deleted file mode 100644
diff --git a/Engine/Plugins/Experimental/EditorDataStorageFeatures/Source/TedsOutliner/Public/TypedElementOutlinerItem.h b/Engine/Plugins/Experimental/EditorDataStorageFeatures/Source/TedsOutliner/Public/TypedElementOutlinerItem.h
deleted file mode 100644
diff --git a/Engine/Plugins/Experimental/EditorDataStorageFeatures/Source/TedsOutliner/Public/TypedElementOutlinerMode.h b/Engine/Plugins/Experimental/EditorDataStorageFeatures/Source/TedsOutliner/Public/TypedElementOutlinerMode.h
deleted file mode 100644
diff --git a/Engine/Plugins/Experimental/EditorDataStorageFeatures/Source/TedsOutliner/Public/Widgets/AlertWidget.h b/Engine/Plugins/Experimental/EditorDataStorageFeatures/Source/TedsOutliner/Public/Widgets/AlertWidget.h
new file mode 100644
--- /dev/null
+++ b/Engine/Plugins/Experimental/EditorDataStorageFeatures/Source/TedsOutliner/Public/Widgets/AlertWidget.h
@@ -0,0 +1,85 @@
+// Copyright Epic Games, Inc. All Rights Reserved.
+
+#pragma once
+
+#include "Elements/Interfaces/TypedElementDataStorageFactory.h"
+#include "Elements/Interfaces/TypedElementDataStorageInterface.h"
+#include "Elements/Interfaces/TypedElementDataStorageUiInterface.h"
+
+#include "AlertWidget.generated.h"
+
+UCLASS()
+class UAlertWidgetFactory : public UTypedElementDataStorageFactory
+{
+	GENERATED_BODY()
+
+public:
+	TEDSOUTLINER_API ~UAlertWidgetFactory() override = default;
+
+	TEDSOUTLINER_API void RegisterWidgetConstructors(ITypedElementDataStorageInterface& DataStorage,
+		ITypedElementDataStorageUiInterface& DataStorageUi) const override;
+	TEDSOUTLINER_API void RegisterQueries(ITypedElementDataStorageInterface& DataStorage) override;
+
+private:
+	void RegisterAlertQueries(ITypedElementDataStorageInterface& DataStorage);
+	void RegisterAlertHeaderQueries(ITypedElementDataStorageInterface& DataStorage);
+};
+
+USTRUCT()
+struct FAlertWidgetConstructor : public FTypedElementWidgetConstructor
+{
+	GENERATED_BODY()
+
+public:
+	static constexpr int32 IconBackgroundSlot = 1;
+	static constexpr int32 IconBadgeSlot = 2;
+	static constexpr int32 CounterTextSlot = 3;
+	static constexpr int32 ActionButtonSlot = 0;
+
+	static constexpr float BadgeFontSize = 7.0f;
+	static constexpr float BadgeHorizontalOffset = 13.0f;
+	static constexpr float BadgeVerticalOffset = 1.0f;
+
+	TEDSOUTLINER_API FAlertWidgetConstructor();
+	TEDSOUTLINER_API ~FAlertWidgetConstructor() override = default;
+
+protected:
+	TEDSOUTLINER_API TSharedPtr<SWidget> CreateWidget(const TypedElementDataStorage::FMetaDataView& Arguments) override;
+	TEDSOUTLINER_API TConstArrayView<const UScriptStruct*> GetAdditionalColumnsList() const override;
+	TEDSOUTLINER_API bool FinalizeWidget(ITypedElementDataStorageInterface* DataStorage, ITypedElementDataStorageUiInterface* DataStorageUi,
+		TypedElementDataStorage::RowHandle Row, const TSharedPtr<SWidget>& Widget) override;
+};
+
+USTRUCT(meta = (DisplayName = "General purpose alert"))
+struct FAlertWidgetTag : public FTypedElementDataStorageTag
+{
+	GENERATED_BODY()
+};
+
+USTRUCT()
+struct FAlertHeaderWidgetConstructor : public FTypedElementWidgetConstructor
+{
+	GENERATED_BODY()
+
+public:
+	TEDSOUTLINER_API FAlertHeaderWidgetConstructor();
+	TEDSOUTLINER_API ~FAlertHeaderWidgetConstructor() override = default;
+
+protected:
+	TEDSOUTLINER_API TSharedPtr<SWidget> CreateWidget(const TypedElementDataStorage::FMetaDataView& Arguments) override;
+	TEDSOUTLINER_API TConstArrayView<const UScriptStruct*> GetAdditionalColumnsList() const override;
+	TEDSOUTLINER_API bool FinalizeWidget(ITypedElementDataStorageInterface* DataStorage, ITypedElementDataStorageUiInterface* DataStorageUi,
+		TypedElementDataStorage::RowHandle Row, const TSharedPtr<SWidget>& Widget) override;
+};
+
+USTRUCT(meta = (DisplayName = "General purpose alert header"))
+struct FAlertHeaderWidgetTag : public FTypedElementDataStorageTag
+{
+	GENERATED_BODY()
+};
+
+USTRUCT(meta = (DisplayName = "Alert header active"))
+struct FAlertHeaderActiveWidgetTag : public FTypedElementDataStorageTag
+{
+	GENERATED_BODY()
+};
diff --git a/Engine/Plugins/Experimental/EditorDataStorageFeatures/Source/TedsOutliner/Public/Widgets/TransformHeadsUpWidget.h b/Engine/Plugins/Experimental/EditorDataStorageFeatures/Source/TedsOutliner/Public/Widgets/TransformHeadsUpWidget.h
new file mode 100644
--- /dev/null
+++ b/Engine/Plugins/Experimental/EditorDataStorageFeatures/Source/TedsOutliner/Public/Widgets/TransformHeadsUpWidget.h
@@ -0,0 +1,56 @@
+// Copyright Epic Games, Inc. All Rights Reserved.
+
+#pragma once
+
+#include "Elements/Interfaces/TypedElementDataStorageFactory.h"
+#include "Elements/Interfaces/TypedElementDataStorageInterface.h"
+#include "Elements/Interfaces/TypedElementDataStorageUiInterface.h"
+#include "Internationalization/Text.h"
+#include "UObject/ObjectMacros.h"
+
+#include "TransformHeadsUpWidget.generated.h"
+
+/**
+ * The heads up transform display provides at a glance view in a scene outliner row of abnormal transform characteristics, including:
+ *		1. Non-uniform scale
+ *		2. Negative scaling on X, Y, or Z axis
+ *		3. Unnormalized rotation
+ */
+UCLASS()
+class UTransformHeadsUpWidgetFactory : public UTypedElementDataStorageFactory
+{
+	GENERATED_BODY()
+
+public:
+	~UTransformHeadsUpWidgetFactory() override = default;
+
+	void RegisterQueries(ITypedElementDataStorageInterface& DataStorage) override;
+	void RegisterWidgetConstructors(ITypedElementDataStorageInterface& DataStorage,
+		ITypedElementDataStorageUiInterface& DataStorageUi) const override;
+};
+
+USTRUCT()
+struct FTransformHeadsUpWidgetConstructor : public FTypedElementWidgetConstructor
+{
+	GENERATED_BODY()
+
+public:
+	FTransformHeadsUpWidgetConstructor();
+	~FTransformHeadsUpWidgetConstructor() override = default;
+
+	TConstArrayView<const UScriptStruct*> GetAdditionalColumnsList() const override;
+
+protected:
+	TSharedPtr<SWidget> CreateWidget(const TypedElementDataStorage::FMetaDataView& Arguments) override;
+	bool FinalizeWidget(
+		ITypedElementDataStorageInterface* DataStorage,
+		ITypedElementDataStorageUiInterface* DataStorageUi,
+		TypedElementRowHandle Row,
+		const TSharedPtr<SWidget>& Widget) override;
+};
+
+USTRUCT(meta = (DisplayName = "Heads up display for transforms widget"))
+struct FTransformHeadsUpWidgetTag : public FTypedElementDataStorageTag
+{
+	GENERATED_BODY()
+};
diff --git a/Engine/Plugins/Experimental/EditorDataStorageFeatures/Source/TedsOutliner/Public/Widgets/TypeInfoWidget.h b/Engine/Plugins/Experimental/EditorDataStorageFeatures/Source/TedsOutliner/Public/Widgets/TypeInfoWidget.h
new file mode 100644
--- /dev/null
+++ b/Engine/Plugins/Experimental/EditorDataStorageFeatures/Source/TedsOutliner/Public/Widgets/TypeInfoWidget.h
@@ -0,0 +1,53 @@
+// Copyright Epic Games, Inc. All Rights Reserved.
+
+#pragma once
+
+#include "Elements/Interfaces/TypedElementDataStorageFactory.h"
+#include "Elements/Interfaces/TypedElementDataStorageInterface.h"
+#include "Elements/Interfaces/TypedElementDataStorageUiInterface.h"
+
+#include "TypeInfoWidget.generated.h"
+
+struct FSlateBrush;
+struct FTypedElementClassTypeInfoColumn;
+
+UCLASS()
+class UTypeInfoWidgetFactory : public UTypedElementDataStorageFactory
+{
+	GENERATED_BODY()
+
+public:
+	~UTypeInfoWidgetFactory() override = default;
+
+	void RegisterWidgetConstructors(ITypedElementDataStorageInterface& DataStorage,
+		ITypedElementDataStorageUiInterface& DataStorageUi) const override;
+};
+
+USTRUCT()
+struct TEDSOUTLINER_API FTypeInfoWidgetConstructor : public FTypedElementWidgetConstructor
+{
+	GENERATED_BODY()
+
+public:
+	FTypeInfoWidgetConstructor();
+	~FTypeInfoWidgetConstructor() override = default;
+
+protected:
+	explicit FTypeInfoWidgetConstructor(const UScriptStruct* InTypeInfo);
+	TSharedPtr<SWidget> CreateWidget(const TypedElementDataStorage::FMetaDataView& Arguments) override;
+	bool FinalizeWidget(ITypedElementDataStorageInterface* DataStorage, ITypedElementDataStorageUiInterface* DataStorageUi,
+		TypedElementRowHandle Row, const TSharedPtr<SWidget>& Widget) override;
+
+protected:
+
+	// Get the icon for a given row, checking the cache to see if we already have one
+	static const FSlateBrush* GetIconForRow(ITypedElementDataStorageInterface* DataStorage, TypedElementRowHandle Row, const FTypedElementClassTypeInfoColumn* TypeInfoColumn);
+	
+protected:
+
+	// Cache to avoid looking up the icon for a class every time
+	static TMap<FName, const FSlateBrush*> CachedIconMap;
+
+	// Whether the widget created by this constructor should be icon or text
+	bool bUseIcon;
+};
diff --git a/Engine/Plugins/Experimental/EditorDataStorageFeatures/Source/TedsOutliner/Public/Widgets/TypedElementAlertWidget.h b/Engine/Plugins/Experimental/EditorDataStorageFeatures/Source/TedsOutliner/Public/Widgets/TypedElementAlertWidget.h
deleted file mode 100644
diff --git a/Engine/Plugins/Experimental/EditorDataStorageFeatures/Source/TedsOutliner/Public/Widgets/TypedElementTransformHeadsUpWidget.h b/Engine/Plugins/Experimental/EditorDataStorageFeatures/Source/TedsOutliner/Public/Widgets/TypedElementTransformHeadsUpWidget.h
deleted file mode 100644
diff --git a/Engine/Plugins/Experimental/EditorDataStorageFeatures/Source/TedsOutliner/Public/Widgets/TypedElementTypeInfoWidget.h b/Engine/Plugins/Experimental/EditorDataStorageFeatures/Source/TedsOutliner/Public/Widgets/TypedElementTypeInfoWidget.h
deleted file mode 100644
diff --git a/Engine/Plugins/Experimental/EditorDataStorageFeatures/Source/TedsPropertyEditor/Private/TedsRowPickingMode.cpp b/Engine/Plugins/Experimental/EditorDataStorageFeatures/Source/TedsPropertyEditor/Private/TedsRowPickingMode.cpp
--- a/Engine/Plugins/Experimental/EditorDataStorageFeatures/Source/TedsPropertyEditor/Private/TedsRowPickingMode.cpp
+++ b/Engine/Plugins/Experimental/EditorDataStorageFeatures/Source/TedsPropertyEditor/Private/TedsRowPickingMode.cpp
@@ -2,10 +2,10 @@
 
 #include "TedsRowPickingMode.h"
 
-#include "TypedElementOutlinerItem.h"
+#include "TedsOutlinerItem.h"
 
-FTedsRowPickingMode::FTedsRowPickingMode(const FTypedElementOutlinerModeParams& Params, FOnSceneOutlinerItemPicked OnItemPickedDelegate)
-	: FTypedElementOutlinerMode(Params)
+FTedsRowPickingMode::FTedsRowPickingMode(const FTedsOutlinerParams& Params, FOnSceneOutlinerItemPicked OnItemPickedDelegate)
+	: FTedsOutlinerMode(Params)
 	, OnItemPicked(OnItemPickedDelegate)
 {
 
@@ -32,11 +32,11 @@ void FTedsRowPickingMode::OnFilterTextCommited(FSceneOutlinerItemSelection& Sele
 	// In picking mode, we check to see if we have any Typed Element items, and if so, fire
 	// off the notification to whoever is listening. This may often cause the widget itself
 	// to be enqueued for destruction
-	TArray<FTypedElementOutlinerTreeItem*> TypedElements;
-	Selection.Get(TypedElements);
-	if (TypedElements.Num() == 1 && TypedElements[0])
+	TArray<FTedsOutlinerTreeItem*> OutlinerRows;
+	Selection.Get(OutlinerRows);
+	if (OutlinerRows.Num() == 1 && OutlinerRows[0])
 	{
 		// Signal that a Typed Element was selected. We assume it is valid as it won't have been added to Selection if not.
-		SceneOutliner->SetItemSelection(TypedElements[0]->AsShared(), true, ESelectInfo::OnKeyPress);
+		SceneOutliner->SetItemSelection(OutlinerRows[0]->AsShared(), true, ESelectInfo::OnKeyPress);
 	}
 }
\ No newline at end of file
diff --git a/Engine/Plugins/Experimental/EditorDataStorageFeatures/Source/TedsPropertyEditor/Private/Widgets/SPropertyMenuTedsRowPicker.cpp b/Engine/Plugins/Experimental/EditorDataStorageFeatures/Source/TedsPropertyEditor/Private/Widgets/SPropertyMenuTedsRowPicker.cpp
--- a/Engine/Plugins/Experimental/EditorDataStorageFeatures/Source/TedsPropertyEditor/Private/Widgets/SPropertyMenuTedsRowPicker.cpp
+++ b/Engine/Plugins/Experimental/EditorDataStorageFeatures/Source/TedsPropertyEditor/Private/Widgets/SPropertyMenuTedsRowPicker.cpp
@@ -4,13 +4,14 @@
 
 #include "Elements/Framework/TypedElementRegistry.h"
 #include "TedsRowPickingMode.h"
-#include "TypedElementOutlinerItem.h"
+#include "TedsOutlinerItem.h"
 
 #define LOCTEXT_NAMESPACE "TedsPropertyEditor"
+
 void SPropertyMenuTedsRowPicker::Construct(const FArguments& InArgs)
 {
 	bAllowClear = InArgs._AllowClear;
-	TypedElementQueryFilter = InArgs._TypedElementQueryFilter;
+	QueryFilter = InArgs._QueryFilter;
 	ElementFilter = InArgs._ElementFilter;
 	OnSet = InArgs._OnSet;
 
@@ -51,7 +52,7 @@ void SPropertyMenuTedsRowPicker::Construct(const FArguments& InArgs)
 			{
 				auto OnItemPicked = FOnSceneOutlinerItemPicked::CreateLambda([&](TSharedRef<ISceneOutlinerTreeItem> Item)
 					{
-						if (FTypedElementOutlinerTreeItem* ElementItem = Item->CastTo<FTypedElementOutlinerTreeItem>())
+						if (FTedsOutlinerTreeItem* ElementItem = Item->CastTo<FTedsOutlinerTreeItem>())
 						{
 							if (ElementItem->IsValid())
 							{
@@ -67,16 +68,16 @@ void SPropertyMenuTedsRowPicker::Construct(const FArguments& InArgs)
 
 				InitOptions.ModeFactory = FCreateSceneOutlinerMode::CreateLambda([&](SSceneOutliner* Outliner)
 					{
-						FTypedElementOutlinerModeParams Params(Outliner);
-						Params.QueryDescription = TypedElementQueryFilter;
+						FTedsOutlinerParams Params(Outliner);
+						Params.QueryDescription = QueryFilter;
 						return new FTedsRowPickingMode(Params, OnItemPicked);
 					});
 
 				TSharedPtr<SSceneOutliner> Outliner = SNew(SSceneOutliner, InitOptions);
 
 				Outliner->AddFilter(
-					MakeShared<TSceneOutlinerPredicateFilter<FTypedElementOutlinerTreeItem>>(
-						FTypedElementOutlinerTreeItem::FFilterPredicate::CreateLambda([this](const TypedElementDataStorage::RowHandle RowHandle) -> bool
+					MakeShared<TSceneOutlinerPredicateFilter<FTedsOutlinerTreeItem>>(
+						FTedsOutlinerTreeItem::FFilterPredicate::CreateLambda([this](const TypedElementDataStorage::RowHandle RowHandle) -> bool
 					{
 						if (ElementFilter.IsBound())
 						{
diff --git a/Engine/Plugins/Experimental/EditorDataStorageFeatures/Source/TedsPropertyEditor/Public/TedsRowPickingMode.h b/Engine/Plugins/Experimental/EditorDataStorageFeatures/Source/TedsPropertyEditor/Public/TedsRowPickingMode.h
--- a/Engine/Plugins/Experimental/EditorDataStorageFeatures/Source/TedsPropertyEditor/Public/TedsRowPickingMode.h
+++ b/Engine/Plugins/Experimental/EditorDataStorageFeatures/Source/TedsPropertyEditor/Public/TedsRowPickingMode.h
@@ -2,7 +2,9 @@
 
 #pragma once
 
-#include "TypedElementOutlinerMode.h"
+#include "TedsOutlinerMode.h"
+#include "SceneOutlinerFwd.h"
+#include "Elements/Common/TypedElementHandles.h"
 
 DECLARE_DELEGATE_OneParam(FOnTedsRowSelected, TypedElementDataStorage::RowHandle RowHandle);
 DECLARE_DELEGATE_RetVal_OneParam(bool, FOnShouldFilterTedsRow, const TypedElementDataStorage::RowHandle RowHandle);
@@ -10,10 +12,10 @@ DECLARE_DELEGATE_RetVal_OneParam(bool, FOnShouldFilterTedsRow, const TypedElemen
 /*
 * Picking mode for TEDs Scene Outliner Widgets. Based off of FActorPickingMode
 */
-class TEDSPROPERTYEDITOR_API FTedsRowPickingMode : public FTypedElementOutlinerMode
+class TEDSPROPERTYEDITOR_API FTedsRowPickingMode : public FTedsOutlinerMode
 {
 public:
-	FTedsRowPickingMode(const FTypedElementOutlinerModeParams& Params, FOnSceneOutlinerItemPicked OnItemPickedDelegate);
+	FTedsRowPickingMode(const FTedsOutlinerParams& Params, FOnSceneOutlinerItemPicked OnItemPickedDelegate);
 
 	virtual ~FTedsRowPickingMode() = default;
 public:
diff --git a/Engine/Plugins/Experimental/EditorDataStorageFeatures/Source/TedsPropertyEditor/Public/Widgets/SPropertyMenuTedsRowPicker.h b/Engine/Plugins/Experimental/EditorDataStorageFeatures/Source/TedsPropertyEditor/Public/Widgets/SPropertyMenuTedsRowPicker.h
--- a/Engine/Plugins/Experimental/EditorDataStorageFeatures/Source/TedsPropertyEditor/Public/Widgets/SPropertyMenuTedsRowPicker.h
+++ b/Engine/Plugins/Experimental/EditorDataStorageFeatures/Source/TedsPropertyEditor/Public/Widgets/SPropertyMenuTedsRowPicker.h
@@ -13,7 +13,7 @@ public:
 		, _ElementFilter()
 	{}
 		SLATE_ARGUMENT(bool, AllowClear)
-		SLATE_ARGUMENT(TypedElementDataStorage::FQueryDescription, TypedElementQueryFilter)
+		SLATE_ARGUMENT(TypedElementDataStorage::FQueryDescription, QueryFilter)
 		SLATE_ARGUMENT(FOnShouldFilterTedsRow, ElementFilter)
 		SLATE_EVENT(FOnTedsRowSelected, OnSet)
 		SLATE_EVENT(FSimpleDelegate, OnClose)
@@ -32,7 +32,7 @@ private:
 private:
 	bool bAllowClear;
 
-	TypedElementDataStorage::FQueryDescription TypedElementQueryFilter;
+	TypedElementDataStorage::FQueryDescription QueryFilter;
 
 	FOnShouldFilterTedsRow ElementFilter;
 
diff --git a/Engine/Source/Runtime/TypedElementFramework/Public/Elements/Columns/TypedElementHiearchyColumns.h b/Engine/Source/Runtime/TypedElementFramework/Public/Elements/Columns/TypedElementHiearchyColumns.h
--- a/Engine/Source/Runtime/TypedElementFramework/Public/Elements/Columns/TypedElementHiearchyColumns.h
+++ b/Engine/Source/Runtime/TypedElementFramework/Public/Elements/Columns/TypedElementHiearchyColumns.h
@@ -12,7 +12,7 @@
  * A reference to the direct hierarchical parent of this row.
  */
 USTRUCT(meta = (DisplayName = "Parent"))
-struct FTypedElementParentColumn final : public FTypedElementDataStorageColumn
+struct FTableRowParentColumn final : public FTypedElementDataStorageColumn
 {
 	GENERATED_BODY()
 
@@ -24,7 +24,7 @@ struct FTypedElementParentColumn final : public FTypedElementDataStorageColumn
  * be used to attempt to find the indexed row. This column can not be used to find rows that are not indexed.
  */
 USTRUCT(meta = (DisplayName = "Parent (Unresolved)"))
-struct FTypedElementUnresolvedParentColumn final : public FTypedElementDataStorageColumn
+struct FUnresolvedTableRowParentColumn final : public FTypedElementDataStorageColumn
 {
 	GENERATED_BODY()
 
