First optimization pass to the Teds AssetDataCBDataSource.
Removed a large series of micro stall when populating TEDS as the same time the asset registry is still doing it's initial asset scan.
On a large project those represented almost 2 m 40 seconds of waiting for the lock to be free.

diff --git a/Engine/Plugins/Experimental/EditorDataStorageFeatures/Source/TedsAssetData/Private/CB/TedsAssetDataCBDataSource.cpp b/Engine/Plugins/Experimental/EditorDataStorageFeatures/Source/TedsAssetData/Private/CB/TedsAssetDataCBDataSource.cpp
--- a/Engine/Plugins/Experimental/EditorDataStorageFeatures/Source/TedsAssetData/Private/CB/TedsAssetDataCBDataSource.cpp
+++ b/Engine/Plugins/Experimental/EditorDataStorageFeatures/Source/TedsAssetData/Private/CB/TedsAssetDataCBDataSource.cpp
@@ -98,7 +98,7 @@ FTedsAssetDataCBDataSource::FTedsAssetDataCBDataSource(ITypedElementDataStorageI
 
 
 	// For now just add the columns one by one but this should be rework to work in batch
-	auto AddAssetDataColumns = [DataSource = static_cast<const FTedsAssetDataCBDataSource*>(this), AssetRegistry = static_cast<const IAssetRegistry*>(&IAssetRegistry::GetChecked())](const FAssetData& InAssetData, TypedElementDataStorage::IQueryContext& InContext, TypedElementDataStorage::RowHandle Row)
+	auto AddAssetDataColumns = [](TypedElementDataStorage::IQueryContext& InContext, TypedElementDataStorage::RowHandle Row, const FAssetData& InAssetData, const FAssetPackageData* PackageData)
 	{
 		// Not optimized at all but we would like to have the data in sooner for testing purposes.
 		
@@ -112,7 +112,7 @@ FTedsAssetDataCBDataSource::FTedsAssetDataCBDataSource(ITypedElementDataStorageI
 			InContext.AddColumns<FAssetTag, FPublicAssetTag>(Row);
 		}
 
-		if (TOptional<FAssetPackageData> PackageData = AssetRegistry->GetAssetPackageDataCopy(InAssetData.PackageName))
+		if (PackageData)
 		{
 			FDiskSizeColumn DiskSizeColumn;
 			DiskSizeColumn.DiskSize = PackageData->DiskSize;
@@ -121,29 +121,51 @@ FTedsAssetDataCBDataSource::FTedsAssetDataCBDataSource(ITypedElementDataStorageI
 
 		FItemNameColumn_Experimental ItemNameColumn;
 		ItemNameColumn.Name = InAssetData.AssetName;
+		InContext.AddColumn(Row, MoveTemp(ItemNameColumn));
 	};
 
 	ProcessAssetDataAndPathUpdateQuery = Database.RegisterQuery(
 		Select(
 			TEXT("FTedsAssetDataCBDataSource: Process Asset Data and Path Updates"),
 			FProcessor(TypedElementDataStorage::EQueryTickPhase::DuringPhysics, Database.GetQueryTickGroupName(TypedElementDataStorage::EQueryTickGroups::Update)),
-			[GenerateVirtualPaths, AddAssetDataColumns](TypedElementDataStorage::IQueryContext& Context, const TypedElementDataStorage::RowHandle Row, const FAssetDataColumn_Experimental& AssetDataColumn)
+			[GenerateVirtualPaths, AddAssetDataColumns, AssetRegistry = static_cast<const IAssetRegistry*>(&IAssetRegistry::GetChecked())](TypedElementDataStorage::IQueryContext& Context, const TypedElementDataStorage::RowHandle* Rows, const FAssetDataColumn_Experimental* AssetDataColumn)
 			{
+				const int32 RowCount = Context.GetRowCount();
+
+				TArray<FName, TInlineAllocator<32>> PackageNames;
+				TArray<TPair<TypedElementDataStorage::RowHandle, const FAssetData*>, TInlineAllocator<32>> RowsAndAssetData;
+				PackageNames.Reserve(RowCount);
+				RowsAndAssetData.Reserve(RowCount);
 
 				FNameBuilder InternalPath;
 				FNameBuilder VirtualPath;
 
-				InternalPath.Reset();
-				AssetDataColumn.AssetData.AppendObjectPath(InternalPath);
+				for (int32 Index = 0; Index < RowCount; ++Index)
+				{ 
+					if (GenerateVirtualPaths(InternalPath, VirtualPath))
+					{
+						const FAssetData& AssetData = AssetDataColumn[Index].AssetData;
+						const TypedElementDataStorage::RowHandle Row = Rows[Index];
 
-				if (GenerateVirtualPaths(InternalPath, VirtualPath))
-				{
-					FVirtualPathColumn_Experimental VirtualPathColumn;
-					VirtualPathColumn.VirtualPath = *VirtualPath;
-					// Todo investigate for a batch add Maybe?
-					Context.AddColumn(Row, MoveTemp(VirtualPathColumn));
+						InternalPath.Reset();
+						AssetData.AppendObjectPath(InternalPath);
+						FVirtualPathColumn_Experimental VirtualPathColumn;
+						VirtualPathColumn.VirtualPath = *VirtualPath;
+						// Todo investigate for a batch add Maybe?
+						Context.AddColumn(Row, MoveTemp(VirtualPathColumn));
 
-					AddAssetDataColumns(AssetDataColumn.AssetData, Context, Row);
+						PackageNames.Add(AssetData.PackageName);
+						RowsAndAssetData.Emplace(Row, &AssetData);
+					}
+				}
+
+				TArray<TOptional<FAssetPackageData>> AssetPackageDatas = AssetRegistry->GetAssetPackageDatasCopy(PackageNames);
+
+				for (int32 Index = 0; Index < AssetPackageDatas.Num(); ++Index)
+				{
+					const TPair<TypedElementDataStorage::RowHandle, const FAssetData*>& Pair = RowsAndAssetData[Index];
+					const TOptional<FAssetPackageData>& AssetPackageData = AssetPackageDatas[Index];
+					AddAssetDataColumns(Context, Pair.Key, *Pair.Value, AssetPackageData.GetPtrOrNull());
 				}
 			}
 		)
@@ -156,9 +178,23 @@ FTedsAssetDataCBDataSource::FTedsAssetDataCBDataSource(ITypedElementDataStorageI
 		Select(
 			TEXT("FTedsAssetDataCBDataSource: Process Asset Data updates"),
 			FProcessor(TypedElementDataStorage::EQueryTickPhase::DuringPhysics, Database.GetQueryTickGroupName(TypedElementDataStorage::EQueryTickGroups::Update)),
-			[AddAssetDataColumns](TypedElementDataStorage::IQueryContext& Context, const TypedElementDataStorage::RowHandle Row, const FAssetDataColumn_Experimental& AssetDataColumn)
+			[AddAssetDataColumns,  AssetRegistry = static_cast<const IAssetRegistry*>(&IAssetRegistry::GetChecked())](TypedElementDataStorage::IQueryContext& Context, const TypedElementDataStorage::RowHandle* Rows, const FAssetDataColumn_Experimental* AssetDataColumn)
 			{
-				AddAssetDataColumns(AssetDataColumn.AssetData, Context, Row);
+				const int32 RowCount = Context.GetRowCount();
+				TArray<FName, TInlineAllocator<32>> PackageNames;
+				PackageNames.Reserve(RowCount);
+
+				for (int32 Index = 0; Index < RowCount; ++Index)
+				{
+					PackageNames.Add(AssetDataColumn[Index].AssetData.PackageName);
+				}
+
+				TArray<TOptional<FAssetPackageData>> AssetPackageDatas = AssetRegistry->GetAssetPackageDatasCopy(PackageNames);
+
+				for (int32 Index = 0; Index < RowCount; ++Index)
+				{
+					AddAssetDataColumns(Context, Rows[Index], AssetDataColumn[Index].AssetData, AssetPackageDatas[Index].GetPtrOrNull());
+				}
 			}
 		)
 		.Where()
