STedsTableViewer: Fix conflicting LOCTEXT macro keys

diff --git a/Engine/Plugins/Experimental/EditorDataStorageFeatures/Source/TedsTableViewer/Private/Widgets/STedsTableViewer.cpp b/Engine/Plugins/Experimental/EditorDataStorageFeatures/Source/TedsTableViewer/Private/Widgets/STedsTableViewer.cpp
--- a/Engine/Plugins/Experimental/EditorDataStorageFeatures/Source/TedsTableViewer/Private/Widgets/STedsTableViewer.cpp
+++ b/Engine/Plugins/Experimental/EditorDataStorageFeatures/Source/TedsTableViewer/Private/Widgets/STedsTableViewer.cpp
@@ -51,7 +51,7 @@ namespace UE::EditorDataStorage
 					.VAlign(VAlign_Center)
 					[
 						SNew(STextBlock)
-							.Text(LOCTEXT("EmptyTableViewerText", "The input query has no results"))
+							.Text(LOCTEXT("EmptyTableViewerQueryText", "The input query has no results"))
 					]
 			];
 		}
@@ -64,7 +64,7 @@ namespace UE::EditorDataStorage
 					.VAlign(VAlign_Center)
 					[
 						SNew(STextBlock)
-							.Text(LOCTEXT("EmptyTableViewerText", "There were no columns specified to display"))
+							.Text(LOCTEXT("EmptyTableViewerColumnsText", "There were no columns specified to display"))
 					]
 			];
 		}
