Teds UI: Create a new SRowDetails widget to show the columns/tags on a row
Teds Debugger: Add an SRowDetails widget to show details on the row selected in the table viewer
Teds: Add a template specialization for TypedElementDataStorage::RowHandle to be used in widgets like SListView
Teds UI: Remove the label column from widgets

diff --git a/Engine/Plugins/Experimental/EditorDataStorage/Source/TedsUI/Private/Widgets/LabelWidget.cpp b/Engine/Plugins/Experimental/EditorDataStorage/Source/TedsUI/Private/Widgets/LabelWidget.cpp
--- a/Engine/Plugins/Experimental/EditorDataStorage/Source/TedsUI/Private/Widgets/LabelWidget.cpp
+++ b/Engine/Plugins/Experimental/EditorDataStorage/Source/TedsUI/Private/Widgets/LabelWidget.cpp
@@ -216,7 +216,6 @@ TSharedPtr<SWidget> FLabelWidgetConstructor::Construct(
 		{
 			if (FinalizeWidget(DataStorage, DataStorageUi, Row, Result))
 			{
-				SetupDebugColumns(Row, DataStorage, DataStorageUi, Result);
 				return Result;
 			}
 		}
diff --git a/Engine/Plugins/Experimental/EditorDataStorageFeatures/Source/TedsDebugger/Private/Widgets/QueryEditor/TedsQueryEditorResultsView.cpp b/Engine/Plugins/Experimental/EditorDataStorageFeatures/Source/TedsDebugger/Private/Widgets/QueryEditor/TedsQueryEditorResultsView.cpp
--- a/Engine/Plugins/Experimental/EditorDataStorageFeatures/Source/TedsDebugger/Private/Widgets/QueryEditor/TedsQueryEditorResultsView.cpp
+++ b/Engine/Plugins/Experimental/EditorDataStorageFeatures/Source/TedsDebugger/Private/Widgets/QueryEditor/TedsQueryEditorResultsView.cpp
@@ -12,6 +12,7 @@
 #include "QueryEditor/TedsQueryEditorModel.h"
 #include "QueryStack/FQueryStackNode_RowView.h"
 #include "Widgets/STedsTableViewer.h"
+#include "Widgets/SRowDetails.h"
 #include "Widgets/Text/STextBlock.h"
 
 #define LOCTEXT_NAMESPACE "TedsDebuggerModule"
@@ -31,23 +32,43 @@ void SResultsView::Construct(const FArguments& InArgs, FTedsQueryEditorModel& In
 
 	// Create a custom column for the table viewer to display row handles
 	CreateRowHandleColumn();
-
-	RowQueryStack = MakeShared<UE::EditorDataStorage::FQueryStackNode_RowView>(&TableViewerRows);
+	
+	RowQueryStack = MakeShared<FQueryStackNode_RowView>(&TableViewerRows);
 
 	ChildSlot
 	[
 		SNew(SVerticalBox)
 		+SVerticalBox::Slot()
-		.AutoHeight()
 		[
-			SNew(SWarningOrErrorBox)
-			.MessageStyle(EMessageStyle::Warning)
-			.Message(LOCTEXT("QueryEditor_TableViewUnreliable", "Integration of table view is WIP. Currently Unreliable"))
-		]
-		+SVerticalBox::Slot()
-		[
-			SAssignNew(TableViewer, UE::EditorDataStorage::STedsTableViewer)
-			.QueryStack(RowQueryStack)
+			SNew(SSplitter)
+			+SSplitter::Slot()
+			.Value(0.5f)
+			[
+				SAssignNew(TableViewer, UE::EditorDataStorage::STedsTableViewer)
+				.QueryStack(RowQueryStack)
+				.OnSelectionChanged(STedsTableViewer::FOnSelectionChanged::CreateLambda(
+					[this](TypedElementDataStorage::RowHandle RowHandle)
+						{
+							if(RowDetailsWidget)
+							{
+								if(RowHandle != TypedElementDataStorage::InvalidRowHandle)
+								{
+									RowDetailsWidget->SetRow(RowHandle);
+								}
+								else
+								{
+									RowDetailsWidget->ClearRow();
+								}
+								
+							}
+						}))
+			]
+			+SSplitter::Slot()
+			.Value(0.5f)
+			[
+				SAssignNew(RowDetailsWidget, UE::EditorDataStorage::SRowDetails)
+			]
+			
 		]
 		+SVerticalBox::Slot()
 		.AutoHeight()
diff --git a/Engine/Plugins/Experimental/EditorDataStorageFeatures/Source/TedsDebugger/Private/Widgets/QueryEditor/TedsQueryEditorResultsView.h b/Engine/Plugins/Experimental/EditorDataStorageFeatures/Source/TedsDebugger/Private/Widgets/QueryEditor/TedsQueryEditorResultsView.h
--- a/Engine/Plugins/Experimental/EditorDataStorageFeatures/Source/TedsDebugger/Private/Widgets/QueryEditor/TedsQueryEditorResultsView.h
+++ b/Engine/Plugins/Experimental/EditorDataStorageFeatures/Source/TedsDebugger/Private/Widgets/QueryEditor/TedsQueryEditorResultsView.h
@@ -4,6 +4,11 @@
 #include "Elements/Interfaces/TypedElementDataStorageInterface.h"
 #include "Widgets/SCompoundWidget.h"
 
+namespace UE::EditorDataStorage
+{
+	class SRowDetails;
+}
+
 class ISceneOutliner;
 class SSceneOutliner;
 class SHorizontalBox;
@@ -51,6 +56,10 @@ namespace UE::EditorDataStorage::Debug::QueryEditor
 
 		// Custom column for the table viewer to display row handles
 		TSharedPtr<UE::EditorDataStorage::FTedsTableViewerColumn> RowHandleColumn;
+
+		// Widget that displays details of a row
+		TSharedPtr<UE::EditorDataStorage::SRowDetails> RowDetailsWidget;
+
 	};
 
 	
diff --git a/Engine/Plugins/Experimental/EditorDataStorageFeatures/Source/TedsTableViewer/Private/TedsTableViewerModel.cpp b/Engine/Plugins/Experimental/EditorDataStorageFeatures/Source/TedsTableViewer/Private/TedsTableViewerModel.cpp
--- a/Engine/Plugins/Experimental/EditorDataStorageFeatures/Source/TedsTableViewer/Private/TedsTableViewerModel.cpp
+++ b/Engine/Plugins/Experimental/EditorDataStorageFeatures/Source/TedsTableViewer/Private/TedsTableViewerModel.cpp
@@ -44,7 +44,7 @@ namespace UE::EditorDataStorage
 		
 		for(TypedElementDataStorage::RowHandle RowHandle : RowQueryStack->GetOrderedRowList())
 		{
-			Items.Add(MakeShared<TedsTableViewerItem>(RowHandle));
+			Items.Add(RowHandle);
 		}
 
 		CachedRowQueryStackRevision = RowQueryStack->GetRevisionId();
@@ -240,7 +240,7 @@ namespace UE::EditorDataStorage
 		// Table Viewer TODO: We can probably store a map of the items instead but this works for now
 		for(const TableViewerItemPtr& Item : Items)
 		{
-			if(Item->RowHandle == InRowHandle)
+			if(Item == InRowHandle)
 			{
 				return IsItemVisible.Execute(Item);
 			}
diff --git a/Engine/Plugins/Experimental/EditorDataStorageFeatures/Source/TedsTableViewer/Private/Widgets/SRowDetails.cpp b/Engine/Plugins/Experimental/EditorDataStorageFeatures/Source/TedsTableViewer/Private/Widgets/SRowDetails.cpp
new file mode 100644
--- /dev/null
+++ b/Engine/Plugins/Experimental/EditorDataStorageFeatures/Source/TedsTableViewer/Private/Widgets/SRowDetails.cpp
@@ -0,0 +1,186 @@
+// Copyright Epic Games, Inc. All Rights Reserved.
+
+#include "Widgets/SRowDetails.h"
+
+#include "Elements/Columns/TypedElementMiscColumns.h"
+#include "Elements/Columns/TypedElementTypeInfoColumns.h"
+#include "Elements/Framework/TypedElementRegistry.h"
+#include "Elements/Interfaces/TypedElementDataStorageInterface.h"
+#include "Elements/Interfaces/TypedElementDataStorageUiInterface.h"
+#include "Widgets/Text/STextBlock.h"
+#include "Widgets/Views/SListView.h"
+
+#define LOCTEXT_NAMESPACE "SRowDetails"
+
+namespace UE::SRowDetails::Local
+{
+	static const FName NameColumn(TEXT("Name"));
+	static const FName DataColumn(TEXT("Data"));
+	static TArray<FName> DefaultWidgetPurposes = { FName(TEXT("RowDetails.Cell.Large")),
+		FName(TEXT("RowDetails.Cell")), FName(TEXT("General.Cell.Large")), FName(TEXT("General.Cell")) };
+}
+
+namespace UE::EditorDataStorage
+{
+	//
+	// SRowDetails
+	//
+	
+	void SRowDetails::Construct(const FArguments& InArgs)
+	{
+		bShowAllDetails = InArgs._ShowAllDetails;
+		WidgetPurposes = InArgs._WidgetPurposesOverride;
+
+		if(WidgetPurposes.IsEmpty())
+		{
+			WidgetPurposes = UE::SRowDetails::Local::DefaultWidgetPurposes;
+		}
+
+		UTypedElementRegistry* Registry = UTypedElementRegistry::GetInstance();
+		checkf(Registry, TEXT("Unable to initialize SRowDetails before TEDS is initialized."));
+		checkf(Registry->AreDataStorageInterfacesSet(), TEXT("Unable to initialize SRowDetails without the editor data storage interfaces."));
+
+		DataStorage = Registry->GetMutableDataStorage();
+		DataStorageUi = Registry->GetMutableDataStorageUi();
+
+		
+		ChildSlot
+		[
+			SAssignNew(ListView, SListView<RowDetailsItemPtr>)
+				.ListItemsSource(&Items)
+				.OnGenerateRow(this, &SRowDetails::CreateRow)
+				.Visibility_Lambda([this]()
+					{
+						return Items.IsEmpty() ? EVisibility::Hidden : EVisibility::Visible;
+					})
+				.HeaderRow
+				(
+					SNew(SHeaderRow)
+					+SHeaderRow::Column(UE::SRowDetails::Local::NameColumn)
+						.DefaultLabel(FText::FromString(TEXT("Name")))
+						.FillWidth(0.3f)
+					+SHeaderRow::Column(UE::SRowDetails::Local::DataColumn)
+						.DefaultLabel(FText::FromString(TEXT("Value")))
+						.FillWidth(0.7f)
+				)
+		];
+	}
+	
+	void SRowDetails::SetRow(TypedElementDataStorage::RowHandle Row)
+	{
+		if (DataStorage->IsRowAssigned(Row))
+		{
+			Items.Reset();
+			
+			TArray<TWeakObjectPtr<const UScriptStruct>> Columns;
+			DataStorage->ListColumns(Row, [&Columns](const UScriptStruct& ColumnType)
+				{
+					Columns.Emplace(&ColumnType);
+					return true;
+				});
+			
+			for (const FName& Purpose : WidgetPurposes)
+			{
+				DataStorageUi->CreateWidgetConstructors(Purpose, ITypedElementDataStorageUiInterface::EMatchApproach::LongestMatch,
+					Columns, {}, [this, Row](
+						TUniquePtr<FTypedElementWidgetConstructor> Constructor, TConstArrayView<TWeakObjectPtr<const UScriptStruct>> Columns)
+					{
+						Items.Add(MakeShared<FRowDetailsItem>(nullptr, MoveTemp(Constructor), Row));
+						return true;
+					});
+			}
+			
+			if (bShowAllDetails)
+			{
+				// Create defaults for the remaining widgets.
+				for (TWeakObjectPtr<const UScriptStruct> Column : Columns)
+				{
+					DataStorageUi->CreateWidgetConstructors(FName(TEXT("General.Cell.Default")), {},
+						[this, Column, Row](
+							TUniquePtr<FTypedElementWidgetConstructor> Constructor, TConstArrayView<TWeakObjectPtr<const UScriptStruct>> Columns)
+						{
+							Items.Add(MakeShared<FRowDetailsItem>(Column, MoveTemp(Constructor), Row));
+							return true;
+						});
+				}
+			}
+			
+			ListView->RequestListRefresh();
+		}
+		else
+		{
+			ClearRow();
+		}
+	}
+	
+	void SRowDetails::ClearRow()
+	{
+		Items.Reset();
+		ListView->RequestListRefresh();
+	}
+	
+	TSharedRef<ITableRow> SRowDetails::CreateRow(RowDetailsItemPtr InItem, const TSharedRef<STableViewBase>& OwnerTable)
+	{
+		TSharedPtr<SRowDetailsRow> Row = SNew(SRowDetailsRow, OwnerTable, DataStorage, DataStorageUi)
+			.Item(InItem);
+		
+		return Row.ToSharedRef();
+	}
+	
+	//
+	// SRowDetailsRow
+	//
+	void SRowDetailsRow::Construct(const FArguments& Args, const TSharedRef<STableViewBase>& OwnerTableView,ITypedElementDataStorageInterface* InDataStorage,
+			ITypedElementDataStorageUiInterface* InDataStorageUi)
+	{
+		Item = Args._Item;
+
+		DataStorage = InDataStorage;
+		DataStorageUi = InDataStorageUi;
+		
+		SMultiColumnTableRow<RowDetailsItemPtr>::Construct(FSuperRowType::FArguments(), OwnerTableView);
+	}
+	
+	TSharedRef<SWidget> SRowDetailsRow::GenerateWidgetForColumn(const FName& ColumnName)
+	{
+		using namespace TypedElementDataStorage;
+		
+		if (!DataStorage->IsRowAvailable(Item->WidgetRow))
+		{
+			Item->WidgetRow = DataStorage->AddRow(DataStorage->FindTable(FName("Editor_WidgetTable")));
+			
+			DataStorage->AddColumn<FTypedElementRowReferenceColumn>(Item->WidgetRow, FTypedElementRowReferenceColumn
+				{
+					.Row = Item->Row
+				});
+			
+			if (Item->ColumnType.IsValid() &&
+				Item->WidgetConstructor->GetAdditionalColumnsList().Contains(FTypedElementScriptStructTypeInfoColumn::StaticStruct()))
+			{
+				DataStorage->AddColumn(Item->WidgetRow, FTypedElementScriptStructTypeInfoColumn
+					{
+						.TypeInfo = Item->ColumnType
+					});
+			}
+		}
+		if (ColumnName == UE::SRowDetails::Local::NameColumn)
+		{
+			return SNew(STextBlock)
+					.Text(FText::FromString(Item->WidgetConstructor->CreateWidgetDisplayName(
+							DataStorage, Item->WidgetRow)));
+		}
+		else if (ColumnName == UE::SRowDetails::Local::DataColumn)
+		{
+			return DataStorageUi->ConstructWidget(Item->WidgetRow, *(Item->WidgetConstructor), {}).ToSharedRef();
+		}
+		else
+		{
+			return SNew(STextBlock)
+					.Text(LOCTEXT("InvalidColumnType", "Invalid Column Type")); 
+		}
+	}
+	
+} // namespace UE::EditorDataStorage
+
+
+#undef LOCTEXT_NAMESPACE // "SRowDetails"
diff --git a/Engine/Plugins/Experimental/EditorDataStorageFeatures/Source/TedsTableViewer/Private/Widgets/STedsTableViewer.cpp b/Engine/Plugins/Experimental/EditorDataStorageFeatures/Source/TedsTableViewer/Private/Widgets/STedsTableViewer.cpp
--- a/Engine/Plugins/Experimental/EditorDataStorageFeatures/Source/TedsTableViewer/Private/Widgets/STedsTableViewer.cpp
+++ b/Engine/Plugins/Experimental/EditorDataStorageFeatures/Source/TedsTableViewer/Private/Widgets/STedsTableViewer.cpp
@@ -13,6 +13,8 @@ namespace UE::EditorDataStorage
 {
 	void STedsTableViewer::Construct(const FArguments& InArgs)
 	{
+		OnSelectionChanged = InArgs._OnSelectionChanged;
+		
 		Model = MakeShared<FTedsTableViewerModel>(InArgs._QueryStack, InArgs._Columns, InArgs._CellWidgetPurposes,
 			FTedsTableViewerModel::FIsItemVisible::CreateSP(this, &STedsTableViewer::IsItemVisible));
 		
@@ -22,7 +24,9 @@ namespace UE::EditorDataStorage
 		ListView = SNew(SListView<TableViewerItemPtr>)
 			.HeaderRow(HeaderRowWidget)
 			.ListItemsSource(&Model->GetItems())
-			.OnGenerateRow(this, &STedsTableViewer::MakeTableRowWidget);
+			.OnGenerateRow(this, &STedsTableViewer::MakeTableRowWidget)
+			.OnSelectionChanged(this, &STedsTableViewer::OnListSelectionChanged)
+			.SelectionMode(ESelectionMode::Single); // We only support single selection for now in the table viewer
 		
 		AssignChildSlot();
 		
@@ -88,6 +92,14 @@ namespace UE::EditorDataStorage
 		AssignChildSlot();
 	}
 
+	void STedsTableViewer::OnListSelectionChanged(TableViewerItemPtr Item, ESelectInfo::Type SelectInfo)
+	{
+		if(OnSelectionChanged.IsBound())
+		{
+			OnSelectionChanged.Execute(Item);
+		}
+	}
+
 	void STedsTableViewer::SetColumns(const TArray<TWeakObjectPtr<const UScriptStruct>>& Columns)
 	{
 		Model->SetColumns(Columns);
diff --git a/Engine/Plugins/Experimental/EditorDataStorageFeatures/Source/TedsTableViewer/Private/Widgets/STedsTableViewerRow.cpp b/Engine/Plugins/Experimental/EditorDataStorageFeatures/Source/TedsTableViewer/Private/Widgets/STedsTableViewerRow.cpp
--- a/Engine/Plugins/Experimental/EditorDataStorageFeatures/Source/TedsTableViewer/Private/Widgets/STedsTableViewerRow.cpp
+++ b/Engine/Plugins/Experimental/EditorDataStorageFeatures/Source/TedsTableViewer/Private/Widgets/STedsTableViewerRow.cpp
@@ -23,7 +23,7 @@ namespace UE::EditorDataStorage
 	{
 		if(TSharedPtr<FTedsTableViewerColumn> Column = TableViewerModel->GetColumn(ColumnName))
 		{
-			if(TSharedPtr<SWidget> RowWidget = Column->ConstructRowWidget(Item->RowHandle))
+			if(TSharedPtr<SWidget> RowWidget = Column->ConstructRowWidget(Item))
 			{
 				return SNew(SBox)
 						.MinDesiredHeight(20.0f)
diff --git a/Engine/Plugins/Experimental/EditorDataStorageFeatures/Source/TedsTableViewer/Public/TedsTableViewerModel.h b/Engine/Plugins/Experimental/EditorDataStorageFeatures/Source/TedsTableViewer/Public/TedsTableViewerModel.h
--- a/Engine/Plugins/Experimental/EditorDataStorageFeatures/Source/TedsTableViewer/Public/TedsTableViewerModel.h
+++ b/Engine/Plugins/Experimental/EditorDataStorageFeatures/Source/TedsTableViewer/Public/TedsTableViewerModel.h
@@ -6,6 +6,7 @@
 #include "Delegates/Delegate.h"
 #include "Elements/Common/TypedElementHandles.h"
 #include "Templates/SharedPointer.h"
+#include "TypedElementUITypes.h"
 #include "UObject/NameTypes.h"
 
 class ITypedElementDataStorageCompatibilityInterface;
@@ -16,21 +17,9 @@ namespace UE::EditorDataStorage
 {
 	class FTedsTableViewerColumn;
 	class IQueryStackNode_Row;
-	
-	// Struct representing a row in the table viewer's UI
-	struct TedsTableViewerItem
-	{
-		TypedElementDataStorage::RowHandle RowHandle;
 
-		TedsTableViewerItem(TypedElementDataStorage::RowHandle InRowHandle)
-			: RowHandle(InRowHandle)
-		{
-			
-		}
-	};
-	
 	// Typedef for an item in the table viewer
-	typedef TSharedPtr<TedsTableViewerItem> TableViewerItemPtr;
+	using TableViewerItemPtr = UIRowType;
 
 	// Model class for the Teds Table Viewer that can be plugged into any widget that is a UI representation of data in Teds
 	// @see STedsTableViewer
diff --git a/Engine/Plugins/Experimental/EditorDataStorageFeatures/Source/TedsTableViewer/Public/TypedElementUITypes.h b/Engine/Plugins/Experimental/EditorDataStorageFeatures/Source/TedsTableViewer/Public/TypedElementUITypes.h
new file mode 100644
--- /dev/null
+++ b/Engine/Plugins/Experimental/EditorDataStorageFeatures/Source/TedsTableViewer/Public/TypedElementUITypes.h
@@ -0,0 +1,99 @@
+// Copyright Epic Games, Inc. All Rights Reserved.
+
+#pragma once
+
+#include "Elements/Common/TypedElementHandles.h"
+#include "Framework/Views/TableViewTypeTraits.h"
+
+namespace UE::EditorDataStorage
+{
+	// Wrapper struct around TypedElementDataStorage::RowHandle so we can template speicalize for it without also specializing for uint64
+	struct UIRowType
+	{
+		TypedElementDataStorage::RowHandle RowHandle;
+		
+		UIRowType()
+			: RowHandle(TypedElementDataStorage::InvalidRowHandle)
+		{}
+
+		UIRowType(TypedElementDataStorage::RowHandle InRowHandle)
+			: RowHandle(InRowHandle)
+		{}
+
+		operator TypedElementDataStorage::RowHandle () const
+		{
+			return RowHandle;
+		}
+
+		UIRowType& operator=(TypedElementDataStorage::RowHandle InRowHandle)
+		{
+			RowHandle = InRowHandle;
+			return *this;
+		}
+		
+		friend uint32 GetTypeHash(const UIRowType& Key)
+		{
+			return GetTypeHash(Key.RowHandle);
+		}
+
+	};
+}
+
+/* Template declaration to describe how a row handle behaves as a type for slate widgets like SListView, STreeView etc
+ * This allows you to use Row Handles with slate widgets that work on pointers by using the wrapper struct e.g
+ * SListView<UE::EditorDataStorage::UIRowType>
+ */
+template <>
+struct TListTypeTraits<UE::EditorDataStorage::UIRowType>
+{
+	typedef UE::EditorDataStorage::UIRowType NullableType;
+
+	using MapKeyFuncs = TDefaultMapHashableKeyFuncs<UE::EditorDataStorage::UIRowType, TSharedRef<ITableRow>, false>;
+	using MapKeyFuncsSparse = TDefaultMapHashableKeyFuncs<UE::EditorDataStorage::UIRowType, FSparseItemInfo, false>;
+	using SetKeyFuncs = DefaultKeyFuncs<UE::EditorDataStorage::UIRowType>;
+
+	template<typename U>
+	static void AddReferencedObjects(FReferenceCollector&,
+		TArray<UE::EditorDataStorage::UIRowType>&,
+		TSet<UE::EditorDataStorage::UIRowType>&,
+		TMap<const U*, UE::EditorDataStorage::UIRowType>&)
+	{
+	}
+
+	static bool IsPtrValid(const UE::EditorDataStorage::UIRowType& InPtr)
+	{
+		return InPtr != TypedElementDataStorage::InvalidRowHandle;
+	}
+
+	static void ResetPtr(UE::EditorDataStorage::UIRowType& InPtr)
+	{
+		InPtr = TypedElementDataStorage::InvalidRowHandle;
+	}
+
+	static UE::EditorDataStorage::UIRowType MakeNullPtr()
+	{
+		return TypedElementDataStorage::InvalidRowHandle;
+	}
+
+	static UE::EditorDataStorage::UIRowType NullableItemTypeConvertToItemType(const UE::EditorDataStorage::UIRowType& InPtr)
+	{
+		return InPtr;
+	}
+
+	static FString DebugDump(UE::EditorDataStorage::UIRowType InPtr)
+	{
+		return FString::Printf(TEXT("%llu"), InPtr.RowHandle);
+	}
+
+	class SerializerType {};
+};
+
+// Template declaration to enable using row handles inside of slate widgets like SListView
+template <>
+struct TIsValidListItem<UE::EditorDataStorage::UIRowType>
+{
+	enum
+	{
+		Value = true
+	};
+};
\ No newline at end of file
diff --git a/Engine/Plugins/Experimental/EditorDataStorageFeatures/Source/TedsTableViewer/Public/Widgets/SRowDetails.h b/Engine/Plugins/Experimental/EditorDataStorageFeatures/Source/TedsTableViewer/Public/Widgets/SRowDetails.h
new file mode 100644
--- /dev/null
+++ b/Engine/Plugins/Experimental/EditorDataStorageFeatures/Source/TedsTableViewer/Public/Widgets/SRowDetails.h
@@ -0,0 +1,95 @@
+// Copyright Epic Games, Inc. All Rights Reserved.
+#pragma once
+
+#include "Elements/Common/TypedElementHandles.h"
+#include "Templates/SharedPointer.h"
+#include "Widgets/SCompoundWidget.h"
+#include "Widgets/Views/SListView.h"
+
+class ITypedElementDataStorageInterface;
+class ITypedElementDataStorageUiInterface;
+struct FTypedElementWidgetConstructor;
+
+namespace UE::EditorDataStorage
+{
+	// A row in the SRowDetails widget that represents a column on the TEDS row we are viewing
+	struct FRowDetailsItem
+	{
+		// The column we this row is displaying data for
+		TWeakObjectPtr<const UScriptStruct> ColumnType;
+
+		// Widget for the column
+		TUniquePtr<FTypedElementWidgetConstructor> WidgetConstructor;
+		
+		TypedElementDataStorage::RowHandle Row = TypedElementDataStorage::InvalidRowHandle;
+		TypedElementDataStorage::RowHandle WidgetRow = TypedElementDataStorage::InvalidRowHandle;
+	};
+	
+	using RowDetailsItemPtr = TSharedPtr<FRowDetailsItem>;
+
+	// A widget to display all the columns/tags on a given row
+	class TEDSTABLEVIEWER_API SRowDetails : public SCompoundWidget
+	{
+	public:
+		
+		~SRowDetails() override = default;
+		
+		SLATE_BEGIN_ARGS(SRowDetails)
+			: _ShowAllDetails(true)
+		{}
+
+			// Whether or not to show columns that don't have a dedicated widget to represent them
+			SLATE_ARGUMENT(bool, ShowAllDetails)
+
+			// Override for the default widget purposes used to create widgets for the columns
+			SLATE_ARGUMENT(TArray<FName>, WidgetPurposesOverride)
+
+		
+		SLATE_END_ARGS()
+		
+		void Construct(const FArguments& InArgs);
+
+		// Set the row to view
+		void SetRow(TypedElementDataStorage::RowHandle Row);
+
+		// Clear the row to view
+		void ClearRow();
+		
+	private:
+		
+		TSharedRef<ITableRow> CreateRow(RowDetailsItemPtr InItem, const TSharedRef<STableViewBase>& OwnerTable);
+
+		TSharedPtr<SListView<RowDetailsItemPtr>> ListView;
+
+		TArray<RowDetailsItemPtr> Items;
+
+		ITypedElementDataStorageInterface* DataStorage = nullptr; 
+		ITypedElementDataStorageUiInterface* DataStorageUi = nullptr;
+
+		bool bShowAllDetails = true;
+
+		TArray<FName> WidgetPurposes;
+
+	};
+
+	class SRowDetailsRow : public SMultiColumnTableRow<RowDetailsItemPtr>
+	{
+	public:
+		
+		SLATE_BEGIN_ARGS(SRowDetailsRow) {}
+		
+			SLATE_ARGUMENT(RowDetailsItemPtr, Item)
+
+		SLATE_END_ARGS()
+		
+		void Construct(const FArguments& Args, const TSharedRef<STableViewBase>& OwnerTableView, ITypedElementDataStorageInterface* InDataStorage,
+			ITypedElementDataStorageUiInterface* InDataStorageUi);
+		
+		TSharedRef<SWidget> GenerateWidgetForColumn(const FName& ColumnName) override;
+
+	private:
+		RowDetailsItemPtr Item;
+		ITypedElementDataStorageInterface* DataStorage = nullptr;
+		ITypedElementDataStorageUiInterface* DataStorageUi = nullptr;
+	};
+} // namespace UE::EditorDataStorage
\ No newline at end of file
diff --git a/Engine/Plugins/Experimental/EditorDataStorageFeatures/Source/TedsTableViewer/Public/Widgets/STedsTableViewer.h b/Engine/Plugins/Experimental/EditorDataStorageFeatures/Source/TedsTableViewer/Public/Widgets/STedsTableViewer.h
--- a/Engine/Plugins/Experimental/EditorDataStorageFeatures/Source/TedsTableViewer/Public/Widgets/STedsTableViewer.h
+++ b/Engine/Plugins/Experimental/EditorDataStorageFeatures/Source/TedsTableViewer/Public/Widgets/STedsTableViewer.h
@@ -27,6 +27,11 @@ namespace UE::EditorDataStorage
 	 */
 	class STedsTableViewer : public SCompoundWidget
 	{
+	public:
+		
+		// Delegate fired when the selection in the table viewer changes
+		DECLARE_DELEGATE_OneParam(FOnSelectionChanged, TypedElementDataStorage::RowHandle)
+
 		SLATE_BEGIN_ARGS(STedsTableViewer)
 			: _CellWidgetPurposes({TEXT("General.Cell")})
 		{
@@ -43,6 +48,9 @@ namespace UE::EditorDataStorage
 
 		// The widget purposes to use to create the widgets
 		SLATE_ARGUMENT(TArray<FName>, CellWidgetPurposes)
+		
+		// Delegate called when the selection changes
+		SLATE_ARGUMENT(FOnSelectionChanged, OnSelectionChanged)
 
 		SLATE_END_ARGS()
 
@@ -67,6 +75,8 @@ namespace UE::EditorDataStorage
 
 		void RefreshColumnWidgets();
 
+		void OnListSelectionChanged(TableViewerItemPtr Item, ESelectInfo::Type SelectInfo);
+
 	private:
 
 		// The actual ListView widget that displays the rows
@@ -77,5 +87,8 @@ namespace UE::EditorDataStorage
 
 		// Our model class
 		TSharedPtr<FTedsTableViewerModel> Model;
+
+		// Delegate fired when the selection changes
+		FOnSelectionChanged OnSelectionChanged;
 	};
 }
\ No newline at end of file
diff --git a/Engine/Source/Runtime/TypedElementFramework/Private/Elements/Interfaces/TypedElementDataStorageUiInterface.cpp b/Engine/Source/Runtime/TypedElementFramework/Private/Elements/Interfaces/TypedElementDataStorageUiInterface.cpp
--- a/Engine/Source/Runtime/TypedElementFramework/Private/Elements/Interfaces/TypedElementDataStorageUiInterface.cpp
+++ b/Engine/Source/Runtime/TypedElementFramework/Private/Elements/Interfaces/TypedElementDataStorageUiInterface.cpp
@@ -135,7 +135,6 @@ TSharedPtr<SWidget> FTypedElementWidgetConstructor::ConstructFinalWidget(
 	return ContainerWidget;
 }
 
-
 TSharedPtr<SWidget> FTypedElementWidgetConstructor::Construct(
 	TypedElementRowHandle Row,
 	ITypedElementDataStorageInterface* DataStorage,
@@ -150,7 +149,6 @@ TSharedPtr<SWidget> FTypedElementWidgetConstructor::Construct(
 		{
 			if (FinalizeWidget(DataStorage, DataStorageUi, Row, Widget))
 			{
-				SetupDebugColumns(Row, DataStorage, DataStorageUi, Widget);
 				return Widget;
 			}
 		}
@@ -199,25 +197,3 @@ bool FTypedElementWidgetConstructor::FinalizeWidget(
 {
 	return true;
 }
-
-void FTypedElementWidgetConstructor::SetupDebugColumns(TypedElementRowHandle Row, ITypedElementDataStorageInterface* DataStorage, ITypedElementDataStorageUiInterface* DataStorageUi, const TSharedPtr<SWidget>& Widget)
-{
-	bool bAddLabelColumn = true;
-
-	const FString WidgetLabel(GetWidgetLabel(Widget));
-
-	// Only add the label column for widgets that are not created for other widgets. This is to prevent recursively displaying widgets for widgets
-	// in the TEDS-Debugger - we only want top level widgets to be displayed
-	if (const FTypedElementRowReferenceColumn* RowReference = DataStorage->GetColumn<FTypedElementRowReferenceColumn>(Row))
-	{
-		if (DataStorage->HasColumns<FTypedElementSlateWidgetReferenceColumn>(RowReference->Row))
-		{
-			bAddLabelColumn = false;
-		}
-	}
-
-	if(bAddLabelColumn)
-	{
-		DataStorage->AddColumn(Row, FTypedElementLabelColumn{.Label = WidgetLabel} );
-	}
-}
diff --git a/Engine/Source/Runtime/TypedElementFramework/Public/Elements/Interfaces/TypedElementDataStorageUiInterface.h b/Engine/Source/Runtime/TypedElementFramework/Public/Elements/Interfaces/TypedElementDataStorageUiInterface.h
--- a/Engine/Source/Runtime/TypedElementFramework/Public/Elements/Interfaces/TypedElementDataStorageUiInterface.h
+++ b/Engine/Source/Runtime/TypedElementFramework/Public/Elements/Interfaces/TypedElementDataStorageUiInterface.h
@@ -103,15 +103,6 @@ protected:
 		TypedElementRowHandle Row,
 		const TSharedPtr<SWidget>& Widget);
 	
-	/** Setup default columns on the widget row
-	 *  TEDS UI TODO - This should be moved to a processor or maybe core TEDS UI
-	 */
-	TYPEDELEMENTFRAMEWORK_API void SetupDebugColumns(
-		TypedElementRowHandle Row,
-		ITypedElementDataStorageInterface* DataStorage,
-		ITypedElementDataStorageUiInterface* DataStorageUi,
-		const TSharedPtr<SWidget>& Widget);
-
 protected:
 
 	TArray<TWeakObjectPtr<const UScriptStruct>> MatchedColumnTypes;
