SRowDetails: Add a constructor to initialize the struct properly

diff --git a/Engine/Plugins/Experimental/EditorDataStorageFeatures/Source/TedsTableViewer/Private/Widgets/SRowDetails.cpp b/Engine/Plugins/Experimental/EditorDataStorageFeatures/Source/TedsTableViewer/Private/Widgets/SRowDetails.cpp
--- a/Engine/Plugins/Experimental/EditorDataStorageFeatures/Source/TedsTableViewer/Private/Widgets/SRowDetails.cpp
+++ b/Engine/Plugins/Experimental/EditorDataStorageFeatures/Source/TedsTableViewer/Private/Widgets/SRowDetails.cpp
@@ -25,7 +25,7 @@ namespace UE::EditorDataStorage
 	//
 	// SRowDetails
 	//
-	
+
 	void SRowDetails::Construct(const FArguments& InArgs)
 	{
 		bShowAllDetails = InArgs._ShowAllDetails;
@@ -179,6 +179,16 @@ namespace UE::EditorDataStorage
 					.Text(LOCTEXT("InvalidColumnType", "Invalid Column Type")); 
 		}
 	}
+
+	// FRowDetailsItem
+	FRowDetailsItem::FRowDetailsItem(const TWeakObjectPtr<const UScriptStruct>& InColumnType,
+		TUniquePtr<FTypedElementWidgetConstructor> InWidgetConstructor, TypedElementDataStorage::RowHandle InRow)
+		: ColumnType(InColumnType)
+		, WidgetConstructor(MoveTemp(InWidgetConstructor))
+		, Row(InRow)
+	{
+	
+	}
 	
 } // namespace UE::EditorDataStorage
 
diff --git a/Engine/Plugins/Experimental/EditorDataStorageFeatures/Source/TedsTableViewer/Public/Widgets/SRowDetails.h b/Engine/Plugins/Experimental/EditorDataStorageFeatures/Source/TedsTableViewer/Public/Widgets/SRowDetails.h
--- a/Engine/Plugins/Experimental/EditorDataStorageFeatures/Source/TedsTableViewer/Public/Widgets/SRowDetails.h
+++ b/Engine/Plugins/Experimental/EditorDataStorageFeatures/Source/TedsTableViewer/Public/Widgets/SRowDetails.h
@@ -23,6 +23,9 @@ namespace UE::EditorDataStorage
 		
 		TypedElementDataStorage::RowHandle Row = TypedElementDataStorage::InvalidRowHandle;
 		TypedElementDataStorage::RowHandle WidgetRow = TypedElementDataStorage::InvalidRowHandle;
+
+		FRowDetailsItem(const TWeakObjectPtr<const UScriptStruct>& InColumnType, TUniquePtr<FTypedElementWidgetConstructor> InWidgetConstructor,
+			TypedElementDataStorage::RowHandle InRow);
 	};
 	
 	using RowDetailsItemPtr = TSharedPtr<FRowDetailsItem>;
