EditorDataStorageFeatures : Add the TedsSettings module

The TedsSettings module is responsible for mirroring settings data from the ISettingsModule into the editor data storage.  There is 1 row per ISettingsSection, and a pointer to the SettingsObject ( UObject ) for the section is stored in the row using the compatible object API.

Individual settings are properties of the SettingsObject, and it is up to the different users of this row data to first query for the settings object before inspecting or modifying its properties.  This offers the same granularity and mode of operation as the ISettingsSection.

Columns are added for the SettingsContainer, SettingsCategory, and SettingsSection.  These currently contain only the respective FName for each, but could contain more data ( e.g. DisplayName, Description ) or pointers ( e.g. to the ISettingsCategory ) in the future.

A console variable "TEDS.Feature.Settings.Enable" is provided to enable this feature, and it is disabled by default at this time.

diff --git a/Engine/Plugins/Experimental/EditorDataStorageFeatures/EditorDataStorageFeatures.uplugin b/Engine/Plugins/Experimental/EditorDataStorageFeatures/EditorDataStorageFeatures.uplugin
--- a/Engine/Plugins/Experimental/EditorDataStorageFeatures/EditorDataStorageFeatures.uplugin
+++ b/Engine/Plugins/Experimental/EditorDataStorageFeatures/EditorDataStorageFeatures.uplugin
@@ -35,6 +35,10 @@
 			"Name": "TedsDebugger",
 			"Type": "Editor"
 		},
+		{
+			"Name": "TedsSettings",
+			"Type": "Editor"
+		},
 		{
 			"Name": "TedsTableViewer",
 			"Type": "Editor"
diff --git a/Engine/Plugins/Experimental/EditorDataStorageFeatures/Source/TedsSettings/Private/TedsSettingsEditorSubsystem.cpp b/Engine/Plugins/Experimental/EditorDataStorageFeatures/Source/TedsSettings/Private/TedsSettingsEditorSubsystem.cpp
new file mode 100644
--- /dev/null
+++ b/Engine/Plugins/Experimental/EditorDataStorageFeatures/Source/TedsSettings/Private/TedsSettingsEditorSubsystem.cpp
@@ -0,0 +1,73 @@
+// Copyright Epic Games, Inc. All Rights Reserved.
+
+#include "TedsSettingsEditorSubsystem.h"
+
+#include "Editor.h"
+#include "TedsSettingsLog.h"
+#include "TedsSettingsManager.h"
+
+namespace UE::EditorDataStorage::Settings
+{
+	static TAutoConsoleVariable<bool> CVarTedsSettingsEnable(
+		TEXT("TEDS.Feature.Settings.Enable"),
+		false,
+		TEXT("When true, settings objects from the ISettingsModule will be mirrored to rows in the editor data storage."));
+}
+
+UTedsSettingsEditorSubsystem::UTedsSettingsEditorSubsystem()
+	: UEditorSubsystem()
+	, SettingsManager{ MakeShared<FTedsSettingsManager>() }
+	, EnabledChangedDelegate{ }
+{
+}
+
+const bool UTedsSettingsEditorSubsystem::IsEnabled() const
+{
+	return UE::EditorDataStorage::Settings::CVarTedsSettingsEnable.GetValueOnGameThread();
+}
+
+UTedsSettingsEditorSubsystem::FOnEnabledChanged& UTedsSettingsEditorSubsystem::OnEnabledChanged()
+{
+	return EnabledChangedDelegate;
+}
+
+void UTedsSettingsEditorSubsystem::Initialize(FSubsystemCollectionBase& Collection)
+{
+	Super::Initialize(Collection);
+
+	UE_LOG(LogTedsSettings, Log, TEXT("UTedsSettingsEditorSubsystem::Initialize"));
+
+	UE::EditorDataStorage::Settings::CVarTedsSettingsEnable->SetOnChangedCallback(
+		FConsoleVariableDelegate::CreateLambda([this](IConsoleVariable* Variable)
+		{
+			const bool bIsEnabled = Variable->GetBool();
+			
+			if (bIsEnabled)
+			{
+				SettingsManager->Initialize();
+			}
+			else
+			{
+				SettingsManager->Shutdown();
+			}
+
+			EnabledChangedDelegate.Broadcast();
+		}));
+
+	if (UE::EditorDataStorage::Settings::CVarTedsSettingsEnable.GetValueOnGameThread())
+	{
+		SettingsManager->Initialize();
+	}
+}
+
+void UTedsSettingsEditorSubsystem::Deinitialize()
+{
+	UE_LOG(LogTedsSettings, Log, TEXT("UTedsSettingsEditorSubsystem::Deinitialize"));
+
+	if (UE::EditorDataStorage::Settings::CVarTedsSettingsEnable.GetValueOnGameThread())
+	{
+		SettingsManager->Shutdown();
+	}
+
+	Super::Deinitialize();
+}
diff --git a/Engine/Plugins/Experimental/EditorDataStorageFeatures/Source/TedsSettings/Private/TedsSettingsLog.cpp b/Engine/Plugins/Experimental/EditorDataStorageFeatures/Source/TedsSettings/Private/TedsSettingsLog.cpp
new file mode 100644
--- /dev/null
+++ b/Engine/Plugins/Experimental/EditorDataStorageFeatures/Source/TedsSettings/Private/TedsSettingsLog.cpp
@@ -0,0 +1,5 @@
+// Copyright Epic Games, Inc. All Rights Reserved.
+
+#include "TedsSettingsLog.h"
+
+DEFINE_LOG_CATEGORY(LogTedsSettings);
diff --git a/Engine/Plugins/Experimental/EditorDataStorageFeatures/Source/TedsSettings/Private/TedsSettingsLog.h b/Engine/Plugins/Experimental/EditorDataStorageFeatures/Source/TedsSettings/Private/TedsSettingsLog.h
new file mode 100644
--- /dev/null
+++ b/Engine/Plugins/Experimental/EditorDataStorageFeatures/Source/TedsSettings/Private/TedsSettingsLog.h
@@ -0,0 +1,7 @@
+// Copyright Epic Games, Inc. All Rights Reserved.
+
+#pragma once
+
+#include "Logging/LogMacros.h"
+
+DECLARE_LOG_CATEGORY_EXTERN(LogTedsSettings, Log, All);
diff --git a/Engine/Plugins/Experimental/EditorDataStorageFeatures/Source/TedsSettings/Private/TedsSettingsManager.cpp b/Engine/Plugins/Experimental/EditorDataStorageFeatures/Source/TedsSettings/Private/TedsSettingsManager.cpp
new file mode 100644
--- /dev/null
+++ b/Engine/Plugins/Experimental/EditorDataStorageFeatures/Source/TedsSettings/Private/TedsSettingsManager.cpp
@@ -0,0 +1,283 @@
+// Copyright Epic Games, Inc. All Rights Reserved.
+
+#include "TedsSettingsManager.h"
+
+#include "Elements/Columns/TypedElementCompatibilityColumns.h"
+#include "Elements/Framework/TypedElementQueryBuilder.h"
+#include "Elements/Framework/TypedElementRegistry.h"
+#include "ISettingsCategory.h"
+#include "ISettingsContainer.h"
+#include "ISettingsModule.h"
+#include "ISettingsSection.h"
+#include "Logging/LogMacros.h"
+#include "Modules/ModuleManager.h"
+#include "ProfilingDebugging/CpuProfilerTrace.h"
+#include "TedsSettingsColumns.h"
+#include "TedsSettingsLog.h"
+#include "UObject/UnrealType.h"
+
+FTedsSettingsManager::FTedsSettingsManager()
+	: bIsInitialized{ false }
+	, SelectAllSettingsQuery{ TypedElementDataStorage::InvalidQueryHandle }
+{
+}
+
+void FTedsSettingsManager::Initialize()
+{
+	if (!bIsInitialized)
+	{
+		FModuleManager::Get().LoadModule(TEXT("TypedElementFramework"));
+
+		UTypedElementRegistry* TypedElementRegistry = UTypedElementRegistry::GetInstance();
+		check(TypedElementRegistry);
+
+		auto OnDataStorage = [this, TypedElementRegistry]
+			{
+				ITypedElementDataStorageInterface* DataStorage = TypedElementRegistry->GetMutableDataStorage();
+				check(DataStorage);
+
+				RegisterQueries(*DataStorage);
+				RegisterSettings();
+			};
+
+		if (TypedElementRegistry->AreDataStorageInterfacesSet())
+		{
+			OnDataStorage();
+		}
+		else
+		{
+			TypedElementRegistry->OnDataStorageInterfacesSet().AddSPLambda(this, OnDataStorage);
+		}
+
+		bIsInitialized = true;
+	}
+}
+
+void FTedsSettingsManager::Shutdown()
+{
+	if (bIsInitialized)
+	{
+		UTypedElementRegistry* TypedElementRegistry = UTypedElementRegistry::GetInstance();
+		check(TypedElementRegistry);
+
+		TypedElementRegistry->OnDataStorageInterfacesSet().RemoveAll(this);
+
+		if (TypedElementRegistry->AreDataStorageInterfacesSet())
+		{
+			ITypedElementDataStorageInterface* DataStorage = TypedElementRegistry->GetMutableDataStorage();
+			check(DataStorage);
+
+			UnregisterSettings();
+			UnregisterQueries(*DataStorage);
+		}
+
+		bIsInitialized = false;
+	}
+}
+
+void FTedsSettingsManager::RegisterQueries(ITypedElementDataStorageInterface& DataStorage)
+{
+	using namespace TypedElementQueryBuilder;
+
+	if (SelectAllSettingsQuery == TypedElementDataStorage::InvalidQueryHandle)
+	{
+		SelectAllSettingsQuery = DataStorage.RegisterQuery(
+			Select()
+			.ReadOnly<FTypedElementUObjectColumn, FSettingsContainerColumn, FSettingsCategoryColumn, FSettingsSectionColumn>()
+			.Compile());
+	}
+}
+
+void FTedsSettingsManager::UnregisterQueries(ITypedElementDataStorageInterface& DataStorage)
+{
+	DataStorage.UnregisterQuery(SelectAllSettingsQuery);
+	SelectAllSettingsQuery = TypedElementDataStorage::InvalidQueryHandle;
+}
+
+void FTedsSettingsManager::RegisterSettings()
+{
+	TRACE_CPUPROFILER_EVENT_SCOPE(TedsSettingsManager.RegisterSettings);
+
+	ISettingsModule* SettingsModule = FModuleManager::GetModulePtr<ISettingsModule>("Settings");
+	check(SettingsModule);
+
+	// TODO: add a delegate to the ISettingsModule to keep data storage in sync when a container is added or removed
+
+	TArray<FName> ContainerNames;
+	SettingsModule->GetContainerNames(ContainerNames);
+
+	for (FName ContainerName : ContainerNames)
+	{
+		UE_LOG(LogTedsSettings, Log, TEXT("Register Settings Container : '%s'"), *ContainerName.ToString());
+
+		ISettingsContainerPtr ContainerPtr = SettingsModule->GetContainer(ContainerName);
+
+		TArray<ISettingsCategoryPtr> Categories;
+		ContainerPtr->GetCategories(Categories);
+
+		for (ISettingsCategoryPtr CategoryPtr : Categories)
+		{
+			const bool bQueryExistingRows = false;
+			UpdateSettingsCategory(CategoryPtr, ContainerName, bQueryExistingRows);
+		}
+
+		// OnCategoryModified is called at the same time as OnSectionRemoved so we only bind to OnCategoryModified for add / update / remove
+		ContainerPtr->OnCategoryModified().AddSPLambda(this, [this, ContainerPtr](const FName& ModifiedCategoryName)
+			{
+				UE_LOG(LogTedsSettings, Log, TEXT("Settings Category modified : '%s->%s'"), *ContainerPtr->GetName().ToString(), *ModifiedCategoryName.ToString());
+
+				ISettingsCategoryPtr CategoryPtr = ContainerPtr->GetCategory(ModifiedCategoryName);
+
+				UpdateSettingsCategory(CategoryPtr, ContainerPtr->GetName());
+			});
+	}
+}
+
+void FTedsSettingsManager::UnregisterSettings()
+{
+	TRACE_CPUPROFILER_EVENT_SCOPE(TedsSettingsManager.UnregisterSettings);
+
+	ISettingsModule* SettingsModule = FModuleManager::GetModulePtr<ISettingsModule>("Settings");
+	check(SettingsModule);
+
+	UTypedElementRegistry* TypedElementRegistry = UTypedElementRegistry::GetInstance();
+	check(TypedElementRegistry);
+
+	ITypedElementDataStorageCompatibilityInterface* DataStorageCompatibility = TypedElementRegistry->GetMutableDataStorageCompatibility();
+	check(DataStorageCompatibility);
+
+	TArray<FName> ContainerNames;
+	SettingsModule->GetContainerNames(ContainerNames);
+
+	for (FName ContainerName : ContainerNames)
+	{
+		UE_LOG(LogTedsSettings, Log, TEXT("Unregister Settings Container : '%s'"), *ContainerName.ToString());
+
+		ISettingsContainerPtr ContainerPtr = SettingsModule->GetContainer(ContainerName);
+
+		ContainerPtr->OnCategoryModified().RemoveAll(this);
+
+		TArray<ISettingsCategoryPtr> Categories;
+		ContainerPtr->GetCategories(Categories);
+
+		for (ISettingsCategoryPtr CategoryPtr : Categories)
+		{
+			UE_LOG(LogTedsSettings, Log, TEXT("Unregister Settings Category : '%s'"), *CategoryPtr->GetName().ToString());
+
+			TArray<ISettingsSectionPtr> Sections;
+			const bool bIgnoreVisibility = true;
+			CategoryPtr->GetSections(Sections, bIgnoreVisibility);
+
+			for (ISettingsSectionPtr SectionPtr : Sections)
+			{
+				if (TStrongObjectPtr<UObject> SettingsObjectPtr = SectionPtr->GetSettingsObject().Pin(); SettingsObjectPtr)
+				{
+					DataStorageCompatibility->RemoveCompatibleObject(SettingsObjectPtr);
+
+					UE_LOG(LogTedsSettings, Log, TEXT("Removed Settings Section : '%s'"), *SectionPtr->GetName().ToString());
+				}
+			}
+		}
+	}
+}
+
+void FTedsSettingsManager::UpdateSettingsCategory(TSharedPtr<ISettingsCategory> SettingsCategory, const FName& ContainerName, const bool bQueryExistingRows)
+{
+	TRACE_CPUPROFILER_EVENT_SCOPE(TedsSettingsManager.UpdateSettingsCategory);
+
+	const FName& CategoryName = SettingsCategory->GetName();
+
+	UE_LOG(LogTedsSettings, Log, TEXT("Update Settings Category: '%s->%s'"), *ContainerName.ToString(), *CategoryName.ToString());
+
+	UTypedElementRegistry* TypedElementRegistry = UTypedElementRegistry::GetInstance();
+	check(TypedElementRegistry);
+
+	ITypedElementDataStorageInterface* DataStorage = TypedElementRegistry->GetMutableDataStorage();
+	check(DataStorage);
+
+	ITypedElementDataStorageCompatibilityInterface* DataStorageCompatibility = TypedElementRegistry->GetMutableDataStorageCompatibility();
+	check(DataStorageCompatibility);
+
+	TArray<TypedElementDataStorage::RowHandle> OldRowHandles;
+	TArray<FName> OldSectionNames;
+
+	// Gather all existing rows for the given { ContainerName, CategoryName } pair.
+	if (bQueryExistingRows)
+	{
+		using namespace TypedElementQueryBuilder;
+		using DSI = ITypedElementDataStorageInterface;
+
+		DataStorage->RunQuery(SelectAllSettingsQuery, CreateDirectQueryCallbackBinding(
+			[&OldRowHandles, &OldSectionNames, &ContainerName, &CategoryName](
+				DSI::IDirectQueryContext& Context,
+				const FTypedElementUObjectColumn* ObjectColumns,
+				const FSettingsContainerColumn* ContainerColumns,
+				const FSettingsCategoryColumn* CategoryColumns,
+				const FSettingsSectionColumn* SectionColumns)
+			{
+				const uint32 RowCount = Context.GetRowCount();
+
+				for (uint32 RowIndex = 0; RowIndex < RowCount; ++RowIndex)
+				{
+					const FName& TempContainerName = ContainerColumns[RowIndex].ContainerName;
+					const FName& TempCategoryName = CategoryColumns[RowIndex].CategoryName;
+					if (TempContainerName == ContainerName &&
+						TempCategoryName == CategoryName)
+					{
+						OldRowHandles.Emplace(Context.GetRowHandles()[RowIndex]);
+						OldSectionNames.Emplace(SectionColumns[RowIndex].SectionName);
+					}
+				}
+			}));
+	}
+
+	TArray<FName> NewSectionNames;
+	TArray<ISettingsSectionPtr> NewSections;
+
+	const bool bIgnoreVisibility = true;
+	SettingsCategory->GetSections(NewSections, bIgnoreVisibility);
+
+	// Iterate the category and add rows for any section not in the old sections list.
+	for (ISettingsSectionPtr SectionPtr : NewSections)
+	{
+		const FName& SectionName = SectionPtr->GetName();
+
+		if (TStrongObjectPtr<UObject> SettingsObjectPtr = SectionPtr->GetSettingsObject().Pin(); SettingsObjectPtr)
+		{
+			NewSectionNames.Emplace(SectionName);
+
+			if (OldSectionNames.Contains(SectionName))
+			{
+				UE_LOG(LogTedsSettings, Verbose, TEXT("Settings Section : '%s' is already in data storage"), *SectionName.ToString());
+
+				continue;
+			}
+
+			TypedElementDataStorage::RowHandle NewRow = DataStorageCompatibility->AddCompatibleObject(SettingsObjectPtr);
+
+			DataStorage->AddColumn<FSettingsContainerColumn>(NewRow, { .ContainerName = ContainerName });
+			DataStorage->AddColumn<FSettingsCategoryColumn>(NewRow, { .CategoryName = CategoryName });
+			DataStorage->AddColumn<FSettingsSectionColumn>(NewRow, { .SectionName = SectionName });
+
+			UE_LOG(LogTedsSettings, Log, TEXT("Added Settings Section : '%s'"), *SectionName.ToString());
+		}
+	}
+
+	// Iterate the old sections and remove rows not in the new sections list.
+	for (int32 RowIndex = 0; RowIndex < OldSectionNames.Num(); ++RowIndex)
+	{
+		const FName& OldSectionName = OldSectionNames[RowIndex];
+
+		if (NewSectionNames.Contains(OldSectionName))
+		{
+			continue;
+		}
+
+		TypedElementDataStorage::RowHandle OldRowHandle = OldRowHandles[RowIndex];
+		check(OldRowHandle != TypedElementInvalidRowHandle);
+
+		DataStorage->RemoveRow(OldRowHandle);
+
+		UE_LOG(LogTedsSettings, Log, TEXT("Removed Settings Section : '%s'"), *OldSectionName.ToString());
+	}
+}
diff --git a/Engine/Plugins/Experimental/EditorDataStorageFeatures/Source/TedsSettings/Private/TedsSettingsManager.h b/Engine/Plugins/Experimental/EditorDataStorageFeatures/Source/TedsSettings/Private/TedsSettingsManager.h
new file mode 100644
--- /dev/null
+++ b/Engine/Plugins/Experimental/EditorDataStorageFeatures/Source/TedsSettings/Private/TedsSettingsManager.h
@@ -0,0 +1,37 @@
+// Copyright Epic Games, Inc. All Rights Reserved.
+
+#pragma once
+
+#include "Elements/Interfaces/TypedElementDataStorageInterface.h"
+#include "Templates/SharedPointer.h"
+
+class ISettingsCategory;
+
+class FTedsSettingsManager final : public TSharedFromThis<FTedsSettingsManager>
+{
+public:
+
+	FTedsSettingsManager();
+
+	const bool IsInitialized() const
+	{
+		return bIsInitialized;
+	}
+
+	void Initialize();
+	void Shutdown();
+
+private:
+
+	void RegisterQueries(ITypedElementDataStorageInterface& DataStorage);
+	void UnregisterQueries(ITypedElementDataStorageInterface& DataStorage);
+
+	void RegisterSettings();
+	void UnregisterSettings();
+
+	void UpdateSettingsCategory(TSharedPtr<ISettingsCategory> SettingsCategory, const FName& ContainerName, const bool bQueryExistingRows = true);
+
+	bool bIsInitialized;
+	TypedElementDataStorage::QueryHandle SelectAllSettingsQuery;
+
+};
diff --git a/Engine/Plugins/Experimental/EditorDataStorageFeatures/Source/TedsSettings/Private/TedsSettingsModule.cpp b/Engine/Plugins/Experimental/EditorDataStorageFeatures/Source/TedsSettings/Private/TedsSettingsModule.cpp
new file mode 100644
--- /dev/null
+++ b/Engine/Plugins/Experimental/EditorDataStorageFeatures/Source/TedsSettings/Private/TedsSettingsModule.cpp
@@ -0,0 +1,5 @@
+// Copyright Epic Games, Inc. All Rights Reserved.
+
+#include "Modules/ModuleManager.h"
+
+IMPLEMENT_MODULE(FDefaultModuleImpl, TedsSettings)
diff --git a/Engine/Plugins/Experimental/EditorDataStorageFeatures/Source/TedsSettings/Public/TedsSettingsColumns.h b/Engine/Plugins/Experimental/EditorDataStorageFeatures/Source/TedsSettings/Public/TedsSettingsColumns.h
new file mode 100644
--- /dev/null
+++ b/Engine/Plugins/Experimental/EditorDataStorageFeatures/Source/TedsSettings/Public/TedsSettingsColumns.h
@@ -0,0 +1,35 @@
+// Copyright Epic Games, Inc. All Rights Reserved.
+
+#pragma once
+
+#include "Elements/Interfaces/TypedElementDataStorageInterface.h"
+#include "UObject/ObjectMacros.h"
+
+#include "TedsSettingsColumns.generated.h"
+
+USTRUCT(meta = (DisplayName = "Settings Container"))
+struct FSettingsContainerColumn final : public FTypedElementDataStorageColumn
+{
+	GENERATED_BODY()
+
+	UPROPERTY()
+	FName ContainerName;
+};
+
+USTRUCT(meta = (DisplayName = "Settings Category"))
+struct FSettingsCategoryColumn final : public FTypedElementDataStorageColumn
+{
+	GENERATED_BODY()
+
+	UPROPERTY()
+	FName CategoryName;
+};
+
+USTRUCT(meta = (DisplayName = "Settings Section"))
+struct FSettingsSectionColumn final : public FTypedElementDataStorageColumn
+{
+	GENERATED_BODY()
+
+	UPROPERTY()
+	FName SectionName;
+};
diff --git a/Engine/Plugins/Experimental/EditorDataStorageFeatures/Source/TedsSettings/Public/TedsSettingsEditorSubsystem.h b/Engine/Plugins/Experimental/EditorDataStorageFeatures/Source/TedsSettings/Public/TedsSettingsEditorSubsystem.h
new file mode 100644
--- /dev/null
+++ b/Engine/Plugins/Experimental/EditorDataStorageFeatures/Source/TedsSettings/Public/TedsSettingsEditorSubsystem.h
@@ -0,0 +1,35 @@
+// Copyright Epic Games, Inc. All Rights Reserved.
+
+#pragma once
+
+#include "EditorSubsystem.h"
+
+#include "TedsSettingsEditorSubsystem.generated.h"
+
+class FTedsSettingsManager;
+
+UCLASS()
+class TEDSSETTINGS_API UTedsSettingsEditorSubsystem : public UEditorSubsystem
+{
+	GENERATED_BODY()
+
+public:
+
+	UTedsSettingsEditorSubsystem();
+
+	const bool IsEnabled() const;
+
+	DECLARE_MULTICAST_DELEGATE(FOnEnabledChanged)
+	FOnEnabledChanged& OnEnabledChanged();
+
+protected: // USubsystem
+
+	virtual void Initialize(FSubsystemCollectionBase& Collection) override;
+	virtual void Deinitialize() override;
+
+private:
+
+	TSharedPtr<FTedsSettingsManager> SettingsManager;
+	FOnEnabledChanged EnabledChangedDelegate;
+
+};
diff --git a/Engine/Plugins/Experimental/EditorDataStorageFeatures/Source/TedsSettings/TedsSettings.Build.cs b/Engine/Plugins/Experimental/EditorDataStorageFeatures/Source/TedsSettings/TedsSettings.Build.cs
new file mode 100644
--- /dev/null
+++ b/Engine/Plugins/Experimental/EditorDataStorageFeatures/Source/TedsSettings/TedsSettings.Build.cs
@@ -0,0 +1,33 @@
+// Copyright Epic Games, Inc. All Rights Reserved.
+
+using UnrealBuildTool;
+
+public class TedsSettings : ModuleRules
+{
+	public TedsSettings(ReadOnlyTargetRules Target) : base(Target)
+	{
+		PCHUsage = ModuleRules.PCHUsageMode.UseExplicitOrSharedPCHs;
+
+		if (Target.bBuildEditor)
+		{
+			PublicIncludePaths.AddRange(new string[] { });
+			PrivateIncludePaths.AddRange(new string[] { });
+
+			PrivateDependencyModuleNames.AddRange(
+				new string[]
+				{
+					"Core",
+					"CoreUObject",
+					"Engine",
+					"UnrealEd",
+				});
+
+			PublicDependencyModuleNames.AddRange(
+				new string[]
+				{
+					"EditorSubsystem",
+					"TypedElementFramework",
+				});
+		}
+	}
+}
diff --git a/Engine/Plugins/Experimental/EditorDataStorageFeatures/Source/TedsSettings/Tests/TestSettings.h b/Engine/Plugins/Experimental/EditorDataStorageFeatures/Source/TedsSettings/Tests/TestSettings.h
new file mode 100644
--- /dev/null
+++ b/Engine/Plugins/Experimental/EditorDataStorageFeatures/Source/TedsSettings/Tests/TestSettings.h
@@ -0,0 +1,19 @@
+// Copyright Epic Games, Inc. All Rights Reserved.
+
+#pragma once
+
+#include "UObject/ObjectMacros.h"
+
+#include "TestSettings.generated.h"
+
+UCLASS()
+class UTestSettings : public UObject
+{
+	GENERATED_BODY()
+
+public:
+
+	UPROPERTY()
+	int32 TestInt32 = 42;
+
+};
