Fix path of TedsSettingsTests.cpp

diff --git a/Engine/Plugins/Experimental/EditorDataStorageFeatures/Source/TedsSettings/Tests/TedsSettingsTests.cpp b/Engine/Plugins/Experimental/EditorDataStorageFeatures/Source/TedsSettings/Tests/TedsSettingsTests.cpp
new file mode 100644
--- /dev/null
+++ b/Engine/Plugins/Experimental/EditorDataStorageFeatures/Source/TedsSettings/Tests/TedsSettingsTests.cpp
@@ -0,0 +1,194 @@
+// Copyright Epic Games, Inc. All Rights Reserved.
+
+#include "Misc/AutomationTest.h"
+
+#if WITH_AUTOMATION_TESTS
+
+#include "Containers/Ticker.h"
+#include "Editor.h"
+#include "Elements/Columns/TypedElementCompatibilityColumns.h"
+#include "Elements/Framework/TypedElementQueryBuilder.h"
+#include "Elements/Framework/TypedElementRegistry.h"
+#include "Elements/Interfaces/TypedElementDataStorageInterface.h"
+#include "ISettingsModule.h"
+#include "Modules/ModuleManager.h"
+#include "TedsSettingsColumns.h"
+#include "TedsSettingsEditorSubsystem.h"
+#include "TestSettings.h"
+
+BEGIN_DEFINE_SPEC(FTedsSettingsTestFixture, "TedsSettings", EAutomationTestFlags::EditorContext | EAutomationTestFlags::EngineFilter)
+
+ISettingsModule* SettingsModule = nullptr;
+UTypedElementRegistry* TypedElementRegistry = nullptr;
+ITypedElementDataStorageInterface* DataStorage = nullptr;
+ITypedElementDataStorageCompatibilityInterface* DataStorageCompatibility = nullptr;
+TypedElementQueryHandle CountAllSettingsQuery = TypedElementDataStorage::InvalidQueryHandle;
+
+uint32 BeforeRowCount = 0;
+TArray<TypedElementDataStorage::RowHandle> TestRowHandles{};
+
+uint32 CountSettingsRowsInDataStorage()
+{
+	using DSI = ITypedElementDataStorageInterface;
+
+	DSI::FQueryResult Result = DataStorage->RunQuery(CountAllSettingsQuery);
+
+	return Result.Count;
+}
+
+template<typename Func>
+void AwaitRowHandleThenVerify(TypedElementDataStorage::RowHandle RowHandle, const FDoneDelegate& Done, Func&& OnVerify)
+{
+	auto OnTick = [this, RowHandle, Done, OnVerify = Forward<Func>(OnVerify)](float FrameTime)
+		{
+			if (DataStorage->IsRowAssigned(RowHandle))
+			{
+				ON_SCOPE_EXIT{ Done.Execute(); };
+
+				OnVerify();
+
+				return false;
+			}
+			return true;
+		};
+
+	FTSTicker::GetCoreTicker().AddTicker(FTickerDelegate::CreateLambda(OnTick));
+}
+
+END_DEFINE_SPEC(FTedsSettingsTestFixture)
+
+void FTedsSettingsTestFixture::Define()
+{
+	check(GEditor);
+	UTedsSettingsEditorSubsystem* SettingsEditorSubsystem = GEditor->GetEditorSubsystem<UTedsSettingsEditorSubsystem>();
+
+	SettingsEditorSubsystem->OnEnabledChanged().RemoveAll(this);
+	SettingsEditorSubsystem->OnEnabledChanged().AddRaw(this, &FTedsSettingsTestFixture::Redefine);
+
+	if (!SettingsEditorSubsystem->IsEnabled())
+	{
+		return;
+	}
+
+	BeforeEach([this]()
+	{
+		SettingsModule = FModuleManager::GetModulePtr<ISettingsModule>("Settings");
+		check(SettingsModule != nullptr);
+	
+		TypedElementRegistry = UTypedElementRegistry::GetInstance();
+		check(TypedElementRegistry != nullptr);
+	
+		DataStorage = TypedElementRegistry->GetMutableDataStorage();
+		check(DataStorage != nullptr);
+	
+		DataStorageCompatibility = TypedElementRegistry->GetMutableDataStorageCompatibility();
+		check(DataStorageCompatibility != nullptr);
+
+		{
+			using namespace TypedElementQueryBuilder;
+			CountAllSettingsQuery = DataStorage->RegisterQuery(
+				Count()
+				.Where()
+					.All<FTypedElementUObjectColumn, FSettingsContainerColumn, FSettingsCategoryColumn, FSettingsSectionColumn>()
+				.Compile());
+		}
+
+		BeforeRowCount = CountSettingsRowsInDataStorage();
+	});
+
+	AfterEach([this]()
+	{
+		for (TypedElementDataStorage::RowHandle RowHandle : TestRowHandles)
+		{
+			DataStorage->RemoveRow(RowHandle);
+		}
+		TestRowHandles.Empty();
+		CountAllSettingsQuery = TypedElementDataStorage::InvalidQueryHandle;
+		SettingsModule = nullptr;
+		TypedElementRegistry = nullptr;
+		DataStorage = nullptr;
+		DataStorageCompatibility = nullptr;
+	});
+
+	Describe("RegisterSettings", [this]()
+	{
+		LatentIt("Should add a row to editor data storage", [this](const FDoneDelegate& Done)
+		{
+			// TODO: Use an existing container for now because the ISettingsModule does not yet have a delegate for container addition/removal
+			const FName& ContainerName = FName(TEXT(/*"TestContainer"*/"Editor"));
+			const FName& CategoryName = FName(TEXT("TestCategory"));
+			const FName& SectionName = FName(TEXT("TestSection"));
+				
+			TObjectPtr<UTestSettings> TestSettingsObject = NewObject<UTestSettings>();
+			
+			SettingsModule->RegisterSettings(ContainerName, CategoryName, SectionName, FText(), FText(), TestSettingsObject);
+
+			TypedElementDataStorage::RowHandle RowHandle = DataStorageCompatibility->FindRowWithCompatibleObject(TestSettingsObject);
+			TestNotEqual(TEXT("RowHandle"), RowHandle, TypedElementDataStorage::InvalidRowHandle);
+
+			if (RowHandle == TypedElementDataStorage::InvalidRowHandle)
+			{
+				Done.Execute();
+				return;
+			}
+
+			TestRowHandles.Push(RowHandle);
+
+			AwaitRowHandleThenVerify(RowHandle, Done, [this, RowHandle, ContainerName, CategoryName, SectionName]()
+			{
+				uint32 AfterRowCount = CountSettingsRowsInDataStorage();
+
+				TestEqual(TEXT("RowCount"), AfterRowCount, BeforeRowCount + 1);
+				
+				TestEqual(TEXT("ContainerName"), DataStorage->GetColumn<FSettingsContainerColumn>(RowHandle)->ContainerName, ContainerName);
+				TestEqual(TEXT("CategoryName"), DataStorage->GetColumn<FSettingsCategoryColumn>(RowHandle)->CategoryName, CategoryName);
+				TestEqual(TEXT("SectionName"), DataStorage->GetColumn<FSettingsSectionColumn>(RowHandle)->SectionName, SectionName);
+			});
+		});
+	});
+
+	Describe("UnregisterSettings", [this]()
+	{
+		LatentIt("Should remove a row from editor data storage", [this](const FDoneDelegate& Done)
+		{
+			// TODO: Use an existing container for now because the ISettingsModule does not yet have a delegate for container addition/removal
+			const FName& ContainerName = FName(TEXT(/*"TestContainer"*/"Editor"));
+			const FName& CategoryName = FName(TEXT("TestCategory"));
+			const FName& SectionName = FName(TEXT("TestSection"));
+
+			TObjectPtr<UTestSettings> TestSettingsObject = NewObject<UTestSettings>();
+
+			SettingsModule->RegisterSettings(ContainerName, CategoryName, SectionName, FText(), FText(), TestSettingsObject);
+
+			TypedElementDataStorage::RowHandle RowHandle = DataStorageCompatibility->FindRowWithCompatibleObject(TestSettingsObject);
+			TestNotEqual(TEXT("RowHandle"), RowHandle, TypedElementDataStorage::InvalidRowHandle);
+
+			if (RowHandle == TypedElementDataStorage::InvalidRowHandle)
+			{
+				Done.Execute();
+				return;
+			}
+
+			TestRowHandles.Push(RowHandle);
+
+			AwaitRowHandleThenVerify(RowHandle, Done, [this, RowHandle, TestSettingsObject, ContainerName, CategoryName, SectionName]()
+			{
+				uint32 AfterRegisterRowCount = CountSettingsRowsInDataStorage();
+
+				TestEqual(TEXT("RowCount"), AfterRegisterRowCount, BeforeRowCount + 1);
+
+				SettingsModule->UnregisterSettings(ContainerName, CategoryName, SectionName);
+
+				uint32 AfterUnregisterRowCount = CountSettingsRowsInDataStorage();
+				TestEqual(TEXT("RowCount"), AfterUnregisterRowCount, BeforeRowCount);
+
+				TestFalse(TEXT("IsRowAssigned"), DataStorage->IsRowAssigned(RowHandle));
+
+				TypedElementDataStorage::RowHandle InvalidRowHandle = DataStorageCompatibility->FindRowWithCompatibleObject(TestSettingsObject);
+				TestEqual(TEXT("InvalidRowHandle"), InvalidRowHandle, TypedElementDataStorage::InvalidRowHandle);
+			});
+		});
+	});
+}
+
+#endif // WITH_AUTOMATION_TESTS
