ISettingsModule : Add OnContainerAdded() delegate

Previously, there was no way to observe the addition of new containers to a settings module.  The current implementation never removes containers ( although they can be empty ) so there is no need for an OnContainerRemoved delegate at this time.

diff --git a/Engine/Plugins/Experimental/EditorDataStorageFeatures/Source/TedsSettings/Private/TedsSettingsManager.cpp b/Engine/Plugins/Experimental/EditorDataStorageFeatures/Source/TedsSettings/Private/TedsSettingsManager.cpp
--- a/Engine/Plugins/Experimental/EditorDataStorageFeatures/Source/TedsSettings/Private/TedsSettingsManager.cpp
+++ b/Engine/Plugins/Experimental/EditorDataStorageFeatures/Source/TedsSettings/Private/TedsSettingsManager.cpp
@@ -101,36 +101,46 @@ void FTedsSettingsManager::RegisterSettings()
 	ISettingsModule* SettingsModule = FModuleManager::GetModulePtr<ISettingsModule>("Settings");
 	check(SettingsModule);
 
-	// TODO: add a delegate to the ISettingsModule to keep data storage in sync when a container is added or removed
-
 	TArray<FName> ContainerNames;
 	SettingsModule->GetContainerNames(ContainerNames);
 
 	for (FName ContainerName : ContainerNames)
 	{
-		UE_LOG(LogTedsSettings, Log, TEXT("Register Settings Container : '%s'"), *ContainerName.ToString());
+		RegisterSettingsContainer(ContainerName);
+	}
 
-		ISettingsContainerPtr ContainerPtr = SettingsModule->GetContainer(ContainerName);
+	SettingsModule->OnContainerAdded().AddSP(this, &FTedsSettingsManager::RegisterSettingsContainer);
+}
 
-		TArray<ISettingsCategoryPtr> Categories;
-		ContainerPtr->GetCategories(Categories);
+void FTedsSettingsManager::RegisterSettingsContainer(const FName& ContainerName)
+{
+	TRACE_CPUPROFILER_EVENT_SCOPE(TedsSettingsManager.RegisterSettingsContainer);
 
-		for (ISettingsCategoryPtr CategoryPtr : Categories)
-		{
-			const bool bQueryExistingRows = false;
-			UpdateSettingsCategory(CategoryPtr, ContainerName, bQueryExistingRows);
-		}
+	ISettingsModule* SettingsModule = FModuleManager::GetModulePtr<ISettingsModule>("Settings");
+	check(SettingsModule);
 
-		// OnCategoryModified is called at the same time as OnSectionRemoved so we only bind to OnCategoryModified for add / update / remove
-		ContainerPtr->OnCategoryModified().AddSPLambda(this, [this, ContainerPtr](const FName& ModifiedCategoryName)
-			{
-				UE_LOG(LogTedsSettings, Log, TEXT("Settings Category modified : '%s->%s'"), *ContainerPtr->GetName().ToString(), *ModifiedCategoryName.ToString());
+	UE_LOG(LogTedsSettings, Log, TEXT("Register Settings Container : '%s'"), *ContainerName.ToString());
 
-				ISettingsCategoryPtr CategoryPtr = ContainerPtr->GetCategory(ModifiedCategoryName);
+	ISettingsContainerPtr ContainerPtr = SettingsModule->GetContainer(ContainerName);
 
-				UpdateSettingsCategory(CategoryPtr, ContainerPtr->GetName());
-			});
+	TArray<ISettingsCategoryPtr> Categories;
+	ContainerPtr->GetCategories(Categories);
+
+	for (ISettingsCategoryPtr CategoryPtr : Categories)
+	{
+		const bool bQueryExistingRows = false;
+		UpdateSettingsCategory(CategoryPtr, ContainerName, bQueryExistingRows);
 	}
+
+	// OnCategoryModified is called at the same time as OnSectionRemoved so we only bind to OnCategoryModified for add / update / remove
+	ContainerPtr->OnCategoryModified().AddSPLambda(this, [this, ContainerPtr](const FName& ModifiedCategoryName)
+		{
+			UE_LOG(LogTedsSettings, Log, TEXT("Settings Category modified : '%s->%s'"), *ContainerPtr->GetName().ToString(), *ModifiedCategoryName.ToString());
+
+			ISettingsCategoryPtr CategoryPtr = ContainerPtr->GetCategory(ModifiedCategoryName);
+
+			UpdateSettingsCategory(CategoryPtr, ContainerPtr->GetName());
+		});
 }
 
 void FTedsSettingsManager::UnregisterSettings()
@@ -146,6 +156,8 @@ void FTedsSettingsManager::UnregisterSettings()
 	ITypedElementDataStorageCompatibilityInterface* DataStorageCompatibility = TypedElementRegistry->GetMutableDataStorageCompatibility();
 	check(DataStorageCompatibility);
 
+	SettingsModule->OnContainerAdded().RemoveAll(this);
+
 	TArray<FName> ContainerNames;
 	SettingsModule->GetContainerNames(ContainerNames);
 
diff --git a/Engine/Plugins/Experimental/EditorDataStorageFeatures/Source/TedsSettings/Private/TedsSettingsManager.h b/Engine/Plugins/Experimental/EditorDataStorageFeatures/Source/TedsSettings/Private/TedsSettingsManager.h
--- a/Engine/Plugins/Experimental/EditorDataStorageFeatures/Source/TedsSettings/Private/TedsSettingsManager.h
+++ b/Engine/Plugins/Experimental/EditorDataStorageFeatures/Source/TedsSettings/Private/TedsSettingsManager.h
@@ -29,6 +29,8 @@ private:
 	void RegisterSettings();
 	void UnregisterSettings();
 
+	void RegisterSettingsContainer(const FName& ContainerName);
+
 	void UpdateSettingsCategory(TSharedPtr<ISettingsCategory> SettingsCategory, const FName& ContainerName, const bool bQueryExistingRows = true);
 
 	bool bIsInitialized;
diff --git a/Engine/Plugins/Experimental/EditorDataStorageFeatures/Source/TedsSettings/Tests/TedsSettingsTests.cpp b/Engine/Plugins/Experimental/EditorDataStorageFeatures/Source/TedsSettings/Tests/TedsSettingsTests.cpp
--- a/Engine/Plugins/Experimental/EditorDataStorageFeatures/Source/TedsSettings/Tests/TedsSettingsTests.cpp
+++ b/Engine/Plugins/Experimental/EditorDataStorageFeatures/Source/TedsSettings/Tests/TedsSettingsTests.cpp
@@ -114,8 +114,7 @@ void FTedsSettingsTestFixture::Define()
 	{
 		LatentIt("Should add a row to editor data storage", [this](const FDoneDelegate& Done)
 		{
-			// TODO: Use an existing container for now because the ISettingsModule does not yet have a delegate for container addition/removal
-			const FName& ContainerName = FName(TEXT(/*"TestContainer"*/"Editor"));
+			const FName& ContainerName = FName(TEXT("TestContainer"));
 			const FName& CategoryName = FName(TEXT("TestCategory"));
 			const FName& SectionName = FName(TEXT("TestSection"));
 				
@@ -151,8 +150,7 @@ void FTedsSettingsTestFixture::Define()
 	{
 		LatentIt("Should remove a row from editor data storage", [this](const FDoneDelegate& Done)
 		{
-			// TODO: Use an existing container for now because the ISettingsModule does not yet have a delegate for container addition/removal
-			const FName& ContainerName = FName(TEXT(/*"TestContainer"*/"Editor"));
+			const FName& ContainerName = FName(TEXT("TestContainer"));
 			const FName& CategoryName = FName(TEXT("TestCategory"));
 			const FName& SectionName = FName(TEXT("TestSection"));
 
diff --git a/Engine/Source/Developer/Settings/Private/SettingsModule.cpp b/Engine/Source/Developer/Settings/Private/SettingsModule.cpp
--- a/Engine/Source/Developer/Settings/Private/SettingsModule.cpp
+++ b/Engine/Source/Developer/Settings/Private/SettingsModule.cpp
@@ -72,6 +72,11 @@ public:
 		}
 	}
 
+	virtual FOnContainerAdded& OnContainerAdded() override
+	{
+		return ContainerAddedDelegate;
+	}
+
 public:
 
 	// IModuleInterface interface
@@ -125,6 +130,8 @@ protected:
 		if (!Container.IsValid())
 		{
 			Container = MakeShareable(new FSettingsContainer(ContainerName));
+
+			ContainerAddedDelegate.Broadcast(ContainerName);
 		}
 
 		return Container.ToSharedRef();
@@ -155,6 +162,9 @@ private:
 	/** Holds the collection of registered settings viewers. */
 	TMap<FName, ISettingsViewer*> ContainerNamesToViewers;
 
+	/** Holds a delegate that is executed when a settings container has been added. */
+	FOnContainerAdded ContainerAddedDelegate;
+
 #if WITH_RELOAD
 	/** Delegate handle for the re-instancing complete notification */
 	FDelegateHandle OnReloadHandle;
diff --git a/Engine/Source/Developer/Settings/Public/ISettingsModule.h b/Engine/Source/Developer/Settings/Public/ISettingsModule.h
--- a/Engine/Source/Developer/Settings/Public/ISettingsModule.h
+++ b/Engine/Source/Developer/Settings/Public/ISettingsModule.h
@@ -96,6 +96,16 @@ public:
 	 */
 	virtual void UnregisterViewer( const FName& ContainerName ) = 0;
 
+public:
+
+	/**
+	 * A delegate that is executed when a settings container has been added.
+	 *
+	 * @return The delegate.
+	 */
+	DECLARE_MULTICAST_DELEGATE_OneParam(FOnContainerAdded, const FName&)
+	virtual FOnContainerAdded& OnContainerAdded() = 0;
+
 public:
 
 	/** Virtual destructor. */
