Teds Table Viewer: Don't allow the table viewer to display widgets created for other widgets
TEDS UI: Add the label column to widget rows to display their user facing display name

diff --git a/Engine/Plugins/Experimental/EditorDataStorage/Source/TedsUI/Private/Widgets/LabelWidget.cpp b/Engine/Plugins/Experimental/EditorDataStorage/Source/TedsUI/Private/Widgets/LabelWidget.cpp
--- a/Engine/Plugins/Experimental/EditorDataStorage/Source/TedsUI/Private/Widgets/LabelWidget.cpp
+++ b/Engine/Plugins/Experimental/EditorDataStorage/Source/TedsUI/Private/Widgets/LabelWidget.cpp
@@ -216,6 +216,7 @@ TSharedPtr<SWidget> FLabelWidgetConstructor::Construct(
 		{
 			if (FinalizeWidget(DataStorage, DataStorageUi, Row, Result))
 			{
+				AddDefaultWidgetColumns(Row, DataStorage);
 				return Result;
 			}
 		}
diff --git a/Engine/Plugins/Experimental/EditorDataStorageFeatures/Source/TedsDebugger/Private/Widgets/SlateVisualizationWidget.cpp b/Engine/Plugins/Experimental/EditorDataStorageFeatures/Source/TedsDebugger/Private/Widgets/SlateVisualizationWidget.cpp
--- a/Engine/Plugins/Experimental/EditorDataStorageFeatures/Source/TedsDebugger/Private/Widgets/SlateVisualizationWidget.cpp
+++ b/Engine/Plugins/Experimental/EditorDataStorageFeatures/Source/TedsDebugger/Private/Widgets/SlateVisualizationWidget.cpp
@@ -11,7 +11,7 @@
 
 void USlateVisualizationWidgetFactory::RegisterWidgetConstructors(ITypedElementDataStorageInterface& DataStorage, ITypedElementDataStorageUiInterface& DataStorageUi) const
 {
-	DataStorageUi.RegisterWidgetFactory<FSlateVisualizationWidgetConstructor>(FName(TEXT("SceneOutliner.Cell")),
+	DataStorageUi.RegisterWidgetFactory<FSlateVisualizationWidgetConstructor>(FName(TEXT("General.Cell")),
 	TypedElementDataStorage::FColumn<FTypedElementSlateWidgetReferenceColumn>());
 }
 
diff --git a/Engine/Plugins/Experimental/EditorDataStorageFeatures/Source/TedsTableViewer/Private/TedsTableViewerModel.cpp b/Engine/Plugins/Experimental/EditorDataStorageFeatures/Source/TedsTableViewer/Private/TedsTableViewerModel.cpp
--- a/Engine/Plugins/Experimental/EditorDataStorageFeatures/Source/TedsTableViewer/Private/TedsTableViewerModel.cpp
+++ b/Engine/Plugins/Experimental/EditorDataStorageFeatures/Source/TedsTableViewer/Private/TedsTableViewerModel.cpp
@@ -6,6 +6,8 @@
 #include "QueryStack/IQueryStackNode_Row.h"
 #include "TedsTableViewerColumn.h"
 #include "TedsTableViewerUtils.h"
+#include "Elements/Columns/TypedElementMiscColumns.h"
+#include "Elements/Columns/TypedElementSlateWidgetColumns.h"
 
 namespace UE::EditorDataStorage
 {
@@ -42,16 +44,37 @@ namespace UE::EditorDataStorage
 	{
 		Items.Empty();
 		
-		for(TypedElementDataStorage::RowHandle RowHandle : RowQueryStack->GetOrderedRowList())
+		for(const TypedElementDataStorage::RowHandle RowHandle : RowQueryStack->GetOrderedRowList())
 		{
-			Items.Add(RowHandle);
+			if(IsRowDisplayable(RowHandle))
+			{
+				Items.Add(RowHandle);
+			}
 		}
 
 		CachedRowQueryStackRevision = RowQueryStack->GetRevisionId();
 
 		OnModelChanged.Broadcast();
 	}
-	
+
+	bool FTedsTableViewerModel::IsRowDisplayable(TypedElementDataStorage::RowHandle InRowHandle) const
+	{
+		// We don't want to display any second level widgets (widgets for widgets and so on...) because they will keep cause the table viewer to
+		// infinitely grow as you keep scrolling (which creates new widgets)
+		if(Storage->HasColumns<FTypedElementSlateWidgetReferenceColumn>(InRowHandle))
+		{
+			if(const FTypedElementRowReferenceColumn* RowReferenceColumn = Storage->GetColumn<FTypedElementRowReferenceColumn>(InRowHandle))
+			{
+				if(Storage->HasColumns<FTypedElementSlateWidgetReferenceColumn>(RowReferenceColumn->Row))
+				{
+					return false;
+				}
+			}
+		}
+
+		return true;
+	}
+
 	bool FTedsTableViewerModel::Tick(float DeltaTime)
 	{
 		// If the revision ID has changed, refresh to update our rows
diff --git a/Engine/Plugins/Experimental/EditorDataStorageFeatures/Source/TedsTableViewer/Public/TedsTableViewerModel.h b/Engine/Plugins/Experimental/EditorDataStorageFeatures/Source/TedsTableViewer/Public/TedsTableViewerModel.h
--- a/Engine/Plugins/Experimental/EditorDataStorageFeatures/Source/TedsTableViewer/Public/TedsTableViewerModel.h
+++ b/Engine/Plugins/Experimental/EditorDataStorageFeatures/Source/TedsTableViewer/Public/TedsTableViewerModel.h
@@ -71,6 +71,9 @@ namespace UE::EditorDataStorage
 
 		void Refresh();
 
+		// Check whether a row is allowed to be displayed in the table viewer
+		bool IsRowDisplayable(TypedElementDataStorage::RowHandle InRowHandle) const;
+
 	private:
 
 		// The row query stack used to supply the rows to display
diff --git a/Engine/Source/Runtime/TypedElementFramework/Private/Elements/Interfaces/TypedElementDataStorageUiInterface.cpp b/Engine/Source/Runtime/TypedElementFramework/Private/Elements/Interfaces/TypedElementDataStorageUiInterface.cpp
--- a/Engine/Source/Runtime/TypedElementFramework/Private/Elements/Interfaces/TypedElementDataStorageUiInterface.cpp
+++ b/Engine/Source/Runtime/TypedElementFramework/Private/Elements/Interfaces/TypedElementDataStorageUiInterface.cpp
@@ -149,6 +149,7 @@ TSharedPtr<SWidget> FTypedElementWidgetConstructor::Construct(
 		{
 			if (FinalizeWidget(DataStorage, DataStorageUi, Row, Widget))
 			{
+				AddDefaultWidgetColumns(Row, DataStorage);
 				return Widget;
 			}
 		}
@@ -166,12 +167,6 @@ bool FTypedElementWidgetConstructor::SetColumns(ITypedElementDataStorageInterfac
 	return true;
 }
 
-FString FTypedElementWidgetConstructor::GetWidgetLabel(const TSharedPtr<SWidget>& Widget)
-{
-	// The default widget label is simply the type of the widget
-	return Widget->GetType().ToString();
-}
-
 FString FTypedElementWidgetConstructor::DescribeColumnType(const UScriptStruct* ColumnType) const
 {
 	static const FName DisplayNameName(TEXT("DisplayName"));
@@ -197,3 +192,9 @@ bool FTypedElementWidgetConstructor::FinalizeWidget(
 {
 	return true;
 }
+
+void FTypedElementWidgetConstructor::AddDefaultWidgetColumns(TypedElementRowHandle Row, ITypedElementDataStorageInterface* DataStorage) const
+{
+	const FString WidgetLabel(CreateWidgetDisplayName(DataStorage, Row));
+	DataStorage->AddColumn(Row, FTypedElementLabelColumn{.Label = WidgetLabel} );
+}
diff --git a/Engine/Source/Runtime/TypedElementFramework/Public/Elements/Interfaces/TypedElementDataStorageUiInterface.h b/Engine/Source/Runtime/TypedElementFramework/Public/Elements/Interfaces/TypedElementDataStorageUiInterface.h
--- a/Engine/Source/Runtime/TypedElementFramework/Public/Elements/Interfaces/TypedElementDataStorageUiInterface.h
+++ b/Engine/Source/Runtime/TypedElementFramework/Public/Elements/Interfaces/TypedElementDataStorageUiInterface.h
@@ -87,9 +87,6 @@ protected:
 	/** Set any values in columns if needed. The columns provided through GetAdditionalColumnsList() will have already been created. */
 	TYPEDELEMENTFRAMEWORK_API virtual bool SetColumns(ITypedElementDataStorageInterface* DataStorage, TypedElementRowHandle Row);
 	
-	/** Get a user facing label for the widget used in table views. */
-	TYPEDELEMENTFRAMEWORK_API virtual FString GetWidgetLabel(const TSharedPtr<SWidget>& Widget);
-
 	/** Creates a (friendly) name for the provided column type. */
 	TYPEDELEMENTFRAMEWORK_API virtual FString DescribeColumnType(const UScriptStruct* ColumnType) const;
 
@@ -102,7 +99,9 @@ protected:
 		ITypedElementDataStorageUiInterface* DataStorageUi,
 		TypedElementRowHandle Row,
 		const TSharedPtr<SWidget>& Widget);
-	
+
+	/** Add the default misc columns we want a widget row to have. */
+	TYPEDELEMENTFRAMEWORK_API void AddDefaultWidgetColumns(TypedElementRowHandle Row, ITypedElementDataStorageInterface* DataStorage) const;
 protected:
 
 	TArray<TWeakObjectPtr<const UScriptStruct>> MatchedColumnTypes;
