RowHandleWidget: Fix conflicting loctext macros

diff --git a/Engine/Plugins/Experimental/EditorDataStorageFeatures/Source/TedsTableViewer/Private/Widgets/RowHandleWidget.cpp b/Engine/Plugins/Experimental/EditorDataStorageFeatures/Source/TedsTableViewer/Private/Widgets/RowHandleWidget.cpp
--- a/Engine/Plugins/Experimental/EditorDataStorageFeatures/Source/TedsTableViewer/Private/Widgets/RowHandleWidget.cpp
+++ b/Engine/Plugins/Experimental/EditorDataStorageFeatures/Source/TedsTableViewer/Private/Widgets/RowHandleWidget.cpp
@@ -17,13 +17,13 @@ void URowHandleWidgetFactory::RegisterWidgetConstructors(ITypedElementDataStorag
 void URowHandleWidgetFactory::RegisterWidgetPurposes(ITypedElementDataStorageUiInterface& DataStorageUi) const
 {
 	DataStorageUi.RegisterWidgetPurpose(FName(TEXT("General.Cell.RowHandle")), ITypedElementDataStorageUiInterface::EPurposeType::UniqueByName,
-	LOCTEXT("RowHandlePurpose", "Specific purpose to request a widget to display row handles."));
+	LOCTEXT("GeneralRowHandlePurpose", "Specific purpose to request a widget to display row handles."));
 
 	DataStorageUi.RegisterWidgetPurpose(FName(TEXT("RowDetails.Cell")), ITypedElementDataStorageUiInterface::EPurposeType::UniqueByNameAndColumn,
-	LOCTEXT("RowHandlePurpose", "Specific purpose to request a widget to display the details on a row (e.g SRowDetails)."));
+	LOCTEXT("DetailsRowHandlePurpose", "Specific purpose to request a widget to display the details on a row (e.g SRowDetails)."));
 
 	DataStorageUi.RegisterWidgetPurpose(FName(TEXT("RowDetails.Cell.Large")), ITypedElementDataStorageUiInterface::EPurposeType::UniqueByNameAndColumn,
-	LOCTEXT("RowHandlePurpose", "Specific purpose to request a widget that is larger than a single cell to display the details on a row (e.g SRowDetails)"));
+	LOCTEXT("LargeDetailsRowHandlePurpose", "Specific purpose to request a widget that is larger than a single cell to display the details on a row (e.g SRowDetails)"));
 
 
 }
