TEDS-Outliner: Add functionality to rename items by clicking on the label widget twice

diff --git a/Engine/Plugins/Experimental/EditorDataStorage/Source/TedsUI/Private/Widgets/LabelWidget.cpp b/Engine/Plugins/Experimental/EditorDataStorage/Source/TedsUI/Private/Widgets/LabelWidget.cpp
--- a/Engine/Plugins/Experimental/EditorDataStorage/Source/TedsUI/Private/Widgets/LabelWidget.cpp
+++ b/Engine/Plugins/Experimental/EditorDataStorage/Source/TedsUI/Private/Widgets/LabelWidget.cpp
@@ -3,6 +3,7 @@
 #include "Widgets/LabelWidget.h"
 
 #include "ActorEditorUtils.h"
+#include "Columns/SlateDelegateColumns.h"
 #include "Elements/Columns/TypedElementLabelColumns.h"
 #include "Elements/Columns/TypedElementMiscColumns.h"
 #include "Elements/Columns/TypedElementSlateWidgetColumns.h"
@@ -155,7 +156,8 @@ TConstArrayView<const UScriptStruct*> FLabelWidgetConstructor::GetAdditionalColu
 	static const TTypedElementColumnTypeList<
 		FTypedElementRowReferenceColumn,
 		FTypedElementU64IntValueCacheColumn,
-		FLabelWidgetColumn> Columns;
+		FLabelWidgetColumn,
+		FExternalWidgetSelectionColumn> Columns;
 	return Columns;
 }
 
@@ -192,7 +194,8 @@ TSharedPtr<SWidget> FLabelWidgetConstructor::Construct(
 						// Note: The use of actor specific functionality should be minimized, but this function acts generic enough that the 
 						// use of actor is just in names.
 						return FActorEditorUtils::ValidateActorName(Label, ErrorMessage);
-					});
+					})
+				.IsSelected_Static(&FLabelWidgetConstructor::IsWidgetSelected, DataStorage, Row);;
 			TextBlock->AddMetadata(MakeShared<TTypedElementUiEditableCapability<SInlineEditableTextBlock>>(*TextBlock));
 			TextBlock->AddMetadata(MakeShared<TTypedElementUiTextCapability<SInlineEditableTextBlock>>(*TextBlock));
 			TextBlock->AddMetadata(MakeShared<TTypedElementUiTooltipCapability<SInlineEditableTextBlock>>(*TextBlock));
@@ -244,4 +247,16 @@ bool FLabelWidgetConstructor::FinalizeWidget(
 		Widget, MatchedColumnTypes.Num() == 2);
 	return true;
 }
+
+bool FLabelWidgetConstructor::IsWidgetSelected(ITypedElementDataStorageInterface* DataStorage, TypedElementRowHandle UiRow)
+{
+	if (const FExternalWidgetSelectionColumn* ExternalWidgetSelectionColumn = DataStorage->GetColumn<FExternalWidgetSelectionColumn>(UiRow))
+	{
+		if (ExternalWidgetSelectionColumn->IsSelected.IsBound())
+		{
+			return ExternalWidgetSelectionColumn->IsSelected.Execute();
+		}
+	}
+	return false;
+}
 #undef LOCTEXT_NAMESPACE
diff --git a/Engine/Plugins/Experimental/EditorDataStorage/Source/TedsUI/Public/Widgets/LabelWidget.h b/Engine/Plugins/Experimental/EditorDataStorage/Source/TedsUI/Public/Widgets/LabelWidget.h
--- a/Engine/Plugins/Experimental/EditorDataStorage/Source/TedsUI/Public/Widgets/LabelWidget.h
+++ b/Engine/Plugins/Experimental/EditorDataStorage/Source/TedsUI/Public/Widgets/LabelWidget.h
@@ -47,6 +47,7 @@ protected:
 	bool SetColumns(ITypedElementDataStorageInterface* DataStorage, TypedElementRowHandle Row) override;
 	bool FinalizeWidget(ITypedElementDataStorageInterface* DataStorage, ITypedElementDataStorageUiInterface* DataStorageUi,
 		TypedElementRowHandle Row, const TSharedPtr<SWidget>& Widget) override;
+	static bool IsWidgetSelected(ITypedElementDataStorageInterface* DataStorage, TypedElementRowHandle UiRow);
 };
 
 USTRUCT(meta = (DisplayName = "Label widget"))
diff --git a/Engine/Plugins/Experimental/EditorDataStorage/Source/TedsUI/TedsUI.Build.cs b/Engine/Plugins/Experimental/EditorDataStorage/Source/TedsUI/TedsUI.Build.cs
--- a/Engine/Plugins/Experimental/EditorDataStorage/Source/TedsUI/TedsUI.Build.cs
+++ b/Engine/Plugins/Experimental/EditorDataStorage/Source/TedsUI/TedsUI.Build.cs
@@ -23,7 +23,13 @@ public class TedsUI : ModuleRules
 					"SlateCore",
 					"ToolMenus",
 					"TypedElementFramework",
-					"TedsCore",
+					"TedsCore"
+				});
+			
+			PrivateDependencyModuleNames.AddRange(
+				new string[]
+				{
+					"ToolWidgets",
 				});
 
 			DynamicallyLoadedModuleNames.AddRange(new string[] {});
diff --git a/Engine/Plugins/Experimental/EditorDataStorageFeatures/Source/TedsOutliner/Private/TedsOutlinerImpl.cpp b/Engine/Plugins/Experimental/EditorDataStorageFeatures/Source/TedsOutliner/Private/TedsOutlinerImpl.cpp
--- a/Engine/Plugins/Experimental/EditorDataStorageFeatures/Source/TedsOutliner/Private/TedsOutlinerImpl.cpp
+++ b/Engine/Plugins/Experimental/EditorDataStorageFeatures/Source/TedsOutliner/Private/TedsOutlinerImpl.cpp
@@ -11,6 +11,7 @@
 #include "Elements/Interfaces/TypedElementDataStorageUiInterface.h"
 #include "Elements/Framework/TypedElementQueryBuilder.h"
 #include "Elements/Framework/TypedElementRegistry.h"
+#include "Columns/SlateDelegateColumns.h"
 #include "Compatibility/SceneOutlinerRowHandleColumn.h"
 #include "TedsOutlinerFilter.h"
 #include "TedsOutlinerItem.h"
@@ -216,9 +217,9 @@ void FTedsOutlinerImpl::SetSelection(const TArray<TypedElementDataStorage::RowHa
 	}
 }
 
-TSharedRef<SWidget> FTedsOutlinerImpl::CreateLabelWidgetForItem(TypedElementRowHandle InRowHandle) const
+TSharedRef<SWidget> FTedsOutlinerImpl::CreateLabelWidgetForItem(TypedElementRowHandle InRowHandle, const STableRow<FSceneOutlinerTreeItemPtr>& InRow) const
 {
-	auto CreateWidgetForQuery = [InRowHandle, this](const TPair<TypedElementDataStorage::QueryHandle, TSharedPtr<FTypedElementWidgetConstructor>>& QueryConstructorPair) -> TSharedPtr<SWidget>
+	auto CreateWidgetForQuery = [InRowHandle, this, &InRow](const TPair<TypedElementDataStorage::QueryHandle, TSharedPtr<FTypedElementWidgetConstructor>>& QueryConstructorPair) -> TSharedPtr<SWidget>
 	{
 		TypedElementDataStorage::FQueryDescription QueryDescription = Storage->GetQueryDescription(QueryConstructorPair.Key);
 		
@@ -245,7 +246,14 @@ TSharedRef<SWidget> FTedsOutlinerImpl::CreateLabelWidgetForItem(TypedElementRowH
 
 		Storage->AddColumn(UiRowHandle, FTedsOutlinerColumn{.Outliner = StaticCastSharedRef<ISceneOutliner>(SceneOutliner->AsShared())});
 		
-		return StorageUi->ConstructWidget(UiRowHandle, *CellWidgetConstructor, MetaDataArgs);
+		TSharedPtr<SWidget> Widget = StorageUi->ConstructWidget(UiRowHandle, *CellWidgetConstructor, MetaDataArgs);
+		
+		if (FExternalWidgetSelectionColumn* ExternalWidgetSelectionColumn = Storage->GetColumn<FExternalWidgetSelectionColumn>(UiRowHandle))
+		{
+			ExternalWidgetSelectionColumn->IsSelected = FIsSelected::CreateSP(&InRow, &STableRow<FSceneOutlinerTreeItemPtr>::IsSelectedExclusively);
+		}
+		return Widget;
+		
 	};
 	
 	TSharedRef<SHorizontalBox> CombinedWidget = SNew(SHorizontalBox);
diff --git a/Engine/Plugins/Experimental/EditorDataStorageFeatures/Source/TedsOutliner/Private/TedsOutlinerItem.cpp b/Engine/Plugins/Experimental/EditorDataStorageFeatures/Source/TedsOutliner/Private/TedsOutlinerItem.cpp
--- a/Engine/Plugins/Experimental/EditorDataStorageFeatures/Source/TedsOutliner/Private/TedsOutlinerItem.cpp
+++ b/Engine/Plugins/Experimental/EditorDataStorageFeatures/Source/TedsOutliner/Private/TedsOutlinerItem.cpp
@@ -42,7 +42,7 @@ bool FTedsOutlinerTreeItem::CanInteract() const
 TSharedRef<SWidget> FTedsOutlinerTreeItem::GenerateLabelWidget(ISceneOutliner& Outliner,
 	const STableRow<FSceneOutlinerTreeItemPtr>& InRow)
 {
-	return TedsOutlinerImpl->CreateLabelWidgetForItem(RowHandle);
+	return TedsOutlinerImpl->CreateLabelWidgetForItem(RowHandle, InRow);
 }
 
 void FTedsOutlinerTreeItem::GenerateContextMenu(UToolMenu* Menu, SSceneOutliner& Outliner)
diff --git a/Engine/Plugins/Experimental/EditorDataStorageFeatures/Source/TedsOutliner/Public/TedsOutlinerImpl.h b/Engine/Plugins/Experimental/EditorDataStorageFeatures/Source/TedsOutliner/Public/TedsOutlinerImpl.h
--- a/Engine/Plugins/Experimental/EditorDataStorageFeatures/Source/TedsOutliner/Public/TedsOutlinerImpl.h
+++ b/Engine/Plugins/Experimental/EditorDataStorageFeatures/Source/TedsOutliner/Public/TedsOutlinerImpl.h
@@ -135,7 +135,7 @@ public:
 	void SetSelection(const TArray<TypedElementDataStorage::RowHandle>& InSelectedRows);
 
 	// Helper function to create a label widget for a given row
-	TSharedRef<SWidget> CreateLabelWidgetForItem(TypedElementRowHandle InRowHandle) const;
+	TSharedRef<SWidget> CreateLabelWidgetForItem(TypedElementRowHandle InRowHandle, const STableRow<FSceneOutlinerTreeItemPtr>& InRow) const;
 
 	// Get the hierarchy data associated with this table viewer
 	const TOptional<FTedsOutlinerHierarchyData>& GetHierarchyData();
diff --git a/Engine/Plugins/Experimental/EditorDataStorageFeatures/Source/TedsOutliner/TedsOutliner.Build.cs b/Engine/Plugins/Experimental/EditorDataStorageFeatures/Source/TedsOutliner/TedsOutliner.Build.cs
--- a/Engine/Plugins/Experimental/EditorDataStorageFeatures/Source/TedsOutliner/TedsOutliner.Build.cs
+++ b/Engine/Plugins/Experimental/EditorDataStorageFeatures/Source/TedsOutliner/TedsOutliner.Build.cs
@@ -32,7 +32,8 @@ public class TedsOutliner : ModuleRules
 					"UnrealEd", // FEditorUndoClient used by SSceneOutliner
 					"ToolMenus",
 					"ApplicationCore",
-					"TedsTableViewer"
+					"TedsTableViewer",
+					"ToolWidgets"
 				});
 
 			DynamicallyLoadedModuleNames.AddRange(new string[] {});
diff --git a/Engine/Source/Developer/ToolWidgets/Public/Columns/SlateDelegateColumns.h b/Engine/Source/Developer/ToolWidgets/Public/Columns/SlateDelegateColumns.h
new file mode 100644
--- /dev/null
+++ b/Engine/Source/Developer/ToolWidgets/Public/Columns/SlateDelegateColumns.h
@@ -0,0 +1,26 @@
+// Copyright Epic Games, Inc. All Rights Reserved.
+
+#pragma once
+
+#include "Elements/Interfaces/TypedElementDataStorageInterface.h"
+#include "Framework/SlateDelegates.h"
+
+#include "SlateDelegateColumns.generated.h"
+
+/**
+ * Column added to a widget row when an external widget manages selection for the widget referenced by the row,
+ * such as an owning SListView or STreeView
+ */
+USTRUCT(meta = (DisplayName = "Widget with externally managed selection"))
+struct FExternalWidgetSelectionColumn : public FEditorDataStorageColumn
+{
+	GENERATED_BODY()
+	/** Delegate to execute to check the status of if the widget is selected or not
+	 * Only needs to be hooked up if an external widget is managing selection, such
+	 * as a SListView or STreeView.
+	 *
+	 * E.g usage: Enables renaming by clicking the label twice if the widget is editable.
+	 * @see SInlineEditableTextBlock::IsSelected
+	 */
+	FIsSelected IsSelected;
+};
\ No newline at end of file
diff --git a/Engine/Source/Developer/ToolWidgets/ToolWidgets.Build.cs b/Engine/Source/Developer/ToolWidgets/ToolWidgets.Build.cs
--- a/Engine/Source/Developer/ToolWidgets/ToolWidgets.Build.cs
+++ b/Engine/Source/Developer/ToolWidgets/ToolWidgets.Build.cs
@@ -17,7 +17,8 @@ public class ToolWidgets : ModuleRules
 				"SlateCore",
 				"InputCore",
 				"ToolMenus",
-				"CoreUObject"
+				"CoreUObject",
+				"TypedElementFramework"
 			}
 		);
 
