Universal Search prototype - Core and UI + TEDS Search Provider

This change introduces the core framework and UI for the Universal Search. It doesn't include any search providers and therefore can't search yet. It does however include:
- an input bar that can be docked in the bottom status bar (default) or top main menu.
- configuration options.
- a search providers interface that support async searches.
- a result window with filters, groups and search results.

Adds support to the Universal Search to search through TEDS. Currently only FString, FText and FName are supported. Properties on columns that have the meta data tag "Searchable" will be included in the search. Several columns in TEDS have gotten this meta tag with in this change list. For each combination of group and column a unique query is created so the work can be time sliced better. A miscellaneous query is also created that excludes all other groups so searchable columns that weren't in a group still produce results.

Seperate CRs can be found under 35437674 and 35446965.

#rnx

diff --git a/Engine/Plugins/Experimental/EditorDataStorageFeatures/Source/TedsAssetData/Public/TedsAssetDataColumns.h b/Engine/Plugins/Experimental/EditorDataStorageFeatures/Source/TedsAssetData/Public/TedsAssetDataColumns.h
--- a/Engine/Plugins/Experimental/EditorDataStorageFeatures/Source/TedsAssetData/Public/TedsAssetDataColumns.h
+++ b/Engine/Plugins/Experimental/EditorDataStorageFeatures/Source/TedsAssetData/Public/TedsAssetDataColumns.h
@@ -14,7 +14,7 @@ struct FAssetPathColumn_Experimental : public FEditorDataStorageColumn
 {
 	GENERATED_BODY()
 	
-	UPROPERTY()
+	UPROPERTY(meta = (Searchable))
 	FName Path;
 };
 
@@ -127,7 +127,7 @@ struct FVirtualPathColumn_Experimental : public FEditorDataStorageColumn
 {
 	GENERATED_BODY()
 
-	UPROPERTY()
+	UPROPERTY(meta = (Searchable))
 	FName VirtualPath;
 };
 
@@ -136,6 +136,6 @@ struct FItemNameColumn_Experimental : public FEditorDataStorageColumn
 {
 	GENERATED_BODY()
 
-	UPROPERTY()
+	UPROPERTY(meta = (Searchable))
 	FName Name;
 };
\ No newline at end of file
diff --git a/Engine/Plugins/Experimental/EditorDataStorageFeatures/Source/TedsSettings/Private/TedsSettingsManager.cpp b/Engine/Plugins/Experimental/EditorDataStorageFeatures/Source/TedsSettings/Private/TedsSettingsManager.cpp
--- a/Engine/Plugins/Experimental/EditorDataStorageFeatures/Source/TedsSettings/Private/TedsSettingsManager.cpp
+++ b/Engine/Plugins/Experimental/EditorDataStorageFeatures/Source/TedsSettings/Private/TedsSettingsManager.cpp
@@ -270,6 +270,7 @@ void FTedsSettingsManager::UpdateSettingsCategory(TSharedPtr<ISettingsCategory>
 			DataStorage->AddColumn<FSettingsContainerColumn>(NewRow, { .ContainerName = ContainerName });
 			DataStorage->AddColumn<FSettingsCategoryColumn>(NewRow, { .CategoryName = CategoryName });
 			DataStorage->AddColumn<FSettingsSectionColumn>(NewRow, { .SectionName = SectionName });
+			DataStorage->AddColumn<FSettingsTag>(NewRow);
 
 			UE_LOG(LogTedsSettings, Log, TEXT("Added Settings Section : '%s'"), *SectionName.ToString());
 		}
diff --git a/Engine/Plugins/Experimental/EditorDataStorageFeatures/Source/TedsSettings/Public/TedsSettingsColumns.h b/Engine/Plugins/Experimental/EditorDataStorageFeatures/Source/TedsSettings/Public/TedsSettingsColumns.h
--- a/Engine/Plugins/Experimental/EditorDataStorageFeatures/Source/TedsSettings/Public/TedsSettingsColumns.h
+++ b/Engine/Plugins/Experimental/EditorDataStorageFeatures/Source/TedsSettings/Public/TedsSettingsColumns.h
@@ -12,7 +12,7 @@ struct FSettingsContainerColumn final : public FEditorDataStorageColumn
 {
 	GENERATED_BODY()
 
-	UPROPERTY()
+	UPROPERTY(meta = (Searchable))
 	FName ContainerName;
 };
 
@@ -21,7 +21,7 @@ struct FSettingsCategoryColumn final : public FEditorDataStorageColumn
 {
 	GENERATED_BODY()
 
-	UPROPERTY()
+	UPROPERTY(meta = (Searchable))
 	FName CategoryName;
 };
 
@@ -30,6 +30,12 @@ struct FSettingsSectionColumn final : public FEditorDataStorageColumn
 {
 	GENERATED_BODY()
 
-	UPROPERTY()
+	UPROPERTY(meta = (Searchable))
 	FName SectionName;
 };
+
+USTRUCT(meta = (DisplayName = "Settings"))
+struct FSettingsTag final : public FEditorDataStorageTag
+{
+	GENERATED_BODY()
+};
diff --git a/Engine/Source/Runtime/TypedElementFramework/Public/Elements/Columns/TypedElementLabelColumns.h b/Engine/Source/Runtime/TypedElementFramework/Public/Elements/Columns/TypedElementLabelColumns.h
--- a/Engine/Source/Runtime/TypedElementFramework/Public/Elements/Columns/TypedElementLabelColumns.h
+++ b/Engine/Source/Runtime/TypedElementFramework/Public/Elements/Columns/TypedElementLabelColumns.h
@@ -16,7 +16,7 @@ struct FTypedElementLabelColumn final : public FEditorDataStorageColumn
 {
 	GENERATED_BODY()
 
-	UPROPERTY()
+	UPROPERTY(meta = (Searchable))
 	FString Label;
 };
 
diff --git a/Engine/Source/Runtime/TypedElementFramework/Public/Elements/Columns/TypedElementPackageColumns.h b/Engine/Source/Runtime/TypedElementFramework/Public/Elements/Columns/TypedElementPackageColumns.h
--- a/Engine/Source/Runtime/TypedElementFramework/Public/Elements/Columns/TypedElementPackageColumns.h
+++ b/Engine/Source/Runtime/TypedElementFramework/Public/Elements/Columns/TypedElementPackageColumns.h
@@ -65,7 +65,7 @@ struct FTypedElementPackagePathColumn final : public FEditorDataStorageColumn
 {
 	GENERATED_BODY()
 
-	UPROPERTY()
+	UPROPERTY(meta = (Searchable))
 	FString Path;
 };
 
