Scene Outliner Dealiaser

Provides access to the item dealiaser for the Outliner. This can be used to find Outliner row items from a TEDS row handle, which in turn can be used for selection.

diff --git a/Engine/Plugins/Experimental/EditorDataStorageFeatures/Source/TedsOutliner/Private/Compatibility/SceneOutlinerTedsBridge.cpp b/Engine/Plugins/Experimental/EditorDataStorageFeatures/Source/TedsOutliner/Private/Compatibility/SceneOutlinerTedsBridge.cpp
--- a/Engine/Plugins/Experimental/EditorDataStorageFeatures/Source/TedsOutliner/Private/Compatibility/SceneOutlinerTedsBridge.cpp
+++ b/Engine/Plugins/Experimental/EditorDataStorageFeatures/Source/TedsOutliner/Private/Compatibility/SceneOutlinerTedsBridge.cpp
@@ -35,8 +35,8 @@ FAutoConsoleCommand BindColumnsToSceneOutlinerConsoleCommand(
 	FConsoleCommandWithArgsDelegate::CreateLambda([](const TArray<FString>& Args)
 		{
 			using namespace TypedElementQueryBuilder;
-			using DSI = ITypedElementDataStorageInterface;
-
+			using namespace TypedElementDataStorage;
+			
 		    const FName WidgetPurposes[] = {TEXT("SceneOutliner.Cell"), TEXT("General.Cell")};
 
 			UTypedElementRegistry* Registry = UTypedElementRegistry::GetInstance();
@@ -96,8 +96,8 @@ FAutoConsoleCommand BindColumnsToSceneOutlinerConsoleCommand(
 							}
 							if (AdditionCount > 0)
 							{
-								static TypedElementQueryHandle CustomQuery = TypedElementInvalidQueryHandle;
-								if (CustomQuery != TypedElementInvalidQueryHandle)
+								static TypedElementQueryHandle CustomQuery = InvalidQueryHandle;
+								if (CustomQuery != InvalidQueryHandle)
 								{
 									DataStorage->UnregisterQuery(CustomQuery);
 								}
@@ -125,6 +125,7 @@ public:
 
 	void AssignQuery(TypedElementQueryHandle Query, const TConstArrayView<FName> CellWidgetPurposes);
 	void RegisterDealiaser(const FTreeItemIDDealiaser& InDealiaser);
+	FTreeItemIDDealiaser GetDealiaser();
 private:
 	void ClearColumns(ISceneOutliner& InOutliner);
 
@@ -415,6 +416,11 @@ TSharedPtr<FSceneOutlinerTedsBridge>* FSceneOutlinerTedsQueryBinder::FindOrAddQu
 	return QueryMapping;
 }
 
+TSharedPtr<FSceneOutlinerTedsBridge>* FSceneOutlinerTedsQueryBinder::FindQueryMapping(const TSharedPtr<ISceneOutliner>& Outliner)
+{
+	return SceneOutliners.Find(Outliner);
+}
+
 void FSceneOutlinerTedsQueryBinder::AssignQuery(TypedElementQueryHandle Query, const TSharedPtr<ISceneOutliner>& Outliner, TConstArrayView<FName> CellWidgetPurposes)
 {
 	CleanupStaleOutliners();
@@ -429,6 +435,18 @@ void FSceneOutlinerTedsQueryBinder::RegisterTreeItemIDDealiaser(const TSharedPtr
 	(*QueryMapping)->RegisterDealiaser(InDealiaser);
 }
 
+FTreeItemIDDealiaser FSceneOutlinerTedsQueryBinder::GetTreeItemIDDealiaser(const TSharedPtr<ISceneOutliner>& Widget)
+{
+	TSharedPtr<FSceneOutlinerTedsBridge>* QueryMapping = FindQueryMapping(Widget);
+
+	if(QueryMapping)
+	{
+		return (*QueryMapping)->GetDealiaser();
+	}
+
+	return FTreeItemIDDealiaser();
+}
+
 void FSceneOutlinerTedsQueryBinder::CleanupStaleOutliners()
 {
 	for (TMap<TWeakPtr<ISceneOutliner>, TSharedPtr<FSceneOutlinerTedsBridge>>::TIterator It(SceneOutliners); It; ++It)
@@ -471,6 +489,12 @@ void FSceneOutlinerTedsBridge::RegisterDealiaser(const FTreeItemIDDealiaser& InD
 	Dealiaser = InDealiaser;
 }
 
+FTreeItemIDDealiaser FSceneOutlinerTedsBridge::GetDealiaser()
+{
+	return Dealiaser;
+}
+
+
 void FSceneOutlinerTedsBridge::AssignQuery(TypedElementQueryHandle Query, const TConstArrayView<FName> InCellWidgetPurposes)
 {
 	using MatchApproach = ITypedElementDataStorageUiInterface::EMatchApproach;
diff --git a/Engine/Plugins/Experimental/EditorDataStorageFeatures/Source/TedsOutliner/Public/Compatibility/SceneOutlinerTedsBridge.h b/Engine/Plugins/Experimental/EditorDataStorageFeatures/Source/TedsOutliner/Public/Compatibility/SceneOutlinerTedsBridge.h
--- a/Engine/Plugins/Experimental/EditorDataStorageFeatures/Source/TedsOutliner/Public/Compatibility/SceneOutlinerTedsBridge.h
+++ b/Engine/Plugins/Experimental/EditorDataStorageFeatures/Source/TedsOutliner/Public/Compatibility/SceneOutlinerTedsBridge.h
@@ -41,9 +41,12 @@ public:
 
 	void AssignQuery(TypedElementQueryHandle Query, const TSharedPtr<ISceneOutliner>& Widget, TConstArrayView<FName> InCellWidgetPurposes);
 
-	// Register a dealiser for a specific TEDS-Outliner to convert a row handle to an FSceneOutlinerTreeItemID
+	// Register a dealiaser for a specific TEDS-Outliner to convert a row handle to an FSceneOutlinerTreeItemID
 	void RegisterTreeItemIDDealiaser(const TSharedPtr<ISceneOutliner>& Widget, const FTreeItemIDDealiaser& InDealiaser);
 
+	// Get the dealiaser for a specific outliner instance
+	FTreeItemIDDealiaser GetTreeItemIDDealiaser(const TSharedPtr<ISceneOutliner>& Widget);
+
 	// Get the name of the Outliner column corresponding to the given TEDS column (if any)
 	FName FindOutlinerColumnFromTEDSColumns(TConstArrayView<TWeakObjectPtr<const UScriptStruct>> TEDSColumns) const;
 
@@ -53,6 +56,8 @@ private:
 	void CleanupStaleOutliners();
 	
 	TSharedPtr<FSceneOutlinerTedsBridge>* FindOrAddQueryMapping(const TSharedPtr<ISceneOutliner>& Widget);
+	TSharedPtr<FSceneOutlinerTedsBridge>* FindQueryMapping(const TSharedPtr<ISceneOutliner>& Widget);
+
 	
 	TMap<TWeakPtr<ISceneOutliner>, TSharedPtr<FSceneOutlinerTedsBridge>> SceneOutliners;
 
