[Teds] Update naming for RowHandle, QueryHandle, and TableHandle

diff --git a/Engine/Plugins/Experimental/EditorDataStorage/Source/TedsUI/Public/Widgets/ExportedTextWidget.h b/Engine/Plugins/Experimental/EditorDataStorage/Source/TedsUI/Public/Widgets/ExportedTextWidget.h
--- a/Engine/Plugins/Experimental/EditorDataStorage/Source/TedsUI/Public/Widgets/ExportedTextWidget.h
+++ b/Engine/Plugins/Experimental/EditorDataStorage/Source/TedsUI/Public/Widgets/ExportedTextWidget.h
@@ -37,14 +37,14 @@ public:
 	virtual TConstArrayView<const UScriptStruct*> GetAdditionalColumnsList() const override;
 	virtual const TypedElementDataStorage::FQueryConditions* GetQueryConditions() const override;
 	virtual FString CreateWidgetDisplayName(
-		ITypedElementDataStorageInterface* DataStorage, TypedElementDataStorage::RowHandle Row) const override;
+		ITypedElementDataStorageInterface* DataStorage, UE::Editor::DataStorage::RowHandle Row) const override;
 
 protected:
 	virtual TSharedPtr<SWidget> CreateWidget(const TypedElementDataStorage::FMetaDataView& Arguments) override;
 	virtual bool FinalizeWidget(
 		ITypedElementDataStorageInterface* DataStorage,
 		ITypedElementDataStorageUiInterface* DataStorageUi,
-		TypedElementRowHandle Row,
+		UE::Editor::DataStorage::RowHandle Row,
 		const TSharedPtr<SWidget>& Widget) override;
 
 protected:
diff --git a/Engine/Plugins/Experimental/EditorDataStorage/Source/TedsUI/Public/Widgets/LabelWidget.h b/Engine/Plugins/Experimental/EditorDataStorage/Source/TedsUI/Public/Widgets/LabelWidget.h
--- a/Engine/Plugins/Experimental/EditorDataStorage/Source/TedsUI/Public/Widgets/LabelWidget.h
+++ b/Engine/Plugins/Experimental/EditorDataStorage/Source/TedsUI/Public/Widgets/LabelWidget.h
@@ -37,17 +37,17 @@ public:
 	TEDSUI_API TConstArrayView<const UScriptStruct*> GetAdditionalColumnsList() const override;
 
 	TEDSUI_API TSharedPtr<SWidget> Construct(
-		TypedElementRowHandle Row, 
+		UE::Editor::DataStorage::RowHandle Row,
 		ITypedElementDataStorageInterface* DataStorage,
 		ITypedElementDataStorageUiInterface* DataStorageUi,
 		const TypedElementDataStorage::FMetaDataView& Arguments) override;
 
 protected:
 	explicit FLabelWidgetConstructor(const UScriptStruct* InTypeInfo);
-	bool SetColumns(ITypedElementDataStorageInterface* DataStorage, TypedElementRowHandle Row) override;
+	bool SetColumns(ITypedElementDataStorageInterface* DataStorage, UE::Editor::DataStorage::RowHandle Row) override;
 	bool FinalizeWidget(ITypedElementDataStorageInterface* DataStorage, ITypedElementDataStorageUiInterface* DataStorageUi,
-		TypedElementRowHandle Row, const TSharedPtr<SWidget>& Widget) override;
-	static bool IsWidgetSelected(ITypedElementDataStorageInterface* DataStorage, TypedElementRowHandle UiRow);
+		UE::Editor::DataStorage::RowHandle Row, const TSharedPtr<SWidget>& Widget) override;
+	static bool IsWidgetSelected(ITypedElementDataStorageInterface* DataStorage, UE::Editor::DataStorage::RowHandle UiRow);
 };
 
 USTRUCT(meta = (DisplayName = "Label widget"))
diff --git a/Engine/Plugins/Experimental/EditorDataStorage/Source/TedsUI/Public/Widgets/PackagePathWidget.h b/Engine/Plugins/Experimental/EditorDataStorage/Source/TedsUI/Public/Widgets/PackagePathWidget.h
--- a/Engine/Plugins/Experimental/EditorDataStorage/Source/TedsUI/Public/Widgets/PackagePathWidget.h
+++ b/Engine/Plugins/Experimental/EditorDataStorage/Source/TedsUI/Public/Widgets/PackagePathWidget.h
@@ -37,7 +37,7 @@ protected:
 
 	TSharedPtr<SWidget> CreateWidget(const TypedElementDataStorage::FMetaDataView& Arguments) override;
 	bool FinalizeWidget(ITypedElementDataStorageInterface* DataStorage, ITypedElementDataStorageUiInterface* DataStorageUi,
-		TypedElementRowHandle Row, const TSharedPtr<SWidget>& Widget) override;
+		UE::Editor::DataStorage::RowHandle Row, const TSharedPtr<SWidget>& Widget) override;
 };
 
 USTRUCT()
@@ -51,5 +51,5 @@ public:
 
 protected:
 	bool FinalizeWidget(ITypedElementDataStorageInterface* DataStorage, ITypedElementDataStorageUiInterface* DataStorageUi,
-		TypedElementRowHandle Row, const TSharedPtr<SWidget>& Widget) override;
+		UE::Editor::DataStorage::RowHandle Row, const TSharedPtr<SWidget>& Widget) override;
 };
diff --git a/Engine/Plugins/Experimental/EditorDataStorageFeatures/Source/TedsAssetData/Private/CB/TedsAssetDataCBDataSource.cpp b/Engine/Plugins/Experimental/EditorDataStorageFeatures/Source/TedsAssetData/Private/CB/TedsAssetDataCBDataSource.cpp
--- a/Engine/Plugins/Experimental/EditorDataStorageFeatures/Source/TedsAssetData/Private/CB/TedsAssetDataCBDataSource.cpp
+++ b/Engine/Plugins/Experimental/EditorDataStorageFeatures/Source/TedsAssetData/Private/CB/TedsAssetDataCBDataSource.cpp
@@ -16,14 +16,14 @@
 #include "Settings/ContentBrowserSettings.h"
 #include "UObject/NameTypes.h"
 
-namespace UE::EditorDataStorage::AssetData::Private
+namespace UE::Editor::AssetData::Private
 {
 
 FTedsAssetDataCBDataSource::FTedsAssetDataCBDataSource(ITypedElementDataStorageInterface& InDatabase)
 	: Database(InDatabase)
 {
 	using namespace TypedElementQueryBuilder;
-	using namespace EditorDataStorage;
+	using namespace UE::Editor::DataStorage;
 
 	InitVirtualPathProcessor();
 
@@ -42,7 +42,7 @@ FTedsAssetDataCBDataSource::FTedsAssetDataCBDataSource(ITypedElementDataStorageI
 		Select(
 			TEXT("FTedsAssetDataCBDataSource: Process Path updates"),
 			FProcessor(TypedElementDataStorage::EQueryTickPhase::DuringPhysics, Database.GetQueryTickGroupName(TypedElementDataStorage::EQueryTickGroups::Update)),
-			[GenerateVirtualPaths](TypedElementDataStorage::IQueryContext& Context, const TypedElementDataStorage::RowHandle* Rows, const FAssetPathColumn_Experimental* PathColumn)
+			[GenerateVirtualPaths](TypedElementDataStorage::IQueryContext& Context, const RowHandle* Rows, const FAssetPathColumn_Experimental* PathColumn)
 			{
 				int32 NumOfRowToProcess = Context.GetRowCount();
 
@@ -72,7 +72,7 @@ FTedsAssetDataCBDataSource::FTedsAssetDataCBDataSource(ITypedElementDataStorageI
 		Select(
 			TEXT("FTedsAssetDataCBDataSource:: Process Asset Data Path Update"),
 			FProcessor(TypedElementDataStorage::EQueryTickPhase::DuringPhysics, Database.GetQueryTickGroupName(TypedElementDataStorage::EQueryTickGroups::Update)),
-			[GenerateVirtualPaths](TypedElementDataStorage::IQueryContext& Context, const TypedElementDataStorage::RowHandle Row, const FAssetDataColumn_Experimental& AssetDataColumn)
+			[GenerateVirtualPaths](TypedElementDataStorage::IQueryContext& Context, const RowHandle Row, const FAssetDataColumn_Experimental& AssetDataColumn)
 			{
 				int32 NumOfRowToProcess = Context.GetRowCount();
 
@@ -98,7 +98,7 @@ FTedsAssetDataCBDataSource::FTedsAssetDataCBDataSource(ITypedElementDataStorageI
 
 
 	// For now just add the columns one by one but this should be rework to work in batch
-	auto AddAssetDataColumns = [](TypedElementDataStorage::IQueryContext& InContext, TypedElementDataStorage::RowHandle Row, const FAssetData& InAssetData, const FAssetPackageData* PackageData)
+	auto AddAssetDataColumns = [](TypedElementDataStorage::IQueryContext& InContext, RowHandle Row, const FAssetData& InAssetData, const FAssetPackageData* PackageData)
 	{
 		// Not optimized at all but we would like to have the data in sooner for testing purposes.
 		
@@ -128,12 +128,12 @@ FTedsAssetDataCBDataSource::FTedsAssetDataCBDataSource(ITypedElementDataStorageI
 		Select(
 			TEXT("FTedsAssetDataCBDataSource: Process Asset Data and Path Updates"),
 			FProcessor(TypedElementDataStorage::EQueryTickPhase::DuringPhysics, Database.GetQueryTickGroupName(TypedElementDataStorage::EQueryTickGroups::Update)),
-			[GenerateVirtualPaths, AddAssetDataColumns, AssetRegistry = static_cast<const IAssetRegistry*>(&IAssetRegistry::GetChecked())](TypedElementDataStorage::IQueryContext& Context, const TypedElementDataStorage::RowHandle* Rows, const FAssetDataColumn_Experimental* AssetDataColumn)
+			[GenerateVirtualPaths, AddAssetDataColumns, AssetRegistry = static_cast<const IAssetRegistry*>(&IAssetRegistry::GetChecked())](TypedElementDataStorage::IQueryContext& Context, const RowHandle* Rows, const FAssetDataColumn_Experimental* AssetDataColumn)
 			{
 				const int32 RowCount = Context.GetRowCount();
 
 				TArray<FName, TInlineAllocator<32>> PackageNames;
-				TArray<TPair<TypedElementDataStorage::RowHandle, const FAssetData*>, TInlineAllocator<32>> RowsAndAssetData;
+				TArray<TPair<RowHandle, const FAssetData*>, TInlineAllocator<32>> RowsAndAssetData;
 				PackageNames.Reserve(RowCount);
 				RowsAndAssetData.Reserve(RowCount);
 
@@ -145,7 +145,7 @@ FTedsAssetDataCBDataSource::FTedsAssetDataCBDataSource(ITypedElementDataStorageI
 					if (GenerateVirtualPaths(InternalPath, VirtualPath))
 					{
 						const FAssetData& AssetData = AssetDataColumn[Index].AssetData;
-						const TypedElementDataStorage::RowHandle Row = Rows[Index];
+						const RowHandle Row = Rows[Index];
 
 						InternalPath.Reset();
 						AssetData.AppendObjectPath(InternalPath);
@@ -163,7 +163,7 @@ FTedsAssetDataCBDataSource::FTedsAssetDataCBDataSource(ITypedElementDataStorageI
 
 				for (int32 Index = 0; Index < AssetPackageDatas.Num(); ++Index)
 				{
-					const TPair<TypedElementDataStorage::RowHandle, const FAssetData*>& Pair = RowsAndAssetData[Index];
+					const TPair<RowHandle, const FAssetData*>& Pair = RowsAndAssetData[Index];
 					const TOptional<FAssetPackageData>& AssetPackageData = AssetPackageDatas[Index];
 					AddAssetDataColumns(Context, Pair.Key, *Pair.Value, AssetPackageData.GetPtrOrNull());
 				}
@@ -178,7 +178,7 @@ FTedsAssetDataCBDataSource::FTedsAssetDataCBDataSource(ITypedElementDataStorageI
 		Select(
 			TEXT("FTedsAssetDataCBDataSource: Process Asset Data updates"),
 			FProcessor(TypedElementDataStorage::EQueryTickPhase::DuringPhysics, Database.GetQueryTickGroupName(TypedElementDataStorage::EQueryTickGroups::Update)),
-			[AddAssetDataColumns,  AssetRegistry = static_cast<const IAssetRegistry*>(&IAssetRegistry::GetChecked())](TypedElementDataStorage::IQueryContext& Context, const TypedElementDataStorage::RowHandle* Rows, const FAssetDataColumn_Experimental* AssetDataColumn)
+			[AddAssetDataColumns,  AssetRegistry = static_cast<const IAssetRegistry*>(&IAssetRegistry::GetChecked())](TypedElementDataStorage::IQueryContext& Context, const RowHandle* Rows, const FAssetDataColumn_Experimental* AssetDataColumn)
 			{
 				const int32 RowCount = Context.GetRowCount();
 				TArray<FName, TInlineAllocator<32>> PackageNames;
@@ -314,4 +314,4 @@ void FTedsAssetDataCBDataSource::FVirtualPathProcessor::ConvertInternalPathToVir
 	OutVirtualPath.Append(InternalPath.GetData(), InternalPath.Len());
 }
 
-} // End of Namespace UE::TypedElementDataStorageAssetData::Private
+} // End of Namespace UE::Editor::AssetData::Private
diff --git a/Engine/Plugins/Experimental/EditorDataStorageFeatures/Source/TedsAssetData/Private/CB/TedsAssetDataCBDataSource.h b/Engine/Plugins/Experimental/EditorDataStorageFeatures/Source/TedsAssetData/Private/CB/TedsAssetDataCBDataSource.h
--- a/Engine/Plugins/Experimental/EditorDataStorageFeatures/Source/TedsAssetData/Private/CB/TedsAssetDataCBDataSource.h
+++ b/Engine/Plugins/Experimental/EditorDataStorageFeatures/Source/TedsAssetData/Private/CB/TedsAssetDataCBDataSource.h
@@ -11,44 +11,47 @@
 class IPlugin;
 class ITypedElementDataStorageInterface;
 
-namespace UE::EditorDataStorage::AssetData::Private
+namespace UE::Editor
 {
-
-class FTedsAssetDataCBDataSource
-{
-public:
-	explicit FTedsAssetDataCBDataSource(ITypedElementDataStorageInterface& InDatabase);
-	~FTedsAssetDataCBDataSource();
-		
-private:
-	ITypedElementDataStorageInterface& Database;
-	TypedElementDataStorage::QueryHandle ProcessPathQuery;
-	TypedElementDataStorage::QueryHandle ProcessAssetDataPathUpdateQuery;
-	TypedElementDataStorage::QueryHandle ProcessAssetDataAndPathUpdateQuery;
-	TypedElementDataStorage::QueryHandle ProcessAssetDataUpdateQuery;
-
-	struct FVirtualPathProcessor
+	namespace AssetData::Private
 	{
-		struct FCachedPluginData
+
+		class FTedsAssetDataCBDataSource
 		{
-			EPluginLoadedFrom LoadedFrom;
-			FString EditorCustomVirtualPath;
-		};
+		public:
+			explicit FTedsAssetDataCBDataSource(ITypedElementDataStorageInterface& InDatabase);
+			~FTedsAssetDataCBDataSource();
 
-		TMap<FString, FCachedPluginData> PluginNameToCachedData;
+		private:
+			ITypedElementDataStorageInterface& Database;
+			DataStorage::QueryHandle ProcessPathQuery;
+			DataStorage::QueryHandle ProcessAssetDataPathUpdateQuery;
+			DataStorage::QueryHandle ProcessAssetDataAndPathUpdateQuery;
+			DataStorage::QueryHandle ProcessAssetDataUpdateQuery;
 
-		bool bShowAllFolder = false;
-		bool bOrganizeFolders = false;
+			struct FVirtualPathProcessor
+			{
+				struct FCachedPluginData
+				{
+					EPluginLoadedFrom LoadedFrom;
+					FString EditorCustomVirtualPath;
+				};
 
-		void ConvertInternalPathToVirtualPath(const FStringView InternalPath, FStringBuilderBase& OutVirtualPath) const;
-	};
+				TMap<FString, FCachedPluginData> PluginNameToCachedData;
 
-	FVirtualPathProcessor VirtualPathProcessor; 
+				bool bShowAllFolder = false;
+				bool bOrganizeFolders = false;
 
-	void InitVirtualPathProcessor();
+				void ConvertInternalPathToVirtualPath(const FStringView InternalPath, FStringBuilderBase& OutVirtualPath) const;
+			};
 
-	void OnPluginContentMounted(IPlugin& InPlugin);
-	void OnPluginUnmounted(IPlugin& InPlugin);
-};
+			FVirtualPathProcessor VirtualPathProcessor;
+
+			void InitVirtualPathProcessor();
+
+			void OnPluginContentMounted(IPlugin& InPlugin);
+			void OnPluginUnmounted(IPlugin& InPlugin);
+		};
 
-} // End of namespace UE::TypedElementsDataStorageAssetData::Private
+	} // namespace AssetData::Private
+} // End of namespace UE::Editor
diff --git a/Engine/Plugins/Experimental/EditorDataStorageFeatures/Source/TedsAssetData/Private/TedsAssetData.cpp b/Engine/Plugins/Experimental/EditorDataStorageFeatures/Source/TedsAssetData/Private/TedsAssetData.cpp
--- a/Engine/Plugins/Experimental/EditorDataStorageFeatures/Source/TedsAssetData/Private/TedsAssetData.cpp
+++ b/Engine/Plugins/Experimental/EditorDataStorageFeatures/Source/TedsAssetData/Private/TedsAssetData.cpp
@@ -16,9 +16,10 @@
 
 #define TRACK_TEDSASSETDATA_MEMORY 0
 
-namespace UE::EditorDataStorage::AssetData::Private
+namespace UE::Editor
 {
-	
+	namespace AssetData::Private
+	{	
 constexpr int32 ParallelForMinBatchSize = 1024 * 4;
 
 struct FPopulateAssetDataRowArgs
@@ -27,7 +28,7 @@ struct FPopulateAssetDataRowArgs
 
 	TypedElementDataStorage::IndexHash ObjectPathHash;
 	TypedElementDataStorage::IndexHash PathHash;
-	TypedElementDataStorage::RowHandle PathRow = TypedElementDataStorage::InvalidRowHandle;
+	DataStorage::RowHandle PathRow = DataStorage::InvalidRowHandle;
 };
 
 // Only safe if the GT is blocked during the operation
@@ -54,9 +55,9 @@ FPopulateAssetDataRowArgs ThreadSafe_PopulateAssetDataTableRow(TAssetData&& InAs
 }
 
 
-void PopulateAssetDataTableRow(FPopulateAssetDataRowArgs&& InAssetDataRowArgs, ITypedElementDataStorageInterface& Database, TypedElementDataStorage::RowHandle RowHandle)
+void PopulateAssetDataTableRow(FPopulateAssetDataRowArgs&& InAssetDataRowArgs, ITypedElementDataStorageInterface& Database, DataStorage::RowHandle RowHandle)
 {
-	if (InAssetDataRowArgs.PathRow !=  TypedElementDataStorage::InvalidRowHandle)
+	if (InAssetDataRowArgs.PathRow !=  DataStorage::InvalidRowHandle)
 	{
 		Database.GetColumn<FAssetsInPathColumn_Experimental>(InAssetDataRowArgs.PathRow)->AssetsRow.Add(RowHandle);
 	}
@@ -114,7 +115,7 @@ void GetPathDepthAndParentFolderIndex(FStringView Path, uint32& OutDepth, int32&
 // Only thread safe if the game thread is blocked 
 FPopulatePathRowArgs ThreadSafe_PopulatePathRowArgs(TypedElementDataStorage::IndexHash AssetRegistryPathHash, FName InAssetRegistryPath, FStringView PathAsString)
 {
-	TypedElementDataStorage::IndexHash ParentAssetRegistryPathHash = TypedElementDataStorage::InvalidRowHandle;
+	TypedElementDataStorage::IndexHash ParentAssetRegistryPathHash = DataStorage::InvalidRowHandle;
 	int32 CharacterIndex;
 	uint32 Depth;
 	GetPathDepthAndParentFolderIndex(PathAsString, Depth, CharacterIndex);
@@ -133,26 +134,26 @@ FPopulatePathRowArgs ThreadSafe_PopulatePathRowArgs(TypedElementDataStorage::Ind
 	return Args;
 }
 
-void PopulatePathDataTableRow(FPopulatePathRowArgs&& InPopulatePathRowArgs, ITypedElementDataStorageInterface& Database, TypedElementDataStorage::RowHandle RowHandle)
+void PopulatePathDataTableRow(FPopulatePathRowArgs&& InPopulatePathRowArgs, ITypedElementDataStorageInterface& Database, DataStorage::RowHandle InRowHandle)
 {
-	if (InPopulatePathRowArgs.ParentAssetRegistryPathHash != TypedElementDataStorage::InvalidRowHandle)
+	if (InPopulatePathRowArgs.ParentAssetRegistryPathHash != DataStorage::InvalidRowHandle)
 	{
-		const TypedElementDataStorage::RowHandle ParentRow = Database.FindIndexedRow(InPopulatePathRowArgs.ParentAssetRegistryPathHash);
+		const DataStorage::RowHandle ParentRow = Database.FindIndexedRow(InPopulatePathRowArgs.ParentAssetRegistryPathHash);
 		if (Database.IsRowAssigned(ParentRow))
 		{
-			Database.GetColumn<FChildrenAssetPathColumn_Experimental>(ParentRow)->ChildrenRows.Add(RowHandle);
-			Database.GetColumn<FParentAssetPathColumn_Experimental>(RowHandle)->ParentRow = ParentRow;
+			Database.GetColumn<FChildrenAssetPathColumn_Experimental>(ParentRow)->ChildrenRows.Add(InRowHandle);
+			Database.GetColumn<FParentAssetPathColumn_Experimental>(InRowHandle)->ParentRow = ParentRow;
 		}
 		else 
 		{
 			// If we were unlucky we may be missing the parent path data.This should resolve it self during a latter event on paths added event.
 			FUnresolvedParentAssetPathColumn_Experimental UnresolvedParentRow;
 			UnresolvedParentRow.Hash = InPopulatePathRowArgs.ParentAssetRegistryPathHash;
-			Database.AddColumn(RowHandle, MoveTemp(UnresolvedParentRow));
+			Database.AddColumn(InRowHandle, MoveTemp(UnresolvedParentRow));
 		}
 	}
 
-	Database.GetColumn<FAssetPathColumn_Experimental>(RowHandle)->Path = InPopulatePathRowArgs.AssetRegistryPath;
+	Database.GetColumn<FAssetPathColumn_Experimental>(InRowHandle)->Path = InPopulatePathRowArgs.AssetRegistryPath;
 }
 
 
@@ -194,9 +195,9 @@ FTedsAssetData::FTedsAssetData(ITypedElementDataStorageInterface& InDatabase)
 			TypedElementQueryBuilder::Select(
 				TEXT("FTedsAssetData: Remove Updated Path Tag"),
 				TypedElementQueryBuilder::FPhaseAmble(TypedElementQueryBuilder::FPhaseAmble::ELocation::Postamble, TypedElementDataStorage::EQueryTickPhase::FrameEnd),
-				[](TypedElementDataStorage::IQueryContext& Context, const TypedElementDataStorage::RowHandle* Rows)
+				[](TypedElementDataStorage::IQueryContext& Context, const DataStorage::RowHandle* Rows)
 				{
-					Context.RemoveColumns<FUpdatedPathTag>(TConstArrayView<TypedElementRowHandle>(Rows, Context.GetRowCount()));
+					Context.RemoveColumns<FUpdatedPathTag>(TConstArrayView<DataStorage::RowHandle>(Rows, Context.GetRowCount()));
 				}
 			)
 			.Where()
@@ -208,9 +209,9 @@ FTedsAssetData::FTedsAssetData(ITypedElementDataStorageInterface& InDatabase)
 			TypedElementQueryBuilder::Select(
 				TEXT("FTedsAssetData: Remove Updated Asset Data Tag"),
 				TypedElementQueryBuilder::FPhaseAmble(TypedElementQueryBuilder::FPhaseAmble::ELocation::Postamble, TypedElementDataStorage::EQueryTickPhase::FrameEnd),
-				[](TypedElementDataStorage::IQueryContext& Context, const TypedElementDataStorage::RowHandle* Rows)
+				[](TypedElementDataStorage::IQueryContext& Context, const DataStorage::RowHandle* Rows)
 				{
-					Context.RemoveColumns<FUpdatedAssetDataTag>(TConstArrayView<TypedElementRowHandle>(Rows, Context.GetRowCount()));
+					Context.RemoveColumns<FUpdatedAssetDataTag>(TConstArrayView<DataStorage::RowHandle>(Rows, Context.GetRowCount()));
 				}
 			)
 			.Where()
@@ -228,13 +229,13 @@ FTedsAssetData::FTedsAssetData(ITypedElementDataStorageInterface& InDatabase)
 		TypedElementQueryBuilder::Select(
 			TEXT("FTedsAssetData: Resolve Missing Asset In Path"),
 			TypedElementQueryBuilder::FProcessor(TypedElementDataStorage::EQueryTickPhase::FrameEnd, Database.GetQueryTickGroupName(TypedElementDataStorage::EQueryTickGroups::Default)),
-			[this](TypedElementDataStorage::IQueryContext& Context, TypedElementDataStorage::RowHandle Row, const FUnresolvedAssetsInPathColumn_Experimental& UnresolvedAssetPath)
+			[this](TypedElementDataStorage::IQueryContext& Context, DataStorage::RowHandle Row, const FUnresolvedAssetsInPathColumn_Experimental& UnresolvedAssetPath)
 			{
 #if TRACK_TEDSASSETDATA_MEMORY
 				LLM_SCOPE_BYNAME(TEXT("FTedsAssetData"))
 #endif
 
-				const TypedElementDataStorage::RowHandle PathRow = Context.FindIndexedRow(UnresolvedAssetPath.Hash);
+				const DataStorage::RowHandle PathRow = Context.FindIndexedRow(UnresolvedAssetPath.Hash);
 				if (Context.IsRowAssigned(PathRow))
 				{
 					Context.RemoveColumns<FUnresolvedAssetsInPathColumn_Experimental>(Row);
@@ -266,13 +267,13 @@ FTedsAssetData::FTedsAssetData(ITypedElementDataStorageInterface& InDatabase)
 		TypedElementQueryBuilder::Select(
 			TEXT("FTedsAssetData: Resolve Missing Parent Path Row"),
 			TypedElementQueryBuilder::FProcessor(TypedElementDataStorage::EQueryTickPhase::FrameEnd, Database.GetQueryTickGroupName(TypedElementDataStorage::EQueryTickGroups::Default)),
-			[this](TypedElementDataStorage::IQueryContext& Context, TypedElementDataStorage::RowHandle Row, const FUnresolvedParentAssetPathColumn_Experimental& UnresolvedParentAssetPath, FParentAssetPathColumn_Experimental& ParentAssetPathColumn)
+			[this](TypedElementDataStorage::IQueryContext& Context, DataStorage::RowHandle Row, const FUnresolvedParentAssetPathColumn_Experimental& UnresolvedParentAssetPath, FParentAssetPathColumn_Experimental& ParentAssetPathColumn)
 			{
 #if TRACK_TEDSASSETDATA_MEMORY
 				LLM_SCOPE_BYNAME(TEXT("FTedsAssetData"))
 #endif
 
-				const TypedElementDataStorage::RowHandle ParentPathRow = Context.FindIndexedRow(UnresolvedParentAssetPath.Hash);
+				const DataStorage::RowHandle ParentPathRow = Context.FindIndexedRow(UnresolvedParentAssetPath.Hash);
 				if (Context.IsRowAssigned(ParentPathRow))
 				{
 					Context.RemoveColumns<FUnresolvedParentAssetPathColumn_Experimental>(Row);
@@ -319,24 +320,24 @@ FTedsAssetData::FTedsAssetData(ITypedElementDataStorageInterface& InDatabase)
 			PopulatePathRowArgs[Index] = ThreadSafe_PopulatePathRowArgs(TypedElementDataStorage::GenerateIndexHash(Path), Path, NameBuilder);
 		});
 
-	TArray<TypedElementDataStorage::RowHandle> ReservedRows;
+	TArray<DataStorage::RowHandle> ReservedRows;
 	ReservedRows.AddUninitialized(PopulatePathRowArgs.Num() + AssetsData.Num());
 	Database.BatchReserveRows(ReservedRows);
 
-	TArray<TPair<TypedElementDataStorage::IndexHash, TypedElementDataStorage::RowHandle>> IndexesToReservedRows;
+	TArray<TPair<TypedElementDataStorage::IndexHash, DataStorage::RowHandle>> IndexesToReservedRows;
 	IndexesToReservedRows.AddUninitialized(ReservedRows.Num());
 
 	// Index Reserved Path Rows
-	TConstArrayView<TypedElementDataStorage::RowHandle> ReservedPopulatePathRows(ReservedRows.GetData(), PopulatePathRowArgs.Num());
-	TArrayView<TPair<TypedElementDataStorage::IndexHash, TypedElementDataStorage::RowHandle>> IndexesToReservedPathRows(IndexesToReservedRows.GetData(), PopulatePathRowArgs.Num());
+	TConstArrayView<DataStorage::RowHandle> ReservedPopulatePathRows(ReservedRows.GetData(), PopulatePathRowArgs.Num());
+	TArrayView<TPair<TypedElementDataStorage::IndexHash, DataStorage::RowHandle>> IndexesToReservedPathRows(IndexesToReservedRows.GetData(), PopulatePathRowArgs.Num());
 	ParallelFor(TEXT("Populating TEDS Asset Registry Path Data Indexes"), ReservedPopulatePathRows.Num(), ParallelForMinBatchSize, [&IndexesToReservedPathRows, &ReservedPopulatePathRows, &PopulatePathRowArgs](int32 Index)
 		{
-			IndexesToReservedPathRows[Index] = TPair<TypedElementDataStorage::IndexHash, TypedElementDataStorage::RowHandle>(PopulatePathRowArgs[Index].AssetRegistryPathHash, ReservedPopulatePathRows[Index]);
+			IndexesToReservedPathRows[Index] = TPair<TypedElementDataStorage::IndexHash, DataStorage::RowHandle>(PopulatePathRowArgs[Index].AssetRegistryPathHash, ReservedPopulatePathRows[Index]);
 		});
 	Database.BatchIndexRows(IndexesToReservedPathRows);
 
 	// Populate Path Rows
-	Database.BatchAddRow(PathsTable, ReservedPopulatePathRows, [PathRowArgs = MoveTemp(PopulatePathRowArgs), Index = 0, this](TypedElementDataStorage::RowHandle InRowHandle) mutable
+	Database.BatchAddRow(PathsTable, ReservedPopulatePathRows, [PathRowArgs = MoveTemp(PopulatePathRowArgs), Index = 0, this](DataStorage::RowHandle InRowHandle) mutable
 		{
 			PopulatePathDataTableRow(MoveTemp(PathRowArgs[Index]), Database, InRowHandle);
 			++Index;
@@ -351,16 +352,16 @@ FTedsAssetData::FTedsAssetData(ITypedElementDataStorageInterface& InDatabase)
 		});
 
 	// Index Reserved Asset Data Rows
-	TConstArrayView<TypedElementDataStorage::RowHandle> ReservedAssetDataRows(ReservedRows.GetData() + ReservedPopulatePathRows.Num(), AssetsData.Num());
-	TArrayView<TPair<TypedElementDataStorage::IndexHash, TypedElementDataStorage::RowHandle>>  IndexesToResevedAssetRows(IndexesToReservedRows.GetData() + ReservedPopulatePathRows.Num(), ReservedAssetDataRows.Num());
+	TConstArrayView<DataStorage::RowHandle> ReservedAssetDataRows(ReservedRows.GetData() + ReservedPopulatePathRows.Num(), AssetsData.Num());
+	TArrayView<TPair<TypedElementDataStorage::IndexHash, DataStorage::RowHandle>>  IndexesToResevedAssetRows(IndexesToReservedRows.GetData() + ReservedPopulatePathRows.Num(), ReservedAssetDataRows.Num());
 	ParallelFor(TEXT("Populating TEDS Asset Registry Asset Data Indexes"), PopulateAssetDataRowArgs.Num(), ParallelForMinBatchSize, [&PopulateAssetDataRowArgs, &ReservedAssetDataRows, &IndexesToResevedAssetRows](int32 Index)
 		{
-			IndexesToResevedAssetRows[Index] = TPair<TypedElementDataStorage::IndexHash, TypedElementDataStorage::RowHandle>(PopulateAssetDataRowArgs[Index].ObjectPathHash, ReservedAssetDataRows[Index]);
+			IndexesToResevedAssetRows[Index] = TPair<TypedElementDataStorage::IndexHash, DataStorage::RowHandle>(PopulateAssetDataRowArgs[Index].ObjectPathHash, ReservedAssetDataRows[Index]);
 		});
 	Database.BatchIndexRows(IndexesToResevedAssetRows);
 
 	// Populate Asset Rows
-	Database.BatchAddRow(AssetsDataTable, ReservedAssetDataRows,  [AssetDataRowArgs = MoveTemp(PopulateAssetDataRowArgs), Index = 0, this](TypedElementDataStorage::RowHandle InRowHandle) mutable
+	Database.BatchAddRow(AssetsDataTable, ReservedAssetDataRows,  [AssetDataRowArgs = MoveTemp(PopulateAssetDataRowArgs), Index = 0, this](DataStorage::RowHandle InRowHandle) mutable
 		{
 			PopulateAssetDataTableRow(MoveTemp(AssetDataRowArgs[Index]), Database, InRowHandle);
 			++Index;
@@ -399,7 +400,7 @@ FTedsAssetData::~FTedsAssetData()
 			AssetRegistry->EnumerateAllCachedPaths([this](FName InPath)
 				{
 					const TypedElementDataStorage::IndexHash PathHash = TypedElementDataStorage::GenerateIndexHash(InPath);
-					const TypedElementDataStorage::RowHandle Row = Database.FindIndexedRow(PathHash);
+					const DataStorage::RowHandle Row = Database.FindIndexedRow(PathHash);
 					Database.RemoveRow(Row);
 					Database.RemoveIndex(PathHash);
 					return true;
@@ -408,7 +409,7 @@ FTedsAssetData::~FTedsAssetData()
 			AssetRegistry->EnumerateAllAssets([this](const FAssetData& InAssetData)
 				{
 					const TypedElementDataStorage::IndexHash AssetPathHash = TypedElementDataStorage::GenerateIndexHash(InAssetData.GetSoftObjectPath());
-					const TypedElementDataStorage::RowHandle Row = Database.FindIndexedRow(AssetPathHash);
+					const DataStorage::RowHandle Row = Database.FindIndexedRow(AssetPathHash);
 					Database.RemoveRow(Row);
 					Database.RemoveIndex(AssetPathHash);
 					return true;
@@ -465,9 +466,9 @@ void FTedsAssetData::OnAssetsAdded(TConstArrayView<FAssetData> InAssetsAdded)
 	int32 Index = 0;
 	if (NewRowsCount > 0)
 	{
-		TArray<TPair<TypedElementDataStorage::IndexHash, TypedElementDataStorage::RowHandle>> IndexToRow;
+		TArray<TPair<TypedElementDataStorage::IndexHash, DataStorage::RowHandle>> IndexToRow;
 		IndexToRow.Reserve(NewRowsCount);
-		Database.BatchAddRow(AssetsDataTable, NewRowsCount, [RowArgs = MoveTemp(PopulateRowArgs), Index = 0, &IndexToRow, this](TypedElementDataStorage::RowHandle InRowHandle) mutable
+		Database.BatchAddRow(AssetsDataTable, NewRowsCount, [RowArgs = MoveTemp(PopulateRowArgs), Index = 0, &IndexToRow, this](DataStorage::RowHandle InRowHandle) mutable
 			{
 				FPopulateAssetDataRowArgs ARowArgs = MoveTemp(RowArgs[Index]);
 				while (!ARowArgs.AssetData.IsValid())
@@ -496,13 +497,13 @@ void FTedsAssetData::OnAssetsRemoved(TConstArrayView<FAssetData> InAssetsRemoved
 	for (const FAssetData& Asset : InAssetsRemoved)
 	{
 		const TypedElementDataStorage::IndexHash AssetHash = TypedElementDataStorage::GenerateIndexHash(Asset.GetSoftObjectPath());
-		const TypedElementDataStorage::RowHandle AssetRow = Database.FindIndexedRow(AssetHash);
+		const DataStorage::RowHandle AssetRow = Database.FindIndexedRow(AssetHash);
 		if (Database.IsRowAssigned(AssetRow))
 		{
 			if (const FAssetDataColumn_Experimental* AssetDataColunm = Database.GetColumn<FAssetDataColumn_Experimental>(AssetRow))
 			{
 				const TypedElementDataStorage::IndexHash FolderPathHash = TypedElementDataStorage::GenerateIndexHash(AssetDataColunm->AssetData.PackagePath);
-				const TypedElementDataStorage::RowHandle FolderRow = Database.FindIndexedRow(FolderPathHash);
+				const DataStorage::RowHandle FolderRow = Database.FindIndexedRow(FolderPathHash);
 				if (FAssetsInPathColumn_Experimental* AssetInFolder = Database.GetColumn<FAssetsInPathColumn_Experimental>(FolderRow))
 				{
 					AssetInFolder->AssetsRow.Remove(AssetRow);
@@ -525,7 +526,7 @@ void FTedsAssetData::OnAssetsUpdated(TConstArrayView<FAssetData> InAssetsUpdated
 	for (const FAssetData& Asset : InAssetsUpdated)
 	{
 		const TypedElementDataStorage::IndexHash AssetHash = TypedElementDataStorage::GenerateIndexHash(Asset.GetSoftObjectPath());
-		const TypedElementDataStorage::RowHandle Row = Database.FindIndexedRow(AssetHash);
+		const DataStorage::RowHandle Row = Database.FindIndexedRow(AssetHash);
 		if (Database.IsRowAssigned(Row))
 		{
 			Database.GetColumn<FAssetDataColumn_Experimental>(Row)->AssetData = Asset;
@@ -545,7 +546,7 @@ void FTedsAssetData::OnAssetsUpdatedOnDisk(TConstArrayView<FAssetData> InAssetsU
 	for (const FAssetData& Asset : InAssetsUpdated)
 	{
 		const TypedElementDataStorage::IndexHash AssetHash = TypedElementDataStorage::GenerateIndexHash(Asset.GetSoftObjectPath());
-		const TypedElementDataStorage::RowHandle Row = Database.FindIndexedRow(AssetHash);
+		const DataStorage::RowHandle Row = Database.FindIndexedRow(AssetHash);
 		if (Database.IsRowAssigned(Row))
 		{
 			Database.GetColumn<FAssetDataColumn_Experimental>(Row)->AssetData = Asset;
@@ -564,7 +565,7 @@ void FTedsAssetData::OnAssetRenamed(const FAssetData& InAsset, const FString& In
 
 	const TypedElementDataStorage::IndexHash NewAssetHash = TypedElementDataStorage::GenerateIndexHash(InAsset.GetSoftObjectPath());
 	const TypedElementDataStorage::IndexHash OldAssetHash = TypedElementDataStorage::GenerateIndexHash(FSoftObjectPath(InOldObjectPath));
-	const TypedElementDataStorage::RowHandle Row = Database.FindIndexedRow(OldAssetHash);
+	const DataStorage::RowHandle Row = Database.FindIndexedRow(OldAssetHash);
 	if (Database.IsRowAssigned(Row))
 	{
 		Database.GetColumn<FAssetDataColumn_Experimental>(Row)->AssetData = InAsset;
@@ -583,8 +584,8 @@ void FTedsAssetData::OnAssetRenamed(const FAssetData& InAsset, const FString& In
 
 		if (NewFolderHash != OldFolderHash)
 		{
-			const TypedElementDataStorage::RowHandle NewPathRow = Database.FindIndexedRow(NewFolderHash);
-			if (NewPathRow !=  TypedElementDataStorage::InvalidRowHandle)
+			const DataStorage::RowHandle NewPathRow = Database.FindIndexedRow(NewFolderHash);
+			if (NewPathRow !=  DataStorage::InvalidRowHandle)
 			{
 				Database.GetColumn<FAssetsInPathColumn_Experimental>(NewPathRow)->AssetsRow.Add(Row);
 			}
@@ -595,8 +596,8 @@ void FTedsAssetData::OnAssetRenamed(const FAssetData& InAsset, const FString& In
 				Database.AddColumn<FUnresolvedAssetsInPathColumn_Experimental>(Row, MoveTemp(UnresolvedAssetsInPathColumn));
 			}
 
-			const TypedElementDataStorage::RowHandle OldPathRow = Database.FindIndexedRow(OldFolderHash);
-			if (NewPathRow !=  TypedElementDataStorage::InvalidRowHandle)
+			const DataStorage::RowHandle OldPathRow = Database.FindIndexedRow(OldFolderHash);
+			if (NewPathRow !=  DataStorage::InvalidRowHandle)
 			{
 				Database.GetColumn<FAssetsInPathColumn_Experimental>(NewPathRow)->AssetsRow.Remove(Row);
 			}
@@ -630,7 +631,7 @@ void FTedsAssetData::OnPathsAdded(TConstArrayView<FStringView> InPathsAdded)
 			const TypedElementDataStorage::IndexHash AssetRegistryPathHash = TypedElementDataStorage::GenerateIndexHash(PathName);
 			FPopulatePathRowArgs RowArgs;
 
-			if (Database.FindIndexedRow(AssetRegistryPathHash) != TypedElementDataStorage::InvalidRowHandle)
+			if (Database.FindIndexedRow(AssetRegistryPathHash) != DataStorage::InvalidRowHandle)
 			{
 				RowArgs.MarkAsInvalid();
 			}
@@ -652,11 +653,11 @@ void FTedsAssetData::OnPathsAdded(TConstArrayView<FStringView> InPathsAdded)
 
 	if (NewRowsCount > 0)
 	{
-		TArray<TypedElementDataStorage::RowHandle> ReservedRow;
+		TArray<DataStorage::RowHandle> ReservedRow;
 		ReservedRow.AddUninitialized(NewRowsCount);
 		Database.BatchReserveRows(ReservedRow);
 
-		TArray<TPair<TypedElementDataStorage::IndexHash, TypedElementDataStorage::RowHandle>> IndexesAndRows;
+		TArray<TPair<TypedElementDataStorage::IndexHash, DataStorage::RowHandle>> IndexesAndRows;
 		IndexesAndRows.Reserve(NewRowsCount);
 
 		int32 RowCount = 0;
@@ -676,7 +677,7 @@ void FTedsAssetData::OnPathsAdded(TConstArrayView<FStringView> InPathsAdded)
 		Database.BatchIndexRows(IndexesAndRows);
 
 		int32 Index = 0;
-		Database.BatchAddRow(PathsTable, ReservedRow, [&PopulateRowArgs, &Index, this](TypedElementDataStorage::RowHandle InRowHandle)
+		Database.BatchAddRow(PathsTable, ReservedRow, [&PopulateRowArgs, &Index, this](DataStorage::RowHandle InRowHandle)
 			{
 				FPopulatePathRowArgs RowArgs = MoveTemp(PopulateRowArgs[Index]);
 				while (!RowArgs)
@@ -706,6 +707,7 @@ void FTedsAssetData::OnPathsRemoved(TConstArrayView<FStringView> InPathsRemoved)
 		Database.RemoveIndex(PathHash);
 	}
 }
-}
+} // namespace AssetData::Private
+} // namespace UE::Editor
 
 #undef TRACK_TEDSASSETDATA_MEMORY
\ No newline at end of file
diff --git a/Engine/Plugins/Experimental/EditorDataStorageFeatures/Source/TedsAssetData/Private/TedsAssetData.h b/Engine/Plugins/Experimental/EditorDataStorageFeatures/Source/TedsAssetData/Private/TedsAssetData.h
--- a/Engine/Plugins/Experimental/EditorDataStorageFeatures/Source/TedsAssetData/Private/TedsAssetData.h
+++ b/Engine/Plugins/Experimental/EditorDataStorageFeatures/Source/TedsAssetData/Private/TedsAssetData.h
@@ -10,43 +10,44 @@ struct FAssetData;
 
 class ITypedElementDataStorageInterface;
 
-namespace UE::EditorDataStorage::AssetData::Private
+namespace UE::Editor
 {
-
-/**
- * Manage the registration and life cycle of the row related representing the data from the asset registry into TEDS.
- */
-class FTedsAssetData
-{
-public:
-	FTedsAssetData(ITypedElementDataStorageInterface& InDatabase);
-
-	~FTedsAssetData();
-
-	void ProcessAllEvents();
-
-private:
-	void OnAssetsAdded(TConstArrayView<FAssetData> InAssetsAdded);
-	void OnAssetsRemoved(TConstArrayView<FAssetData> InAssetsRemoved);
-	void OnAssetsUpdated(TConstArrayView<FAssetData> InAssetsUpdated);
-	void OnAssetsUpdatedOnDisk(TConstArrayView<FAssetData> InAssetsUpdated);
-
-	void OnAssetRenamed(const FAssetData& InAsset, const FString& InOldObjectPath);
-
-	void OnPathsAdded(TConstArrayView<FStringView> InPathsAdded);
-	void OnPathsRemoved(TConstArrayView<FStringView> InPathsRemoved);
-
-	ITypedElementDataStorageInterface& Database;
-	TypedElementDataStorage::TableHandle PathsTable = TypedElementDataStorage::InvalidTableHandle;
-	TypedElementDataStorage::TableHandle AssetsDataTable = TypedElementDataStorage::InvalidTableHandle;
-
-	TypedElementDataStorage::QueryHandle UpdateAssetsInPathQuery = TypedElementDataStorage::InvalidQueryHandle;
-	TypedElementDataStorage::QueryHandle UpdateAssetsInPath = TypedElementDataStorage::InvalidQueryHandle;
-	TypedElementDataStorage::QueryHandle ResolveMissingAssetInPathQuery = TypedElementDataStorage::InvalidQueryHandle;
-	TypedElementDataStorage::QueryHandle UpdateParentToChildrenAssetPathQuery = TypedElementDataStorage::InvalidQueryHandle;
-	TypedElementDataStorage::QueryHandle ResolveMissingParentPathQuery = TypedElementDataStorage::InvalidQueryHandle;
-	TypedElementDataStorage::QueryHandle RemoveUpdatedPathTagQuery = TypedElementDataStorage::InvalidQueryHandle;
-	TypedElementDataStorage::QueryHandle RemoveUpdatedAssetDataTagQuery = TypedElementDataStorage::InvalidQueryHandle;
-};
-
-}
+	namespace AssetData::Private
+	{
+		/**
+		 * Manage the registration and life cycle of the row related representing the data from the asset registry into TEDS.
+		 */
+		class FTedsAssetData
+		{
+		public:
+			FTedsAssetData(ITypedElementDataStorageInterface& InDatabase);
+
+			~FTedsAssetData();
+
+			void ProcessAllEvents();
+
+		private:
+			void OnAssetsAdded(TConstArrayView<FAssetData> InAssetsAdded);
+			void OnAssetsRemoved(TConstArrayView<FAssetData> InAssetsRemoved);
+			void OnAssetsUpdated(TConstArrayView<FAssetData> InAssetsUpdated);
+			void OnAssetsUpdatedOnDisk(TConstArrayView<FAssetData> InAssetsUpdated);
+
+			void OnAssetRenamed(const FAssetData& InAsset, const FString& InOldObjectPath);
+
+			void OnPathsAdded(TConstArrayView<FStringView> InPathsAdded);
+			void OnPathsRemoved(TConstArrayView<FStringView> InPathsRemoved);
+
+			ITypedElementDataStorageInterface& Database;
+			DataStorage::TableHandle PathsTable = DataStorage::InvalidTableHandle;
+			DataStorage::TableHandle AssetsDataTable = DataStorage::InvalidTableHandle;
+
+			DataStorage::QueryHandle UpdateAssetsInPathQuery = DataStorage::InvalidQueryHandle;
+			DataStorage::QueryHandle UpdateAssetsInPath = DataStorage::InvalidQueryHandle;
+			DataStorage::QueryHandle ResolveMissingAssetInPathQuery = DataStorage::InvalidQueryHandle;
+			DataStorage::QueryHandle UpdateParentToChildrenAssetPathQuery = DataStorage::InvalidQueryHandle;
+			DataStorage::QueryHandle ResolveMissingParentPathQuery = DataStorage::InvalidQueryHandle;
+			DataStorage::QueryHandle RemoveUpdatedPathTagQuery = DataStorage::InvalidQueryHandle;
+			DataStorage::QueryHandle RemoveUpdatedAssetDataTagQuery = DataStorage::InvalidQueryHandle;
+		};
+	} // namespace AssetData::Private
+} // namespace UE::Editor
diff --git a/Engine/Plugins/Experimental/EditorDataStorageFeatures/Source/TedsAssetData/Private/TedsAssetDataModule.cpp b/Engine/Plugins/Experimental/EditorDataStorageFeatures/Source/TedsAssetData/Private/TedsAssetDataModule.cpp
--- a/Engine/Plugins/Experimental/EditorDataStorageFeatures/Source/TedsAssetData/Private/TedsAssetDataModule.cpp
+++ b/Engine/Plugins/Experimental/EditorDataStorageFeatures/Source/TedsAssetData/Private/TedsAssetDataModule.cpp
@@ -8,9 +8,9 @@
 #include "Modules/ModuleManager.h"
 #include "TedsAssetData.h"
 
-IMPLEMENT_MODULE(UE::EditorDataStorage::AssetData::FTedsAssetDataModule, TedsAssetData);
+IMPLEMENT_MODULE(UE::Editor::AssetData::FTedsAssetDataModule, TedsAssetData);
 
-namespace UE::EditorDataStorage::AssetData
+namespace UE::Editor::AssetData
 {
 
 namespace Private
@@ -30,7 +30,7 @@ TAutoConsoleVariable<bool> CVarTEDSAssetDataStorage(TEXT("TEDS.AssetDataStorage"
 			Module.DisableTedsAssetRegistryStorage();
 		}
 	}));
-}
+} // namespace Private
 
 void FTedsAssetDataModule::StartupModule()
 {
@@ -114,4 +114,4 @@ void FTedsAssetDataModule::InitAssetRegistryStorage()
 	AssetRegistryStorage = MakeUnique<Private::FTedsAssetData>(MutableDataStorage);
 }
 
-}
\ No newline at end of file
+} // namespace UE::Editor::AssetData
\ No newline at end of file
diff --git a/Engine/Plugins/Experimental/EditorDataStorageFeatures/Source/TedsAssetData/Private/TedsAssetDataModule.h b/Engine/Plugins/Experimental/EditorDataStorageFeatures/Source/TedsAssetData/Private/TedsAssetDataModule.h
--- a/Engine/Plugins/Experimental/EditorDataStorageFeatures/Source/TedsAssetData/Private/TedsAssetDataModule.h
+++ b/Engine/Plugins/Experimental/EditorDataStorageFeatures/Source/TedsAssetData/Private/TedsAssetDataModule.h
@@ -5,14 +5,14 @@
 #include "Modules/ModuleInterface.h"
 #include "Templates/UniquePtr.h"
 
-namespace UE::EditorDataStorage::AssetData
+namespace UE::Editor::AssetData
 { 
 
 namespace Private
 {
 	class FTedsAssetData;
 	class FTedsAssetDataCBDataSource;
-}
+} // namespace Private
 
 class FTedsAssetDataModule : public IModuleInterface
 {
@@ -41,4 +41,4 @@ private:
 	TUniquePtr<Private::FTedsAssetData> AssetRegistryStorage;
 };
 
-}
+} // namespace UE::Editor::AssetData
diff --git a/Engine/Plugins/Experimental/EditorDataStorageFeatures/Source/TedsAssetData/Private/Tests/TedsAssetDataTest.cpp b/Engine/Plugins/Experimental/EditorDataStorageFeatures/Source/TedsAssetData/Private/Tests/TedsAssetDataTest.cpp
--- a/Engine/Plugins/Experimental/EditorDataStorageFeatures/Source/TedsAssetData/Private/Tests/TedsAssetDataTest.cpp
+++ b/Engine/Plugins/Experimental/EditorDataStorageFeatures/Source/TedsAssetData/Private/Tests/TedsAssetDataTest.cpp
@@ -20,7 +20,7 @@ IMPLEMENT_SIMPLE_AUTOMATION_TEST(FTedsAssetDataTest, "TEDS.AssetRegistry.Validat
  
 bool FTedsAssetDataTest::RunTest(const FString& Parameters)
 {
-	UE::EditorDataStorage::AssetData::FTedsAssetDataModule& TedsAssetRegistry = UE::EditorDataStorage::AssetData::FTedsAssetDataModule::GetChecked();
+	UE::Editor::AssetData::FTedsAssetDataModule& TedsAssetRegistry = UE::Editor::AssetData::FTedsAssetDataModule::GetChecked();
 
 	const bool bIsEnabled = TedsAssetRegistry.IsTedsAssetRegistryStorageEnabled();
 	ON_SCOPE_EXIT
diff --git a/Engine/Plugins/Experimental/EditorDataStorageFeatures/Source/TedsDebugger/Private/QueryEditor/TedsQueryEditor.cpp b/Engine/Plugins/Experimental/EditorDataStorageFeatures/Source/TedsDebugger/Private/QueryEditor/TedsQueryEditor.cpp
--- a/Engine/Plugins/Experimental/EditorDataStorageFeatures/Source/TedsDebugger/Private/QueryEditor/TedsQueryEditor.cpp
+++ b/Engine/Plugins/Experimental/EditorDataStorageFeatures/Source/TedsDebugger/Private/QueryEditor/TedsQueryEditor.cpp
@@ -12,7 +12,7 @@
 
 #define LOCTEXT_NAMESPACE "TedsQueryEditor"
 
-namespace UE::EditorDataStorage::Debug::QueryEditor
+namespace UE::Editor::DataStorage::Debug::QueryEditor
 {
 	struct SQueryEditorWidget::ColumnComboItem
 	{
@@ -25,9 +25,7 @@ namespace UE::EditorDataStorage::Debug::QueryEditor
 	};
 
 	void SQueryEditorWidget::Construct(const FArguments& InArgs, FTedsQueryEditorModel& QueryEditorModel)
-	{
-		using namespace UE::EditorDataStorage::Debug;
-	
+	{	
 		ComboItems.Reset();
 		Model = &QueryEditorModel;
 	
@@ -43,12 +41,12 @@ namespace UE::EditorDataStorage::Debug::QueryEditor
 					SNew(SHorizontalBox)
 					+SHorizontalBox::Slot()
 					[
-						SNew(SConditionCollectionViewWidget, *Model, QueryEditor::EOperatorType::Select)
+						SNew(SConditionCollectionViewWidget, *Model, EOperatorType::Select)
 					]
 					+SHorizontalBox::Slot()
 					.AutoWidth()
 					[
-						SNew(SConditionComboWidget, *Model, QueryEditor::EOperatorType::Select)
+						SNew(SConditionComboWidget, *Model, EOperatorType::Select)
 					]
 				]
 				+SVerticalBox::Slot()
@@ -57,12 +55,12 @@ namespace UE::EditorDataStorage::Debug::QueryEditor
 					SNew(SHorizontalBox)
 					+SHorizontalBox::Slot()
 					[
-						SNew(SConditionCollectionViewWidget, *Model, QueryEditor::EOperatorType::All)
+						SNew(SConditionCollectionViewWidget, *Model, EOperatorType::All)
 					]
 					+SHorizontalBox::Slot()
 					.AutoWidth()
 					[
-						SNew(SConditionComboWidget, *Model, QueryEditor::EOperatorType::All)
+						SNew(SConditionComboWidget, *Model, EOperatorType::All)
 					]
 				]
 				+SVerticalBox::Slot()
@@ -71,12 +69,12 @@ namespace UE::EditorDataStorage::Debug::QueryEditor
 					SNew(SHorizontalBox)
 					+SHorizontalBox::Slot()
 					[
-					SNew(SConditionCollectionViewWidget, *Model, QueryEditor::EOperatorType::Any)
+					SNew(SConditionCollectionViewWidget, *Model, EOperatorType::Any)
 					]
 					+SHorizontalBox::Slot()
 					.AutoWidth()
 					[
-					SNew(SConditionComboWidget, *Model, QueryEditor::EOperatorType::Any)
+					SNew(SConditionComboWidget, *Model, EOperatorType::Any)
 					]
 				]
 				+SVerticalBox::Slot()
@@ -85,12 +83,12 @@ namespace UE::EditorDataStorage::Debug::QueryEditor
 					SNew(SHorizontalBox)
 					+SHorizontalBox::Slot()
 					[
-					SNew(SConditionCollectionViewWidget, *Model, QueryEditor::EOperatorType::None)
+					SNew(SConditionCollectionViewWidget, *Model, EOperatorType::None)
 					]
 					+SHorizontalBox::Slot()
 					.AutoWidth()
 					[
-					SNew(SConditionComboWidget, *Model, QueryEditor::EOperatorType::None)
+					SNew(SConditionComboWidget, *Model, EOperatorType::None)
 					]
 				]
 				+SVerticalBox::Slot()
@@ -101,6 +99,6 @@ namespace UE::EditorDataStorage::Debug::QueryEditor
 			
 		];
 	}
-}
+} // namespace UE::Editor::DataStorage::Debug::QueryEditor
 
 #undef LOCTEXT_NAMESPACE
\ No newline at end of file
diff --git a/Engine/Plugins/Experimental/EditorDataStorageFeatures/Source/TedsDebugger/Private/QueryEditor/TedsQueryEditor.h b/Engine/Plugins/Experimental/EditorDataStorageFeatures/Source/TedsDebugger/Private/QueryEditor/TedsQueryEditor.h
--- a/Engine/Plugins/Experimental/EditorDataStorageFeatures/Source/TedsDebugger/Private/QueryEditor/TedsQueryEditor.h
+++ b/Engine/Plugins/Experimental/EditorDataStorageFeatures/Source/TedsDebugger/Private/QueryEditor/TedsQueryEditor.h
@@ -6,7 +6,7 @@
 class SDockTab;
 class SWindow;
 
-namespace UE::EditorDataStorage::Debug::QueryEditor
+namespace UE::Editor::DataStorage::Debug::QueryEditor
 {
 	class FTedsQueryEditorModel;
 	class SQueryEditorWidget : public SCompoundWidget
diff --git a/Engine/Plugins/Experimental/EditorDataStorageFeatures/Source/TedsDebugger/Private/QueryEditor/TedsQueryEditorModel.cpp b/Engine/Plugins/Experimental/EditorDataStorageFeatures/Source/TedsDebugger/Private/QueryEditor/TedsQueryEditorModel.cpp
--- a/Engine/Plugins/Experimental/EditorDataStorageFeatures/Source/TedsDebugger/Private/QueryEditor/TedsQueryEditorModel.cpp
+++ b/Engine/Plugins/Experimental/EditorDataStorageFeatures/Source/TedsDebugger/Private/QueryEditor/TedsQueryEditorModel.cpp
@@ -6,7 +6,7 @@
 #include "Elements/Interfaces/TypedElementDataStorageInterface.h"
 #include "UObject/UObjectIterator.h"
 
-namespace UE::EditorDataStorage::Debug::QueryEditor
+namespace UE::Editor::DataStorage::Debug::QueryEditor
 {
 	bool FConditionEntryHandle::operator==(const FConditionEntryHandle& Rhs) const
 	{
diff --git a/Engine/Plugins/Experimental/EditorDataStorageFeatures/Source/TedsDebugger/Private/QueryEditor/TedsQueryEditorModel.h b/Engine/Plugins/Experimental/EditorDataStorageFeatures/Source/TedsDebugger/Private/QueryEditor/TedsQueryEditorModel.h
--- a/Engine/Plugins/Experimental/EditorDataStorageFeatures/Source/TedsDebugger/Private/QueryEditor/TedsQueryEditorModel.h
+++ b/Engine/Plugins/Experimental/EditorDataStorageFeatures/Source/TedsDebugger/Private/QueryEditor/TedsQueryEditorModel.h
@@ -6,10 +6,12 @@
 
 class ITypedElementDataStorageInterface;
 
-namespace UE::EditorDataStorage::Debug::QueryEditor
+namespace UE::Editor::DataStorage
+{
+namespace Debug::QueryEditor
 {
 	class FTedsQueryEditorModel;
-	
+
 	enum class EOperatorType : uint32
 	{
 		Invalid, // Not initialized
@@ -114,6 +116,7 @@ namespace UE::EditorDataStorage::Debug::QueryEditor
 
 		uint64 CurrentVersion = 0;
 		uint64 GeneratedVersion = 0;
-		TypedElementDataStorage::QueryHandle QueryHandle;
+		QueryHandle QueryHandle;
 	};
-}
+} // namespace Debug::QueryEditor
+} // namespace UE::Editor::DataStorage
diff --git a/Engine/Plugins/Experimental/EditorDataStorageFeatures/Source/TedsDebugger/Private/STedsDebugger.cpp b/Engine/Plugins/Experimental/EditorDataStorageFeatures/Source/TedsDebugger/Private/STedsDebugger.cpp
--- a/Engine/Plugins/Experimental/EditorDataStorageFeatures/Source/TedsDebugger/Private/STedsDebugger.cpp
+++ b/Engine/Plugins/Experimental/EditorDataStorageFeatures/Source/TedsDebugger/Private/STedsDebugger.cpp
@@ -16,13 +16,14 @@
 
 #define LOCTEXT_NAMESPACE "STedsDebugger"
 
-namespace UE::EditorDataStorage::Debugger::Private
+namespace UE::Editor::DataStorage::Debug
 {
-	FName QueryEditorToolTabName = TEXT("TEDS Query Editor");
-	FName TableViewerToolTabName = TEXT("TEDS Table Viewer");
-	FName ToolbarTabName = TEXT("TEDS Debugger Toolbar");
-}
-
+	namespace Private
+	{
+		FName QueryEditorToolTabName = TEXT("TEDS Query Editor");
+		FName TableViewerToolTabName = TEXT("TEDS Table Viewer");
+		FName ToolbarTabName = TEXT("TEDS Debugger Toolbar");
+	}
 
 STedsDebugger::~STedsDebugger()
 {
@@ -43,8 +44,6 @@ void STedsDebugger::Construct(const FArguments& InArgs, const TSharedRef<SDockTa
 	// Register Tab Spawners
 	RegisterTabSpawners();
 
-	using namespace UE::EditorDataStorage::Debugger::Private;
-
 	// Setup the default layout
 	TSharedRef<FTabManager::FLayout> Layout = FTabManager::NewLayout("TedsDebuggerLayout_v0")
 	->AddArea
@@ -54,7 +53,7 @@ void STedsDebugger::Construct(const FArguments& InArgs, const TSharedRef<SDockTa
 		->Split
 		(
 			FTabManager::NewStack()
-			->AddTab(ToolbarTabName, ETabState::OpenedTab)
+			->AddTab(Private::ToolbarTabName, ETabState::OpenedTab)
 			->SetHideTabWell(true)
 		)
 		->Split
@@ -64,8 +63,8 @@ void STedsDebugger::Construct(const FArguments& InArgs, const TSharedRef<SDockTa
 			->Split
 			(
 				FTabManager::NewStack()
-					->AddTab(QueryEditorToolTabName, ETabState::OpenedTab)
-					->AddTab(TableViewerToolTabName, ETabState::OpenedTab)
+					->AddTab(Private::QueryEditorToolTabName, ETabState::OpenedTab)
+					->AddTab(Private::TableViewerToolTabName, ETabState::OpenedTab)
 			)
 		)
 	);
@@ -136,15 +135,15 @@ TSharedRef<SDockTab> STedsDebugger::SpawnQueryEditorTab(const FSpawnTabArgs& Arg
 		if (Registry && Registry->AreDataStorageInterfacesSet())
 		{
 			ITypedElementDataStorageInterface* DataStorageInterface = Registry->GetMutableDataStorage();
-			QueryEditorModel = MakeUnique<UE::EditorDataStorage::Debug::QueryEditor::FTedsQueryEditorModel>(*DataStorageInterface);
+			QueryEditorModel = MakeUnique<QueryEditor::FTedsQueryEditorModel>(*DataStorageInterface);
 		}
 	}
 	if (QueryEditorModel)
 	{
 		QueryEditorModel->Reset();	
 
-		TSharedRef<UE::EditorDataStorage::Debug::QueryEditor::SQueryEditorWidget> QueryEditor =
-			SNew(UE::EditorDataStorage::Debug::QueryEditor::SQueryEditorWidget, *QueryEditorModel);
+		TSharedRef<QueryEditor::SQueryEditorWidget> QueryEditor =
+			SNew(QueryEditor::SQueryEditorWidget, *QueryEditorModel);
 		DockTab->SetContent(QueryEditor);
 	}
 	else
@@ -207,20 +206,18 @@ TSharedRef<SDockTab> STedsDebugger::SpawnTableViewerTab(const FSpawnTabArgs& Arg
 
 void STedsDebugger::RegisterTabSpawners()
 {
-	using namespace UE::EditorDataStorage::Debugger::Private;
-
 	const TSharedRef<FWorkspaceItem> AppMenuGroup =
 		TabManager->AddLocalWorkspaceMenuCategory(LOCTEXT("TedsDebuggerGroupName", "Teds Debugger"));
 
 	TabManager->RegisterTabSpawner(
-		ToolbarTabName,
+		Private::ToolbarTabName,
 		FOnSpawnTab::CreateRaw(this, &STedsDebugger::SpawnToolbar))
 		.SetGroup(AppMenuGroup)
 		.SetDisplayName(LOCTEXT("TedsDebugger_ToolbarDisplayName", "Toolbar"))
 		.SetAutoGenerateMenuEntry(false);
 	
 	TabManager->RegisterTabSpawner(
-		QueryEditorToolTabName,
+		Private::QueryEditorToolTabName,
 		FOnSpawnTab::CreateRaw(this, &STedsDebugger::SpawnQueryEditorTab))
 		.SetGroup(AppMenuGroup)
 		.SetDisplayName(LOCTEXT("TedsDebugger_QueryEditorDisplayName", "Query Editor"))
@@ -228,7 +225,7 @@ void STedsDebugger::RegisterTabSpawners()
 		.SetIcon(FSlateIcon(FAppStyle::GetAppStyleSetName(), "Debug"));
 	
 	TabManager->RegisterTabSpawner(
-		TableViewerToolTabName,
+		Private::TableViewerToolTabName,
 		FOnSpawnTab::CreateRaw(this, &STedsDebugger::SpawnTableViewerTab))
 		.SetGroup(AppMenuGroup)
 		.SetDisplayName(LOCTEXT("TedsDebugger_TableViewerDisplayName", "Table Viewer"))
@@ -248,7 +245,7 @@ void STedsDebugger::NavigateToRow(TypedElementDataStorage::RowHandle InRow) cons
 	// If the debugger isn't already open, open it
 	if (!TableViewerInstance.IsValid())
 	{
-		TabManager->TryInvokeTab(UE::EditorDataStorage::Debugger::Private::TableViewerToolTabName);
+		TabManager->TryInvokeTab(Private::TableViewerToolTabName);
 	}
 
 	TSharedPtr<ISceneOutliner> TableViewerPinned = TableViewerInstance.Pin();
@@ -280,6 +277,6 @@ void STedsDebugger::NavigateToRow(TypedElementDataStorage::RowHandle InRow) cons
 		return false;
 	}));
 }
-
+} // namespace UE::Editor::DataStorage::Debug
 
 #undef LOCTEXT_NAMESPACE
diff --git a/Engine/Plugins/Experimental/EditorDataStorageFeatures/Source/TedsDebugger/Private/STedsDebugger.h b/Engine/Plugins/Experimental/EditorDataStorageFeatures/Source/TedsDebugger/Private/STedsDebugger.h
--- a/Engine/Plugins/Experimental/EditorDataStorageFeatures/Source/TedsDebugger/Private/STedsDebugger.h
+++ b/Engine/Plugins/Experimental/EditorDataStorageFeatures/Source/TedsDebugger/Private/STedsDebugger.h
@@ -8,47 +8,53 @@
 
 class ISceneOutliner;
 
-class STedsDebugger : public SCompoundWidget
+namespace UE::Editor::DataStorage
 {
-public:
-
-	SLATE_BEGIN_ARGS(STedsDebugger) 
+	namespace Debug
 	{
-	}
-	SLATE_END_ARGS()
-
-public:
-	STedsDebugger() = default;
-	virtual ~STedsDebugger() override;
-
-	/**
-	* Constructs the debugger.
-	*
-	* @param InArgs The Slate argument list.
-	* @param ConstructUnderMajorTab The major tab which will contain the session front-end.
-	* @param ConstructUnderWindow The window in which this widget is being constructed.
-	*/
-	void Construct(const FArguments& InArgs, const TSharedRef<SDockTab>& ConstructUnderMajorTab, const TSharedPtr<SWindow>& ConstructUnderWindow);
-
-	void NavigateToRow(TypedElementDataStorage::RowHandle InRow) const;
+		class STedsDebugger : public SCompoundWidget
+		{
+		public:
+
+			SLATE_BEGIN_ARGS(STedsDebugger) 
+			{
+			}
+			SLATE_END_ARGS()
+
+		public:
+			STedsDebugger() = default;
+			virtual ~STedsDebugger() override;
+
+			/**
+			* Constructs the debugger.
+			*
+			* @param InArgs The Slate argument list.
+			* @param ConstructUnderMajorTab The major tab which will contain the session front-end.
+			* @param ConstructUnderWindow The window in which this widget is being constructed.
+			*/
+			void Construct(const FArguments& InArgs, const TSharedRef<SDockTab>& ConstructUnderMajorTab, const TSharedPtr<SWindow>& ConstructUnderWindow);
+
+			void NavigateToRow(RowHandle InRow) const;
 	
-private:
-	TSharedRef<SDockTab> SpawnToolbar(const FSpawnTabArgs& Args);
-	TSharedRef<SDockTab> SpawnQueryEditorTab(const FSpawnTabArgs& Args);
-	TSharedRef<SDockTab> SpawnTableViewerTab(const FSpawnTabArgs& Args);
-	void FillWindowMenu( FMenuBuilder& MenuBuilder);
+		private:
+			TSharedRef<SDockTab> SpawnToolbar(const FSpawnTabArgs& Args);
+			TSharedRef<SDockTab> SpawnQueryEditorTab(const FSpawnTabArgs& Args);
+			TSharedRef<SDockTab> SpawnTableViewerTab(const FSpawnTabArgs& Args);
+			void FillWindowMenu( FMenuBuilder& MenuBuilder);
 
-	void RegisterTabSpawners();
+			void RegisterTabSpawners();
 
-private:
+		private:
 	
-	// Holds the tab manager that manages the front-end's tabs.
-	TSharedPtr<FTabManager> TabManager;
-
-	// Table Viewer
-	TypedElementDataStorage::QueryHandle TableViewerQuery;
-	TWeakPtr<ISceneOutliner> TableViewerInstance;
-
-	// Query Editor
-	TUniquePtr<UE::EditorDataStorage::Debug::QueryEditor::FTedsQueryEditorModel> QueryEditorModel;
-};
+			// Holds the tab manager that manages the front-end's tabs.
+			TSharedPtr<FTabManager> TabManager;
+
+			// Table Viewer
+			QueryHandle TableViewerQuery;
+			TWeakPtr<ISceneOutliner> TableViewerInstance;
+
+			// Query Editor
+			TUniquePtr<QueryEditor::FTedsQueryEditorModel> QueryEditorModel;
+		};
+	} // namespace Debug
+} // namespace UE::Editor::DataStorage
diff --git a/Engine/Plugins/Experimental/EditorDataStorageFeatures/Source/TedsDebugger/Private/TedsDebuggerModule.cpp b/Engine/Plugins/Experimental/EditorDataStorageFeatures/Source/TedsDebugger/Private/TedsDebuggerModule.cpp
--- a/Engine/Plugins/Experimental/EditorDataStorageFeatures/Source/TedsDebugger/Private/TedsDebuggerModule.cpp
+++ b/Engine/Plugins/Experimental/EditorDataStorageFeatures/Source/TedsDebugger/Private/TedsDebuggerModule.cpp
@@ -12,7 +12,7 @@
 
 #define LOCTEXT_NAMESPACE "TedsDebuggerModule"
 
-namespace UE::EditorDataStorage::Debug
+namespace UE::Editor::DataStorage::Debug
 {
 namespace Private
 {
@@ -37,10 +37,8 @@ void FTedsDebuggerModule::ShutdownModule()
 
 void FTedsDebuggerModule::RegisterTabSpawners()
 {
-	using namespace UE::EditorDataStorage::Debug::Private;
-
 	FGlobalTabmanager::Get()->RegisterNomadTabSpawner(
-		TedsDebuggerTableName,
+		Private::TedsDebuggerTableName,
 		FOnSpawnTab::CreateRaw(this, &FTedsDebuggerModule::OpenTedsDebuggerTab))
 		.SetGroup(WorkspaceMenu::GetMenuStructure().GetDeveloperToolsDebugCategory())
 		.SetDisplayName(LOCTEXT("TedsDebugger_QueryEditorDisplayName", "TEDS Debugger"))
@@ -52,7 +50,7 @@ void FTedsDebuggerModule::UnregisterTabSpawners() const
 {
 	if (FSlateApplication::IsInitialized())
 	{
-		FGlobalTabmanager::Get()->UnregisterNomadTabSpawner(UE::EditorDataStorage::Debug::Private::TedsDebuggerTableName);
+		FGlobalTabmanager::Get()->UnregisterNomadTabSpawner(Private::TedsDebuggerTableName);
 	}
 }
 
@@ -70,7 +68,7 @@ TSharedRef<SDockTab> FTedsDebuggerModule::OpenTedsDebuggerTab(const FSpawnTabArg
 	return MajorTab;
 }
 
-void FTedsDebuggerModule::NavigateToRow(TypedElementDataStorage::RowHandle InRow) const
+void FTedsDebuggerModule::NavigateToRow(RowHandle InRow) const
 {
 	UTypedElementRegistry* Registry = UTypedElementRegistry::GetInstance();
 
@@ -82,7 +80,7 @@ void FTedsDebuggerModule::NavigateToRow(TypedElementDataStorage::RowHandle InRow
 	// If the debugger isn't already open, open it
 	if(!TedsDebuggerInstance.IsValid())
 	{
-		FGlobalTabmanager::Get()->TryInvokeTab(UE::EditorDataStorage::Debug::Private::TedsDebuggerTableName);
+		FGlobalTabmanager::Get()->TryInvokeTab(Private::TedsDebuggerTableName);
 	}
 
 	TSharedPtr<STedsDebugger> TedsDebuggerPinned = TedsDebuggerInstance.Pin();
@@ -93,8 +91,8 @@ void FTedsDebuggerModule::NavigateToRow(TypedElementDataStorage::RowHandle InRow
 
 	TedsDebuggerPinned->NavigateToRow(InRow);
 }
-}
+} // namespace UE::Editor::DataStorage::Debug
 
-IMPLEMENT_MODULE(UE::EditorDataStorage::Debug::FTedsDebuggerModule, TedsDebugger);
+IMPLEMENT_MODULE(UE::Editor::DataStorage::Debug::FTedsDebuggerModule, TedsDebugger);
 
 #undef LOCTEXT_NAMESPACE
\ No newline at end of file
diff --git a/Engine/Plugins/Experimental/EditorDataStorageFeatures/Source/TedsDebugger/Private/TedsDebuggerModule.h b/Engine/Plugins/Experimental/EditorDataStorageFeatures/Source/TedsDebugger/Private/TedsDebuggerModule.h
--- a/Engine/Plugins/Experimental/EditorDataStorageFeatures/Source/TedsDebugger/Private/TedsDebuggerModule.h
+++ b/Engine/Plugins/Experimental/EditorDataStorageFeatures/Source/TedsDebugger/Private/TedsDebuggerModule.h
@@ -10,34 +10,38 @@ class SDockTab;
 class SWidget;
 class FSpawnTabArgs;
 class ISceneOutliner;
-class STedsDebugger;
 
-namespace UE::EditorDataStorage::Debug
+namespace UE::Editor::DataStorage
 {
-/**
- * Implements the Teds Debugger module.
- */
-class FTedsDebuggerModule
-	: public IModuleInterface
-{
-public:
-
-	FTedsDebuggerModule() = default;
-
-	// IModuleInterface interface
-	virtual void StartupModule() override;
-	virtual void ShutdownModule() override;
-
-	// Open the TEDS Debugger if not already open, and navigate to the given row in the table viewer tab
-	void NavigateToRow(TypedElementDataStorage::RowHandle InRow) const;
-
-private:
-	void RegisterTabSpawners();
-	void UnregisterTabSpawners() const;
+	namespace Debug
+	{
+		class STedsDebugger;
+
+		/**
+		 * Implements the Teds Debugger module.
+		 */
+		class FTedsDebuggerModule
+			: public IModuleInterface
+		{
+		public:
+
+			FTedsDebuggerModule() = default;
+
+			// IModuleInterface interface
+			virtual void StartupModule() override;
+			virtual void ShutdownModule() override;
+
+			// Open the TEDS Debugger if not already open, and navigate to the given row in the table viewer tab
+			void NavigateToRow(RowHandle InRow) const;
+
+		private:
+			void RegisterTabSpawners();
+			void UnregisterTabSpawners() const;
 	
-	TSharedRef<SDockTab> OpenTedsDebuggerTab(const FSpawnTabArgs& SpawnTabArgs);
+			TSharedRef<SDockTab> OpenTedsDebuggerTab(const FSpawnTabArgs& SpawnTabArgs);
 
-private:
-	TWeakPtr<STedsDebugger> TedsDebuggerInstance;
-};
-}
\ No newline at end of file
+		private:
+			TWeakPtr<STedsDebugger> TedsDebuggerInstance;
+		};
+	} // namespace Debug
+} // namespace UE::Editor::DataStorage
diff --git a/Engine/Plugins/Experimental/EditorDataStorageFeatures/Source/TedsDebugger/Private/Widgets/QueryEditor/TedsConditionCollectionViewWidget.cpp b/Engine/Plugins/Experimental/EditorDataStorageFeatures/Source/TedsDebugger/Private/Widgets/QueryEditor/TedsConditionCollectionViewWidget.cpp
--- a/Engine/Plugins/Experimental/EditorDataStorageFeatures/Source/TedsDebugger/Private/Widgets/QueryEditor/TedsConditionCollectionViewWidget.cpp
+++ b/Engine/Plugins/Experimental/EditorDataStorageFeatures/Source/TedsDebugger/Private/Widgets/QueryEditor/TedsConditionCollectionViewWidget.cpp
@@ -9,172 +9,173 @@
 #include "Widgets/Layout/SWrapBox.h"
 #include "Widgets/Text/STextBlock.h"
 
-using namespace UE::EditorDataStorage::Debug::QueryEditor;
-
 #define LOCTEXT_NAMESPACE "TedsDebuggerModule"
 
-namespace UE::EditorDataStorage::Debug::QueryEditor::Private
+namespace UE::Editor::DataStorage::Debug::QueryEditor
 {
-	FText GetOperatorTypeText(EOperatorType OperatorType)
+	namespace Private
 	{
-		switch(OperatorType)
+		FText GetOperatorTypeText(EOperatorType OperatorType)
 		{
-		case EOperatorType::Select:
-			return FText::FromStringView(TEXT("Select<"));
-		case EOperatorType::All:
-			return FText::FromStringView(TEXT("All<"));
-		case EOperatorType::Any:
-			return FText::FromStringView(TEXT("Any<"));
-		case EOperatorType::None:
-			return FText::FromStringView(TEXT("None<"));
-		case EOperatorType::Unset:
-			return FText::FromStringView(TEXT("Unset<"));
-		case EOperatorType::Invalid:
-		default:
-			return FText::FromStringView(TEXT("Invalid<"));
+			switch (OperatorType)
+			{
+			case EOperatorType::Select:
+				return FText::FromStringView(TEXT("Select<"));
+			case EOperatorType::All:
+				return FText::FromStringView(TEXT("All<"));
+			case EOperatorType::Any:
+				return FText::FromStringView(TEXT("Any<"));
+			case EOperatorType::None:
+				return FText::FromStringView(TEXT("None<"));
+			case EOperatorType::Unset:
+				return FText::FromStringView(TEXT("Unset<"));
+			case EOperatorType::Invalid:
+			default:
+				return FText::FromStringView(TEXT("Invalid<"));
+			}
 		}
 	}
-}
 
-SConditionCollectionViewWidget::~SConditionCollectionViewWidget()
-{
-	Model->GetModelChangedDelegate().Remove(OnModelChangedDelegate);
-}
+	SConditionCollectionViewWidget::~SConditionCollectionViewWidget()
+	{
+		Model->GetModelChangedDelegate().Remove(OnModelChangedDelegate);
+	}
 
-void SConditionCollectionViewWidget::Construct(
-	const FArguments& InArgs,
-	FTedsQueryEditorModel& InModel, 
-	QueryEditor::EOperatorType InOperatorType)
-{
-	Model = &InModel;
-	OperatorType = InOperatorType;
+	void SConditionCollectionViewWidget::Construct(
+		const FArguments& InArgs,
+		FTedsQueryEditorModel& InModel, 
+		QueryEditor::EOperatorType InOperatorType)
+	{
+		Model = &InModel;
+		OperatorType = InOperatorType;
 
-	OnModelChangedDelegate = Model->GetModelChangedDelegate().AddRaw(this, &SConditionCollectionViewWidget::OnModelChanged);
+		OnModelChangedDelegate = Model->GetModelChangedDelegate().AddRaw(this, &SConditionCollectionViewWidget::OnModelChanged);
 
-	ChildSlot
-	[
-		SNew(SBorder)
-		.Padding(4.0, 3.0, 3.0, 4.0)
-		[
-		SNew(SHorizontalBox)
-		+SHorizontalBox::Slot()
-		.AutoWidth()
+		ChildSlot
 		[
-			SNew(SVerticalBox)
-			+SVerticalBox::Slot()
-			.VAlign(EVerticalAlignment::VAlign_Center)
+			SNew(SBorder)
+			.Padding(4.0, 3.0, 3.0, 4.0)
 			[
-				SNew(STextBlock)
-				.Text(Private::GetOperatorTypeText(OperatorType))
+			SNew(SHorizontalBox)
+			+SHorizontalBox::Slot()
+			.AutoWidth()
+			[
+				SNew(SVerticalBox)
+				+SVerticalBox::Slot()
+				.VAlign(EVerticalAlignment::VAlign_Center)
+				[
+					SNew(STextBlock)
+					.Text(Private::GetOperatorTypeText(OperatorType))
+				]
 			]
-		]
-		+SHorizontalBox::Slot()
-		[
-		SAssignNew(ColumnButtonWrap, SWrapBox)
-		.UseAllottedSize(true)
-		+SWrapBox::Slot()
-		.FillEmptySpace(false)
-		[
-			SNullWidget::NullWidget
-		]
-		]
-		+SHorizontalBox::Slot()
-		.AutoWidth()
-		[
-			SNew(SVerticalBox)
-			+SVerticalBox::Slot()
-			.VAlign(EVerticalAlignment::VAlign_Center)
+			+SHorizontalBox::Slot()
 			[
-				SNew(STextBlock)
-				.Text(LOCTEXT("TedsQueryEditor_OpType_>", ">"))
+			SAssignNew(ColumnButtonWrap, SWrapBox)
+			.UseAllottedSize(true)
+			+SWrapBox::Slot()
+			.FillEmptySpace(false)
+			[
+				SNullWidget::NullWidget
+			]
+			]
+			+SHorizontalBox::Slot()
+			.AutoWidth()
+			[
+				SNew(SVerticalBox)
+				+SVerticalBox::Slot()
+				.VAlign(EVerticalAlignment::VAlign_Center)
+				[
+					SNew(STextBlock)
+					.Text(LOCTEXT("TedsQueryEditor_OpType_>", ">"))
+				]
 			]
-		]
-		]
-	];
+			]
+		];
 
-	SetCanTick(true);
-}
+		SetCanTick(true);
+	}
 
-void SConditionCollectionViewWidget::Tick(const FGeometry& AllottedGeometry, const double InCurrentTime, const float InDeltaTime)
-{
-	SCompoundWidget::Tick(AllottedGeometry, InCurrentTime, InDeltaTime);
+	void SConditionCollectionViewWidget::Tick(const FGeometry& AllottedGeometry, const double InCurrentTime, const float InDeltaTime)
+	{
+		SCompoundWidget::Tick(AllottedGeometry, InCurrentTime, InDeltaTime);
 
-	ColumnButtonWrap->ClearChildren();
-	SetCanTick(false);
+		ColumnButtonWrap->ClearChildren();
+		SetCanTick(false);
 
-	ButtonStyle = FButtonStyle(FCoreStyle::Get().GetWidgetStyle< FButtonStyle >( "Button" ))
-		.SetHoveredForeground(FSlateColor(FLinearColor(0.9f, 0.0f, 0.0f, 1.0f)));
+		ButtonStyle = FButtonStyle(FCoreStyle::Get().GetWidgetStyle< FButtonStyle >( "Button" ))
+			.SetHoveredForeground(FSlateColor(FLinearColor(0.9f, 0.0f, 0.0f, 1.0f)));
 
-	if (Model->CountConditionsOfOperator(OperatorType) != 0)
-	{
-		Model->ForEachCondition([this](FTedsQueryEditorModel&, FConditionEntryHandle Handle)
+		if (Model->CountConditionsOfOperator(OperatorType) != 0)
 		{
-			if (Model->GetOperatorType(Handle) != OperatorType)
+			Model->ForEachCondition([this](FTedsQueryEditorModel&, FConditionEntryHandle Handle)
 			{
-				return;
-			}
-			// TODO Add styling to the button to go red and put an "X" to signal clicking
-			// will remove the column
-			const UScriptStruct* Type = Model->GetColumnScriptStruct(Handle);
-			FText TempName = FText::FromString(Type->GetName());
+				if (Model->GetOperatorType(Handle) != OperatorType)
+				{
+					return;
+				}
+				// TODO Add styling to the button to go red and put an "X" to signal clicking
+				// will remove the column
+				const UScriptStruct* Type = Model->GetColumnScriptStruct(Handle);
+				FText TempName = FText::FromString(Type->GetName());
 			
-			TSharedRef<STextBlock> TextNameOfColumn = SNew(STextBlock)
-				.Text(TempName);
+				TSharedRef<STextBlock> TextNameOfColumn = SNew(STextBlock)
+					.Text(TempName);
 
-			TSharedRef<SOverlay> ButtonContent = SNew(SOverlay);
+				TSharedRef<SOverlay> ButtonContent = SNew(SOverlay);
 			
-			TSharedRef<SButton> Button = SNew(SButton)
-				.Text(TempName)
-				.OnClicked_Lambda( [this, Handle]() -> FReply 
+				TSharedRef<SButton> Button = SNew(SButton)
+					.Text(TempName)
+					.OnClicked_Lambda( [this, Handle]() -> FReply 
+					{
+						Model->SetOperatorType(Handle, EOperatorType::Unset);
+						return FReply::Handled();
+					})
+					.IsEnabled_Lambda( [this, Handle]() -> bool
+					{
+						return Model->CanSetOperatorType(Handle, EOperatorType::Unset).Key;
+					})
+					.Content()
+					[
+						TextNameOfColumn
+					];
+				STextBlock* TextBlockPtr = &TextNameOfColumn.Get();
+				Button->SetOnHovered(FSimpleDelegate::CreateLambda([TextBlockPtr]()
 				{
-					Model->SetOperatorType(Handle, EOperatorType::Unset);
-					return FReply::Handled();
-				})
-				.IsEnabled_Lambda( [this, Handle]() -> bool
+					TextBlockPtr->SetVisibility(EVisibility::Hidden);
+				}));
+				Button->SetOnUnhovered(FSimpleDelegate::CreateLambda([TextBlockPtr, TempName]()
 				{
-					return Model->CanSetOperatorType(Handle, EOperatorType::Unset).Key;
-				})
-				.Content()
-				[
-					TextNameOfColumn
-				];
-			STextBlock* TextBlockPtr = &TextNameOfColumn.Get();
-			Button->SetOnHovered(FSimpleDelegate::CreateLambda([TextBlockPtr]()
-			{
-				TextBlockPtr->SetVisibility(EVisibility::Hidden);
-			}));
-			Button->SetOnUnhovered(FSimpleDelegate::CreateLambda([TextBlockPtr, TempName]()
-			{
-				TextBlockPtr->SetVisibility(EVisibility::Visible);
-			}));
+					TextBlockPtr->SetVisibility(EVisibility::Visible);
+				}));
 
 			
-			ButtonContent->AddSlot()
-			[
-				Button
-			];
+				ButtonContent->AddSlot()
+				[
+					Button
+				];
 			
+				ColumnButtonWrap->AddSlot()
+				[
+					ButtonContent
+				];
+
+			});
+		}
+		else
+		{
 			ColumnButtonWrap->AddSlot()
 			[
-				ButtonContent
+				SNullWidget::NullWidget
 			];
-
-		});
+		}
+	
 	}
-	else
+
+	void SConditionCollectionViewWidget::OnModelChanged()
 	{
-		ColumnButtonWrap->AddSlot()
-		[
-			SNullWidget::NullWidget
-		];
+		SetCanTick(true);
+		Invalidate(EInvalidateWidgetReason::Layout);
 	}
-	
-}
-
-void SConditionCollectionViewWidget::OnModelChanged()
-{
-	SetCanTick(true);
-	Invalidate(EInvalidateWidgetReason::Layout);
 }
 
 #undef LOCTEXT_NAMESPACE
\ No newline at end of file
diff --git a/Engine/Plugins/Experimental/EditorDataStorageFeatures/Source/TedsDebugger/Private/Widgets/QueryEditor/TedsConditionCollectionViewWidget.h b/Engine/Plugins/Experimental/EditorDataStorageFeatures/Source/TedsDebugger/Private/Widgets/QueryEditor/TedsConditionCollectionViewWidget.h
--- a/Engine/Plugins/Experimental/EditorDataStorageFeatures/Source/TedsDebugger/Private/Widgets/QueryEditor/TedsConditionCollectionViewWidget.h
+++ b/Engine/Plugins/Experimental/EditorDataStorageFeatures/Source/TedsDebugger/Private/Widgets/QueryEditor/TedsConditionCollectionViewWidget.h
@@ -8,7 +8,7 @@
 class SWrapBox;
 class SHorizontalBox;
 
-namespace UE::EditorDataStorage::Debug::QueryEditor
+namespace UE::Editor::DataStorage::Debug::QueryEditor
 {
 	enum class EOperatorType : uint32;
 	class FTedsQueryEditorModel;
@@ -20,7 +20,7 @@ namespace UE::EditorDataStorage::Debug::QueryEditor
 		SLATE_END_ARGS()
 
 		~SConditionCollectionViewWidget() override;
-		void Construct(const FArguments& InArgs, FTedsQueryEditorModel& InModel, QueryEditor::EOperatorType InOperatorType);
+		void Construct(const FArguments& InArgs, FTedsQueryEditorModel& InModel, EOperatorType InOperatorType);
 		void Tick(const FGeometry& AllottedGeometry, const double InCurrentTime, const float InDeltaTime) override;
 		
 	private:
@@ -28,7 +28,7 @@ namespace UE::EditorDataStorage::Debug::QueryEditor
 		void OnModelChanged();
 
 		FTedsQueryEditorModel* Model = nullptr;
-		QueryEditor::EOperatorType OperatorType = EOperatorType::Unset;
+		EOperatorType OperatorType = EOperatorType::Unset;
 		FDelegateHandle OnModelChangedDelegate;
 
 		TSharedPtr<SWrapBox> ColumnButtonWrap;
diff --git a/Engine/Plugins/Experimental/EditorDataStorageFeatures/Source/TedsDebugger/Private/Widgets/QueryEditor/TedsConditionSelectionComboWidget.h b/Engine/Plugins/Experimental/EditorDataStorageFeatures/Source/TedsDebugger/Private/Widgets/QueryEditor/TedsConditionSelectionComboWidget.h
--- a/Engine/Plugins/Experimental/EditorDataStorageFeatures/Source/TedsDebugger/Private/Widgets/QueryEditor/TedsConditionSelectionComboWidget.h
+++ b/Engine/Plugins/Experimental/EditorDataStorageFeatures/Source/TedsDebugger/Private/Widgets/QueryEditor/TedsConditionSelectionComboWidget.h
@@ -8,7 +8,7 @@
 template<typename>
 class SComboBox;
 
-namespace UE::EditorDataStorage::Debug::QueryEditor
+namespace UE::Editor::DataStorage::Debug::QueryEditor
 {
 	class FTedsQueryEditorModel;
 	struct FConditionEntry;
@@ -40,4 +40,4 @@ namespace UE::EditorDataStorage::Debug::QueryEditor
 
 		TSharedPtr<SFilterBar<TSharedPtr<FComboItem>>> FilterThing;
 	};
-}
+} // namespace UE::Editor::DataStorage::Debug::QueryEditor
diff --git a/Engine/Plugins/Experimental/EditorDataStorageFeatures/Source/TedsDebugger/Private/Widgets/QueryEditor/TedsQueryEditorResultsView.h b/Engine/Plugins/Experimental/EditorDataStorageFeatures/Source/TedsDebugger/Private/Widgets/QueryEditor/TedsQueryEditorResultsView.h
--- a/Engine/Plugins/Experimental/EditorDataStorageFeatures/Source/TedsDebugger/Private/Widgets/QueryEditor/TedsQueryEditorResultsView.h
+++ b/Engine/Plugins/Experimental/EditorDataStorageFeatures/Source/TedsDebugger/Private/Widgets/QueryEditor/TedsQueryEditorResultsView.h
@@ -4,63 +4,57 @@
 #include "Elements/Interfaces/TypedElementDataStorageInterface.h"
 #include "Widgets/SCompoundWidget.h"
 
-namespace UE::EditorDataStorage
-{
-	class SRowDetails;
-}
-
 class ISceneOutliner;
 class SSceneOutliner;
 class SHorizontalBox;
 struct FTypedElementWidgetConstructor;
 
-namespace UE::EditorDataStorage
+namespace UE::Editor::DataStorage
 {
 	class STedsTableViewer;
 	class FQueryStackNode_RowView;
 	class FTedsTableViewerColumn;
-}
-
-namespace UE::EditorDataStorage::Debug::QueryEditor
-{
-	class FTedsQueryEditorModel;
+	class SRowDetails;
 
-	class SResultsView : public SCompoundWidget
+	namespace Debug::QueryEditor
 	{
-	public:
-		SLATE_BEGIN_ARGS( SResultsView ){}
-		SLATE_END_ARGS()
+		class FTedsQueryEditorModel;
 
-		~SResultsView() override;
-		void Construct(const FArguments& InArgs, FTedsQueryEditorModel& InModel);
-		void Tick(const FGeometry& AllottedGeometry, const double InCurrentTime, const float InDeltaTime) override;
+		class SResultsView : public SCompoundWidget
+		{
+		public:
+			SLATE_BEGIN_ARGS( SResultsView ){}
+			SLATE_END_ARGS()
 
-	private:
+			~SResultsView() override;
+			void Construct(const FArguments& InArgs, FTedsQueryEditorModel& InModel);
+			void Tick(const FGeometry& AllottedGeometry, const double InCurrentTime, const float InDeltaTime) override;
 
-		void OnModelChanged();
-		void CreateRowHandleColumn();
-		
-		FTedsQueryEditorModel* Model = nullptr;
-		FDelegateHandle ModelChangedDelegateHandle;
-		bool bModelDirty = true;
+		private:
 
+			void OnModelChanged();
+			void CreateRowHandleColumn();
+		
+			FTedsQueryEditorModel* Model = nullptr;
+			FDelegateHandle ModelChangedDelegateHandle;
+			bool bModelDirty = true;
 
-		TypedElementQueryHandle CountQueryHandle = TypedElementDataStorage::InvalidQueryHandle;
-		TypedElementQueryHandle TableViewerQueryHandle = TypedElementDataStorage::InvalidQueryHandle;
 
-		TArray<TypedElementDataStorage::RowHandle> TableViewerRows;
-		// We have to keep a TSet copy because queries return duplicate rows sometimes and to have some form of sorted order for the rows for now
-		TSet<TypedElementDataStorage::RowHandle> TableViewerRows_Set;
-		TSharedPtr<UE::EditorDataStorage::STedsTableViewer> TableViewer;
-		TSharedPtr<UE::EditorDataStorage::FQueryStackNode_RowView> RowQueryStack;
+			QueryHandle CountQueryHandle = InvalidQueryHandle;
+			QueryHandle TableViewerQueryHandle = InvalidQueryHandle;
 
-		// Custom column for the table viewer to display row handles
-		TSharedPtr<UE::EditorDataStorage::FTedsTableViewerColumn> RowHandleColumn;
+			TArray<RowHandle> TableViewerRows;
+			// We have to keep a TSet copy because queries return duplicate rows sometimes and to have some form of sorted order for the rows for now
+			TSet<RowHandle> TableViewerRows_Set;
+			TSharedPtr<STedsTableViewer> TableViewer;
+			TSharedPtr<FQueryStackNode_RowView> RowQueryStack;
 
-		// Widget that displays details of a row
-		TSharedPtr<UE::EditorDataStorage::SRowDetails> RowDetailsWidget;
+			// Custom column for the table viewer to display row handles
+			TSharedPtr<FTedsTableViewerColumn> RowHandleColumn;
 
-	};
+			// Widget that displays details of a row
+			TSharedPtr<SRowDetails> RowDetailsWidget;
+		};
 
-	
-}
+	} // namespace Debug::QueryEditor
+} // namespace UE::Editor::DataStorage
diff --git a/Engine/Plugins/Experimental/EditorDataStorageFeatures/Source/TedsDebugger/Private/Widgets/RowReferenceWidget.cpp b/Engine/Plugins/Experimental/EditorDataStorageFeatures/Source/TedsDebugger/Private/Widgets/RowReferenceWidget.cpp
--- a/Engine/Plugins/Experimental/EditorDataStorageFeatures/Source/TedsDebugger/Private/Widgets/RowReferenceWidget.cpp
+++ b/Engine/Plugins/Experimental/EditorDataStorageFeatures/Source/TedsDebugger/Private/Widgets/RowReferenceWidget.cpp
@@ -17,9 +17,11 @@
 
 #define LOCTEXT_NAMESPACE "RowReferenceWidget"
 
-namespace UE::EditorDataStorage::Debug::Private
+namespace UE::Editor::DataStorage::Debug::Private
 {
-	void OnNavigateHyperlink(const ITypedElementDataStorageInterface* DataStorage, TypedElementDataStorage::RowHandle TargetRowHandle, TypedElementDataStorage::RowHandle UiRowHandle)
+	using namespace UE::Editor::DataStorage;
+
+	void OnNavigateHyperlink(const ITypedElementDataStorageInterface* DataStorage,  RowHandle TargetRowHandle, RowHandle UiRowHandle)
 	{
 		const FTedsOutlinerColumn* TedsOutlinerColumn = DataStorage->GetColumn<FTedsOutlinerColumn>(UiRowHandle);
 
@@ -49,11 +51,11 @@ namespace UE::EditorDataStorage::Debug::Private
 		}
 
 		// If it wasn't found in the table viewer owning this widget, navigate to it in the global TEDS debugger
-		UE::EditorDataStorage::Debug::FTedsDebuggerModule& TedsDebuggerModule = FModuleManager::GetModuleChecked<UE::EditorDataStorage::Debug::FTedsDebuggerModule>("TedsDebugger");
+		Debug::FTedsDebuggerModule& TedsDebuggerModule = FModuleManager::GetModuleChecked<Debug::FTedsDebuggerModule>("TedsDebugger");
 		TedsDebuggerModule.NavigateToRow(TargetRowHandle);
 	}
 	
-	void CreateInternalWidget(const TWeakPtr<SWidget>& InWidget, TypedElementDataStorage::RowHandle UiRow, TypedElementDataStorage::RowHandle TargetRow)
+	void CreateInternalWidget(const TWeakPtr<SWidget>& InWidget, RowHandle UiRow, RowHandle TargetRow)
 	{
 		const TSharedPtr<SWidget> Widget = InWidget.Pin();
 
@@ -86,14 +88,14 @@ namespace UE::EditorDataStorage::Debug::Private
 					.Text(Text)
 					.Style(FAppStyle::Get(), "Common.GotoBlueprintHyperlink")
 					.ToolTipText(TooltipText)
-					.OnNavigate(FSimpleDelegate::CreateStatic(&UE::EditorDataStorage::Debug::Private::OnNavigateHyperlink, DataStorage, TargetRow, UiRow));
+					.OnNavigate(FSimpleDelegate::CreateStatic(&OnNavigateHyperlink, DataStorage, TargetRow, UiRow));
 			
 			WidgetInstance->SetContent(HyperlinkWidget);
 		}
 	}
 
 
-}
+} // namespace UE::Editor::DataStorage::Debug::Private
 
 URowReferenceWidgetFactory::~URowReferenceWidgetFactory()
 {
@@ -110,8 +112,9 @@ void URowReferenceWidgetFactory::RegisterQueries(ITypedElementDataStorageInterfa
 {
 	using namespace TypedElementQueryBuilder;
 	using namespace TypedElementDataStorage;
+	using namespace UE::Editor::DataStorage;
 	
-	const TypedElementQueryHandle UpdateRowReferenceWidget = DataStorage.RegisterQuery(
+	const QueryHandle UpdateRowReferenceWidget = DataStorage.RegisterQuery(
 		Select()
 			.ReadOnly<FTypedElementRowReferenceColumn>()
 		.Where()
@@ -132,7 +135,7 @@ void URowReferenceWidgetFactory::RegisterQueries(ITypedElementDataStorageInterfa
 				Context.RunSubquery(0, Target.Row, CreateSubqueryCallbackBinding(
 					[&Widget, UiRowHandle](const FTypedElementRowReferenceColumn& Target)
 					{
-						UE::EditorDataStorage::Debug::Private::CreateInternalWidget(Widget.Widget, UiRowHandle, Target.Row);
+						Debug::Private::CreateInternalWidget(Widget.Widget, UiRowHandle, Target.Row);
 					}));
 			})
 		.DependsOn()
@@ -153,17 +156,18 @@ TSharedPtr<SWidget> FRowReferenceWidgetConstructor::CreateWidget(const TypedElem
 			.VAlign(VAlign_Center);
 }
 
-bool FRowReferenceWidgetConstructor::FinalizeWidget(ITypedElementDataStorageInterface* DataStorage, ITypedElementDataStorageUiInterface* DataStorageUi, TypedElementDataStorage::RowHandle Row, const TSharedPtr<SWidget>& Widget)
+bool FRowReferenceWidgetConstructor::FinalizeWidget(ITypedElementDataStorageInterface* DataStorage, ITypedElementDataStorageUiInterface* DataStorageUi, UE::Editor::DataStorage::RowHandle Row, const TSharedPtr<SWidget>& Widget)
 {
+	using namespace UE::Editor::DataStorage;
 	checkf(Widget, TEXT("Referenced widget is not valid. A constructed widget may not have been cleaned up. This can "
 		"also happen if this processor is running in the same phase as the processors responsible for cleaning up old "
 		"references."));
 
 	// The actual row we want to view in the widget
-	TypedElementRowHandle TargetRowReference = TypedElementInvalidRowHandle;
+	RowHandle TargetRowReference = InvalidRowHandle;
 
 	// The target row for which this widget was created
-	const TypedElementRowHandle TargetRow = DataStorage->GetColumn<FTypedElementRowReferenceColumn>(Row)->Row;
+	const RowHandle TargetRow = DataStorage->GetColumn<FTypedElementRowReferenceColumn>(Row)->Row;
 
 	// Check if the target row has a row reference column, if so we want to view the row in there.
 	if(const FTypedElementRowReferenceColumn* RowReferenceColumn = DataStorage->GetColumn<FTypedElementRowReferenceColumn>(TargetRow))
@@ -171,7 +175,7 @@ bool FRowReferenceWidgetConstructor::FinalizeWidget(ITypedElementDataStorageInte
 		TargetRowReference = RowReferenceColumn->Row;
 	}
 	
-	UE::EditorDataStorage::Debug::Private::CreateInternalWidget(Widget, Row, TargetRowReference);
+	Debug::Private::CreateInternalWidget(Widget, Row, TargetRowReference);
 
 	return true;
 }
diff --git a/Engine/Plugins/Experimental/EditorDataStorageFeatures/Source/TedsDebugger/Private/Widgets/SlateVisualizationWidget.cpp b/Engine/Plugins/Experimental/EditorDataStorageFeatures/Source/TedsDebugger/Private/Widgets/SlateVisualizationWidget.cpp
--- a/Engine/Plugins/Experimental/EditorDataStorageFeatures/Source/TedsDebugger/Private/Widgets/SlateVisualizationWidget.cpp
+++ b/Engine/Plugins/Experimental/EditorDataStorageFeatures/Source/TedsDebugger/Private/Widgets/SlateVisualizationWidget.cpp
@@ -26,13 +26,13 @@ TSharedPtr<SWidget> FSlateVisualizationWidgetConstructor::CreateWidget(const Typ
 	return SNew(SHorizontalBox);
 }
 
-bool FSlateVisualizationWidgetConstructor::FinalizeWidget(ITypedElementDataStorageInterface* DataStorage, ITypedElementDataStorageUiInterface* DataStorageUi, TypedElementRowHandle Row, const TSharedPtr<SWidget>& Widget)
+bool FSlateVisualizationWidgetConstructor::FinalizeWidget(ITypedElementDataStorageInterface* DataStorage, ITypedElementDataStorageUiInterface* DataStorageUi, UE::Editor::DataStorage::RowHandle Row, const TSharedPtr<SWidget>& Widget)
 {
 	checkf(Widget, TEXT("Referenced widget is not valid. A constructed widget may not have been cleaned up. This can "
 	"also happen if this processor is running in the same phase as the processors responsible for cleaning up old "
 	"references."));
 
-	TypedElementRowHandle TargetRow = DataStorage->GetColumn<FTypedElementRowReferenceColumn>(Row)->Row;
+	UE::Editor::DataStorage::RowHandle TargetRow = DataStorage->GetColumn<FTypedElementRowReferenceColumn>(Row)->Row;
 
 	if (const FTypedElementSlateWidgetReferenceColumn* SlateWidgetReferenceColumn = DataStorage->GetColumn<FTypedElementSlateWidgetReferenceColumn>(TargetRow))
 	{
diff --git a/Engine/Plugins/Experimental/EditorDataStorageFeatures/Source/TedsDebugger/Private/Widgets/SlateVisualizationWidget.h b/Engine/Plugins/Experimental/EditorDataStorageFeatures/Source/TedsDebugger/Private/Widgets/SlateVisualizationWidget.h
--- a/Engine/Plugins/Experimental/EditorDataStorageFeatures/Source/TedsDebugger/Private/Widgets/SlateVisualizationWidget.h
+++ b/Engine/Plugins/Experimental/EditorDataStorageFeatures/Source/TedsDebugger/Private/Widgets/SlateVisualizationWidget.h
@@ -37,5 +37,5 @@ public:
 protected:
 	TSharedPtr<SWidget> CreateWidget(const TypedElementDataStorage::FMetaDataView& Arguments) override;
 	bool FinalizeWidget(ITypedElementDataStorageInterface* DataStorage, ITypedElementDataStorageUiInterface* DataStorageUi,
-		TypedElementRowHandle Row, const TSharedPtr<SWidget>& Widget) override;
+		UE::Editor::DataStorage::RowHandle Row, const TSharedPtr<SWidget>& Widget) override;
 };
\ No newline at end of file
diff --git a/Engine/Plugins/Experimental/EditorDataStorageFeatures/Source/TedsOutliner/Private/Compatibility/SceneOutlinerRowHandleColumn.cpp b/Engine/Plugins/Experimental/EditorDataStorageFeatures/Source/TedsOutliner/Private/Compatibility/SceneOutlinerRowHandleColumn.cpp
--- a/Engine/Plugins/Experimental/EditorDataStorageFeatures/Source/TedsOutliner/Private/Compatibility/SceneOutlinerRowHandleColumn.cpp
+++ b/Engine/Plugins/Experimental/EditorDataStorageFeatures/Source/TedsOutliner/Private/Compatibility/SceneOutlinerRowHandleColumn.cpp
@@ -15,14 +15,14 @@ FSceneOutlinerRowHandleColumn::FSceneOutlinerRowHandleColumn(ISceneOutliner& Sce
 	auto AssignWidgetToColumn = [this](TUniquePtr<FTypedElementWidgetConstructor> Constructor, TConstArrayView<TWeakObjectPtr<const UScriptStruct>>)
 	{
 		TSharedPtr<FTypedElementWidgetConstructor> WidgetConstructor(Constructor.Release());
-		TableViewerColumn = MakeShared<UE::EditorDataStorage::FTedsTableViewerColumn>(TEXT("Row Handle"), WidgetConstructor);
+		TableViewerColumn = MakeShared<UE::Editor::DataStorage::FTedsTableViewerColumn>(TEXT("Row Handle"), WidgetConstructor);
 		return false;
 	};
 	
 	UTypedElementRegistry* Registry = UTypedElementRegistry::GetInstance();
-	checkf(Registry, TEXT("FTedsOutlinerRowHandleColumn created before UTypedElementRegistry is available."));
+	checkf(Registry, TEXT("FSceneOutlinerRowHandleColumn created before UTypedElementRegistry is available."));
 	ITypedElementDataStorageUiInterface* StorageUi = Registry->GetMutableDataStorageUi();
-	checkf(StorageUi, TEXT("FTedsOutlinerRowHandleColumn created before data storage interfaces were initialized."))
+	checkf(StorageUi, TEXT("FSceneOutlinerRowHandleColumn created before data storage interfaces were initialized."))
 
 	StorageUi->CreateWidgetConstructors(TEXT("General.Cell.RowHandle"), TypedElementDataStorage::FMetaDataView(), AssignWidgetToColumn);
 }
@@ -53,7 +53,7 @@ const TSharedRef<SWidget> FSceneOutlinerRowHandleColumn::ConstructRowWidget(FSce
 
 	if (const FTedsOutlinerTreeItem* OutlinerTreeItem = TreeItem->CastTo<FTedsOutlinerTreeItem>())
 	{
-		const TypedElementDataStorage::RowHandle RowHandle = OutlinerTreeItem->GetRowHandle();
+		const UE::Editor::DataStorage::RowHandle RowHandle = OutlinerTreeItem->GetRowHandle();
 
 		if(TSharedPtr<SWidget> Widget = TableViewerColumn->ConstructRowWidget(RowHandle))
 		{
@@ -74,7 +74,7 @@ void FSceneOutlinerRowHandleColumn::PopulateSearchStrings(const ISceneOutlinerTr
 
 void FSceneOutlinerRowHandleColumn::SortItems(TArray<FSceneOutlinerTreeItemPtr>& OutItems, const EColumnSortMode::Type SortMode) const
 {
-	FSceneOutlinerSortHelper<TypedElementDataStorage::RowHandle>()
+	FSceneOutlinerSortHelper<UE::Editor::DataStorage::RowHandle>()
 		/** Sort by type first */
 		.Primary([this](const ISceneOutlinerTreeItem& Item)
 		{
@@ -83,7 +83,7 @@ void FSceneOutlinerRowHandleColumn::SortItems(TArray<FSceneOutlinerTreeItemPtr>&
 				return OutlinerTreeItem->GetRowHandle();
 			}
 
-			return TypedElementDataStorage::InvalidRowHandle;
+			return UE::Editor::DataStorage::InvalidRowHandle;
 		}, SortMode)
 		.Sort(OutItems);
 }
diff --git a/Engine/Plugins/Experimental/EditorDataStorageFeatures/Source/TedsOutliner/Private/Compatibility/SceneOutlinerRowHandleColumn.h b/Engine/Plugins/Experimental/EditorDataStorageFeatures/Source/TedsOutliner/Private/Compatibility/SceneOutlinerRowHandleColumn.h
--- a/Engine/Plugins/Experimental/EditorDataStorageFeatures/Source/TedsOutliner/Private/Compatibility/SceneOutlinerRowHandleColumn.h
+++ b/Engine/Plugins/Experimental/EditorDataStorageFeatures/Source/TedsOutliner/Private/Compatibility/SceneOutlinerRowHandleColumn.h
@@ -6,7 +6,7 @@
 #include "ISceneOutliner.h"
 #include "ISceneOutlinerColumn.h"
 
-namespace UE::EditorDataStorage
+namespace UE::Editor::DataStorage
 {
 	class FTedsTableViewerColumn;
 }
@@ -39,5 +39,5 @@ private:
 	TWeakPtr<ISceneOutliner> WeakSceneOutliner;
 
 	// The Table Viewer column we are going to use internally to create the widget
-	TSharedPtr<UE::EditorDataStorage::FTedsTableViewerColumn> TableViewerColumn;
+	TSharedPtr<UE::Editor::DataStorage::FTedsTableViewerColumn> TableViewerColumn;
 };
diff --git a/Engine/Plugins/Experimental/EditorDataStorageFeatures/Source/TedsOutliner/Private/Compatibility/SceneOutlinerTedsBridge.cpp b/Engine/Plugins/Experimental/EditorDataStorageFeatures/Source/TedsOutliner/Private/Compatibility/SceneOutlinerTedsBridge.cpp
--- a/Engine/Plugins/Experimental/EditorDataStorageFeatures/Source/TedsOutliner/Private/Compatibility/SceneOutlinerTedsBridge.cpp
+++ b/Engine/Plugins/Experimental/EditorDataStorageFeatures/Source/TedsOutliner/Private/Compatibility/SceneOutlinerTedsBridge.cpp
@@ -42,7 +42,7 @@ FAutoConsoleCommand BindColumnsToSceneOutlinerConsoleCommand(
 			UTypedElementRegistry* Registry = UTypedElementRegistry::GetInstance();
 			if (ITypedElementDataStorageInterface* DataStorage = Registry->GetMutableDataStorage())
 			{
-				static TypedElementQueryHandle Queries[] =
+				static UE::Editor::DataStorage::QueryHandle Queries[] =
 				{
 					DataStorage->RegisterQuery(Select().ReadWrite<FTypedElementLabelColumn>().Compile()),
 					DataStorage->RegisterQuery(Select().ReadOnly<FTypedElementLocalTransformColumn>().Compile()),
@@ -69,7 +69,7 @@ FAutoConsoleCommand BindColumnsToSceneOutlinerConsoleCommand(
 						if (Args[0].IsNumeric())
 						{
 							int32 QueryIndex = FCString::Atoi(*Args[0]);
-							if (QueryIndex < sizeof(Queries) / sizeof(TypedElementQueryHandle))
+							if (QueryIndex < sizeof(Queries) / sizeof(UE::Editor::DataStorage::QueryHandle))
 							{
 								Binder.AssignQuery(Queries[QueryIndex], SceneOutliner, WidgetPurposes);
 								return;
@@ -96,8 +96,8 @@ FAutoConsoleCommand BindColumnsToSceneOutlinerConsoleCommand(
 							}
 							if (AdditionCount > 0)
 							{
-								static TypedElementQueryHandle CustomQuery = InvalidQueryHandle;
-								if (CustomQuery != InvalidQueryHandle)
+								static UE::Editor::DataStorage::QueryHandle CustomQuery = UE::Editor::DataStorage::InvalidQueryHandle;
+								if (CustomQuery != UE::Editor::DataStorage::InvalidQueryHandle)
 								{
 									DataStorage->UnregisterQuery(CustomQuery);
 								}
@@ -107,7 +107,7 @@ FAutoConsoleCommand BindColumnsToSceneOutlinerConsoleCommand(
 							}
 						}
 					}
-					Binder.AssignQuery(TypedElementInvalidQueryHandle, SceneOutliner, WidgetPurposes);
+					Binder.AssignQuery(UE::Editor::DataStorage::InvalidQueryHandle, SceneOutliner, WidgetPurposes);
 				}
 			}
 		}));
@@ -123,9 +123,10 @@ public:
 		ITypedElementDataStorageCompatibilityInterface& InStorageCompatibility,
 		const TSharedPtr<ISceneOutliner>& InOutliner);
 
-	void AssignQuery(TypedElementQueryHandle Query, const TConstArrayView<FName> CellWidgetPurposes);
+	void AssignQuery(UE::Editor::DataStorage::QueryHandle Query, const TConstArrayView<FName> CellWidgetPurposes);
 	void RegisterDealiaser(const FTreeItemIDDealiaser& InDealiaser);
 	FTreeItemIDDealiaser GetDealiaser();
+	
 private:
 	void ClearColumns(ISceneOutliner& InOutliner);
 
@@ -165,11 +166,11 @@ public:
 
 		using namespace TypedElementDataStorage;
 
-		TableViewerColumnImpl = MakeUnique<UE::EditorDataStorage::FTedsTableViewerColumn>(NameId, InCellWidgetConstructor, InColumnTypes,
+		TableViewerColumnImpl = MakeUnique<UE::Editor::DataStorage::FTedsTableViewerColumn>(NameId, InCellWidgetConstructor, InColumnTypes,
 			InHeaderWidgetConstructor, FComboMetaDataView(FGenericMetaDataView(MetaData)).Next(FQueryMetaDataView(Storage.GetQueryDescription(QueryHandle))));
 
 		TableViewerColumnImpl->SetIsRowVisibleDelegate(
-		UE::EditorDataStorage::FTedsTableViewerColumn::FIsRowVisible::CreateRaw(this, &FOutlinerColumn::IsRowVisible)
+		UE::Editor::DataStorage::FTedsTableViewerColumn::FIsRowVisible::CreateRaw(this, &FOutlinerColumn::IsRowVisible)
 		);
 		
 		// Try to find a fallback column from the regular item, for handling cases like folders which are not in TEDS but want to use TEDS columns
@@ -308,7 +309,7 @@ public:
 	}
 
 	// The table viewer implementation that we internally use to create our widgets
-	TUniquePtr<UE::EditorDataStorage::FTedsTableViewerColumn> TableViewerColumnImpl;
+	TUniquePtr<UE::Editor::DataStorage::FTedsTableViewerColumn> TableViewerColumnImpl;
 	
 	ITypedElementDataStorageInterface& Storage;
 	ITypedElementDataStorageUiInterface& StorageUi;
@@ -421,7 +422,7 @@ TSharedPtr<FSceneOutlinerTedsBridge>* FSceneOutlinerTedsQueryBinder::FindQueryMa
 	return SceneOutliners.Find(Outliner);
 }
 
-void FSceneOutlinerTedsQueryBinder::AssignQuery(TypedElementQueryHandle Query, const TSharedPtr<ISceneOutliner>& Outliner, TConstArrayView<FName> CellWidgetPurposes)
+void FSceneOutlinerTedsQueryBinder::AssignQuery(UE::Editor::DataStorage::QueryHandle Query, const TSharedPtr<ISceneOutliner>& Outliner, TConstArrayView<FName> CellWidgetPurposes)
 {
 	CleanupStaleOutliners();
 
@@ -495,7 +496,7 @@ FTreeItemIDDealiaser FSceneOutlinerTedsBridge::GetDealiaser()
 }
 
 
-void FSceneOutlinerTedsBridge::AssignQuery(TypedElementQueryHandle Query, const TConstArrayView<FName> InCellWidgetPurposes)
+void FSceneOutlinerTedsBridge::AssignQuery(UE::Editor::DataStorage::QueryHandle Query, const TConstArrayView<FName> InCellWidgetPurposes)
 {
 	using MatchApproach = ITypedElementDataStorageUiInterface::EMatchApproach;
 	constexpr int32 DefaultPriorityIndex = 100;
@@ -515,7 +516,7 @@ void FSceneOutlinerTedsBridge::AssignQuery(TypedElementQueryHandle Query, const
 			int32 SelectionCount = Description.SelectionTypes.Num();
 			AddedColumns.Reset(SelectionCount);
 
-			TArray<TWeakObjectPtr<const UScriptStruct>> ColumnTypes = UE::EditorDataStorage::TableViewerUtils::CreateVerifiedColumnTypeArray(Description.SelectionTypes);
+			TArray<TWeakObjectPtr<const UScriptStruct>> ColumnTypes = UE::Editor::DataStorage::TableViewerUtils::CreateVerifiedColumnTypeArray(Description.SelectionTypes);
 
 			int32 IndexOffset = 0;
 			auto ColumnConstructor = [this, Query, MetaDataView, &IndexOffset, &OutlinerPinned, Binder](
@@ -533,7 +534,7 @@ void FSceneOutlinerTedsBridge::AssignQuery(TypedElementQueryHandle Query, const
 
 					OutlinerPinned->RemoveColumn(FallbackColumn);
 
-					FName NameId = UE::EditorDataStorage::TableViewerUtils::FindLongestMatchingName(ColumnTypes, IndexOffset);
+					FName NameId = UE::Editor::DataStorage::TableViewerUtils::FindLongestMatchingName(ColumnTypes, IndexOffset);
 					AddedColumns.Add(NameId);
 					OutlinerPinned->AddColumn(NameId,
 						FSceneOutlinerColumnInfo(ESceneOutlinerColumnVisibility::Visible, ColumnPriority,
@@ -541,7 +542,7 @@ void FSceneOutlinerTedsBridge::AssignQuery(TypedElementQueryHandle Query, const
 								[this, Query, MetaDataView, NameId, &ColumnTypes, CellConstructor, &OutlinerPinned, FallbackColumn](ISceneOutliner&)
 								{
 									TSharedPtr<FTypedElementWidgetConstructor> HeaderConstructor = 
-										UE::EditorDataStorage::TableViewerUtils::CreateHeaderWidgetConstructor(*StorageUi, MetaDataView, ColumnTypes, CellWidgetPurposes);
+										UE::Editor::DataStorage::TableViewerUtils::CreateHeaderWidgetConstructor(*StorageUi, MetaDataView, ColumnTypes, CellWidgetPurposes);
 									return MakeShared<FOutlinerColumn>(
 										Query, *Storage, *StorageUi, *StorageCompatibility, NameId,
 										TArray<TWeakObjectPtr<const UScriptStruct>>(ColumnTypes.GetData(), ColumnTypes.Num()), 
@@ -578,7 +579,7 @@ void FSceneOutlinerTedsBridge::AssignQuery(TypedElementQueryHandle Query, const
 									TArray<TWeakObjectPtr<const UScriptStruct>> ColumnTypesStored;
 									ColumnTypesStored.Add(ColumnType);
 									TSharedPtr<FTypedElementWidgetConstructor> HeaderConstructor =
-										UE::EditorDataStorage::TableViewerUtils::CreateHeaderWidgetConstructor(*StorageUi, MetaDataView, { ColumnType }, CellWidgetPurposes);
+										UE::Editor::DataStorage::TableViewerUtils::CreateHeaderWidgetConstructor(*StorageUi, MetaDataView, { ColumnType }, CellWidgetPurposes);
 									return MakeShared<FOutlinerColumn>(
 										Query, *Storage, *StorageUi, *StorageCompatibility, NameId, MoveTemp(ColumnTypesStored),
 										HeaderConstructor, CellConstructor, FallbackColumn, OutlinerPinned, Dealiaser);
diff --git a/Engine/Plugins/Experimental/EditorDataStorageFeatures/Source/TedsOutliner/Private/TedsOutlinerHierarchy.cpp b/Engine/Plugins/Experimental/EditorDataStorageFeatures/Source/TedsOutliner/Private/TedsOutlinerHierarchy.cpp
--- a/Engine/Plugins/Experimental/EditorDataStorageFeatures/Source/TedsOutliner/Private/TedsOutlinerHierarchy.cpp
+++ b/Engine/Plugins/Experimental/EditorDataStorageFeatures/Source/TedsOutliner/Private/TedsOutlinerHierarchy.cpp
@@ -55,7 +55,7 @@ FSceneOutlinerTreeItemPtr FTedsOutlinerHierarchy::FindOrCreateParentItem(const I
 		return nullptr;
 	}
 	
-	const TypedElementRowHandle ParentRowHandle = TedsOutlinerImpl->GetParentRow(TedsTreeItem->GetRowHandle());
+	const UE::Editor::DataStorage::RowHandle ParentRowHandle = TedsOutlinerImpl->GetParentRow(TedsTreeItem->GetRowHandle());
 
 	if(!Storage->IsRowAvailable(ParentRowHandle))
 	{
