[Teds] Update naming for RowHandle, QueryHandle, and TableHandle

diff --git a/Engine/Plugins/Experimental/EditorDataStorageFeatures/Source/TedsOutliner/Private/TedsOutlinerImpl.cpp b/Engine/Plugins/Experimental/EditorDataStorageFeatures/Source/TedsOutliner/Private/TedsOutlinerImpl.cpp
--- a/Engine/Plugins/Experimental/EditorDataStorageFeatures/Source/TedsOutliner/Private/TedsOutlinerImpl.cpp
+++ b/Engine/Plugins/Experimental/EditorDataStorageFeatures/Source/TedsOutliner/Private/TedsOutlinerImpl.cpp
@@ -217,9 +217,9 @@ void FTedsOutlinerImpl::SetSelection(const TArray<TypedElementDataStorage::RowHa
 	}
 }
 
-TSharedRef<SWidget> FTedsOutlinerImpl::CreateLabelWidgetForItem(TypedElementRowHandle InRowHandle, const STableRow<FSceneOutlinerTreeItemPtr>& InRow) const
+TSharedRef<SWidget> FTedsOutlinerImpl::CreateLabelWidgetForItem(UE::Editor::DataStorage::RowHandle InRowHandle, const STableRow<FSceneOutlinerTreeItemPtr>& InRow) const
 {
-	auto CreateWidgetForQuery = [InRowHandle, this, &InRow](const TPair<TypedElementDataStorage::QueryHandle, TSharedPtr<FTypedElementWidgetConstructor>>& QueryConstructorPair) -> TSharedPtr<SWidget>
+	auto CreateWidgetForQuery = [InRowHandle, this, &InRow](const TPair<UE::Editor::DataStorage::QueryHandle, TSharedPtr<FTypedElementWidgetConstructor>>& QueryConstructorPair) -> TSharedPtr<SWidget>
 	{
 		TypedElementDataStorage::FQueryDescription QueryDescription = Storage->GetQueryDescription(QueryConstructorPair.Key);
 		
@@ -237,7 +237,7 @@ TSharedRef<SWidget> FTedsOutlinerImpl::CreateLabelWidgetForItem(TypedElementRowH
 		TArray<TWeakObjectPtr<const UScriptStruct>> ColumnTypes(QueryDescription.SelectionTypes);
 		TSharedPtr<FTypedElementWidgetConstructor> CellWidgetConstructor = QueryConstructorPair.Value;
 
-		TypedElementRowHandle UiRowHandle = Storage->AddRow(Storage->FindTable(UE::EditorDataStorage::TableViewerUtils::GetWidgetTableName()));
+		UE::Editor::DataStorage::RowHandle UiRowHandle = Storage->AddRow(Storage->FindTable(UE::Editor::DataStorage::TableViewerUtils::GetWidgetTableName()));
 
 		if (FTypedElementRowReferenceColumn* RowReference = Storage->GetColumn<FTypedElementRowReferenceColumn>(UiRowHandle))
 		{
@@ -382,13 +382,13 @@ void FTedsOutlinerImpl::CreateItemsFromQuery(TArray<FSceneOutlinerTreeItemPtr>&
 	TypedElementDataStorage::DirectQueryCallback RowCollector = CreateDirectQueryCallbackBinding(
 		[&Rows](DSI::IDirectQueryContext& Context)
 		{
-			TConstArrayView<TypedElementRowHandle> ContextRows = Context.GetRowHandles();
+			TConstArrayView<UE::Editor::DataStorage::RowHandle> ContextRows = Context.GetRowHandles();
 			Rows.Append(ContextRows);
 		});
 
 	Storage->RunQuery(RowHandleQuery, RowCollector);
 	
-	for (const TypedElementRowHandle& Row : Rows)
+	for (const UE::Editor::DataStorage::RowHandle& Row : Rows)
 	{
 		if (!CanDisplayRow(Row))
 		{
@@ -430,7 +430,7 @@ void FTedsOutlinerImpl::CreateChildren(const FSceneOutlinerTreeItemPtr& Item, TA
 		return;
 	}
 		
-	TypedElementRowHandle ItemRowHandle = TedsTreeItem->GetRowHandle();
+	UE::Editor::DataStorage::RowHandle ItemRowHandle = TedsTreeItem->GetRowHandle();
 
 	if(!Storage->IsRowAssigned(ItemRowHandle))
 	{
@@ -438,7 +438,7 @@ void FTedsOutlinerImpl::CreateChildren(const FSceneOutlinerTreeItemPtr& Item, TA
 	}
 
 
-	TSet<TypedElementDataStorage::RowHandle> MatchedRowsWithParentColumn;
+	TSet<UE::Editor::DataStorage::RowHandle> MatchedRowsWithParentColumn;
 	
 	// Collect all entities that are owned by our entity
 	TypedElementDataStorage::DirectQueryCallback ChildRowCollector = CreateDirectQueryCallbackBinding(
@@ -449,20 +449,20 @@ void FTedsOutlinerImpl::CreateChildren(const FSceneOutlinerTreeItemPtr& Item, TA
 
 	Storage->RunQuery(ChildRowHandleQuery, ChildRowCollector);
 
-	TArray<TypedElementDataStorage::RowHandle> ChildItems;
+	TArray<UE::Editor::DataStorage::RowHandle> ChildItems;
 
 	// Recursively get the children for each entity
-	TFunction<void(TypedElementRowHandle)> GetChildrenRecursive = [&ChildItems, &MatchedRowsWithParentColumn, DataStorage = Storage, &GetChildrenRecursive, InHierarchyData = HierarchyData]
-	(TypedElementRowHandle EntityRowHandle) -> void
+	TFunction<void(UE::Editor::DataStorage::RowHandle)> GetChildrenRecursive = [&ChildItems, &MatchedRowsWithParentColumn, DataStorage = Storage, &GetChildrenRecursive, InHierarchyData = HierarchyData]
+	(UE::Editor::DataStorage::RowHandle EntityRowHandle) -> void
 	{
-		for(TypedElementRowHandle ChildEntityRowHandle : MatchedRowsWithParentColumn)
+		for(UE::Editor::DataStorage::RowHandle ChildEntityRowHandle : MatchedRowsWithParentColumn)
 		{
 			void* ParentColumnData = DataStorage->GetColumnData(ChildEntityRowHandle, InHierarchyData.GetValue().HierarchyColumn);
 
 			if (ensureMsgf(ParentColumnData, TEXT("We should always the a parent column since we only grabbed rows with those ")))
 			{
 				// Get the parent row handle
-				const TypedElementDataStorage::RowHandle ParentRowHandle = InHierarchyData.GetValue().GetParent.Execute(ParentColumnData);
+				const UE::Editor::DataStorage::RowHandle ParentRowHandle = InHierarchyData.GetValue().GetParent.Execute(ParentColumnData);
 				
 				// Check if this entity is owned by the entity we are looking children for
 				if (ParentRowHandle == EntityRowHandle)
@@ -481,7 +481,7 @@ void FTedsOutlinerImpl::CreateChildren(const FSceneOutlinerTreeItemPtr& Item, TA
 	GetChildrenRecursive(ItemRowHandle);
 
 	// Actually create the items for the child entities 
-	for (TypedElementRowHandle ChildItemRowHandle : ChildItems)
+	for (UE::Editor::DataStorage::RowHandle ChildItemRowHandle : ChildItems)
 	{
 		if (!CanDisplayRow(ChildItemRowHandle))
 		{
@@ -495,12 +495,12 @@ void FTedsOutlinerImpl::CreateChildren(const FSceneOutlinerTreeItemPtr& Item, TA
 	}
 }
 
-TypedElementDataStorage::RowHandle FTedsOutlinerImpl::GetParentRow(TypedElementDataStorage::RowHandle InRowHandle)
+UE::Editor::DataStorage::RowHandle FTedsOutlinerImpl::GetParentRow(UE::Editor::DataStorage::RowHandle InRowHandle)
 {
 	// No parent if there is no hierarchy data specified
 	if (!HierarchyData.IsSet())
 	{
-		return TypedElementDataStorage::InvalidRowHandle;
+		return UE::Editor::DataStorage::InvalidRowHandle;
 	}
 	
 	// If this entity does not have a parent entity, return InvalidRowHandle
@@ -508,26 +508,26 @@ TypedElementDataStorage::RowHandle FTedsOutlinerImpl::GetParentRow(TypedElementD
 	
 	if (!ParentColumnData)
 	{
-		return TypedElementDataStorage::InvalidRowHandle;
+		return UE::Editor::DataStorage::InvalidRowHandle;
 	}
 
 	// If the parent is invalid for some reason, return InvalidRowHandle
-	const TypedElementRowHandle ParentRowHandle = HierarchyData.GetValue().GetParent.Execute(ParentColumnData);
+	const UE::Editor::DataStorage::RowHandle ParentRowHandle = HierarchyData.GetValue().GetParent.Execute(ParentColumnData);
 	
 	if (!Storage->IsRowAvailable(ParentRowHandle))
 	{
-		return TypedElementDataStorage::InvalidRowHandle;
+		return UE::Editor::DataStorage::InvalidRowHandle;
 	}
 	
 	if (!CanDisplayRow(ParentRowHandle))
 	{
-		return TypedElementDataStorage::InvalidRowHandle;
+		return UE::Editor::DataStorage::InvalidRowHandle;
 	}
 
 	return ParentRowHandle;
 }
 
-void FTedsOutlinerImpl::OnItemAdded(TypedElementDataStorage::RowHandle ItemRowHandle)
+void FTedsOutlinerImpl::OnItemAdded(UE::Editor::DataStorage::RowHandle ItemRowHandle)
 {
 	if (!CanDisplayRow(ItemRowHandle))
 	{
@@ -540,7 +540,7 @@ void FTedsOutlinerImpl::OnItemAdded(TypedElementDataStorage::RowHandle ItemRowHa
 	HierarchyChangedEvent.Broadcast(EventData);
 }
 
-void FTedsOutlinerImpl::OnItemRemoved(TypedElementDataStorage::RowHandle ItemRowHandle)
+void FTedsOutlinerImpl::OnItemRemoved(UE::Editor::DataStorage::RowHandle ItemRowHandle)
 {
 	FSceneOutlinerHierarchyChangedData EventData;
 	EventData.Type = FSceneOutlinerHierarchyChangedData::Removed;
@@ -548,7 +548,7 @@ void FTedsOutlinerImpl::OnItemRemoved(TypedElementDataStorage::RowHandle ItemRow
 	HierarchyChangedEvent.Broadcast(EventData);
 }
 
-void FTedsOutlinerImpl::OnItemMoved(TypedElementDataStorage::RowHandle ItemRowHandle)
+void FTedsOutlinerImpl::OnItemMoved(UE::Editor::DataStorage::RowHandle ItemRowHandle)
 {
 	if (!CanDisplayRow(ItemRowHandle))
 	{
@@ -584,7 +584,7 @@ void FTedsOutlinerImpl::RecompileQueries()
 		Select(
 				TEXT("Add Row to Outliner"),
 				FObserver::OnAdd<FTypedElementLabelColumn>().SetExecutionMode(EExecutionMode::GameThread),
-				[this](IQueryContext& Context, TypedElementRowHandle Row)
+				[this](IQueryContext& Context, UE::Editor::DataStorage::RowHandle Row)
 				{
 					OnItemAdded(Row);
 				})
@@ -598,7 +598,7 @@ void FTedsOutlinerImpl::RecompileQueries()
 		Select(
 				TEXT("Remove Row from Outliner"),
 				FObserver::OnRemove<FTypedElementLabelColumn>().SetExecutionMode(EExecutionMode::GameThread),
-				[this](IQueryContext& Context, TypedElementRowHandle Row)
+				[this](IQueryContext& Context, UE::Editor::DataStorage::RowHandle Row)
 				{
 					OnItemRemoved(Row);
 				})
@@ -667,7 +667,7 @@ void FTedsOutlinerImpl::RecompileQueries()
 							Select(
 							TEXT("Row selected"),
 							FObserver::OnAdd<FTypedElementSelectionColumn>().SetExecutionMode(EExecutionMode::GameThread),
-							[this](IQueryContext& Context, TypedElementRowHandle Row)
+							[this](IQueryContext& Context, UE::Editor::DataStorage::RowHandle Row)
 							{
 								bSelectionDirty = true;
 							})
@@ -681,7 +681,7 @@ void FTedsOutlinerImpl::RecompileQueries()
 							Select(
 							TEXT("Row deselected"),
 							FObserver::OnRemove<FTypedElementSelectionColumn>().SetExecutionMode(EExecutionMode::GameThread),
-							[this](IQueryContext& Context, TypedElementRowHandle Row)
+							[this](IQueryContext& Context, UE::Editor::DataStorage::RowHandle Row)
 							{
 								bSelectionDirty = true;
 							})
diff --git a/Engine/Plugins/Experimental/EditorDataStorageFeatures/Source/TedsOutliner/Private/TedsOutlinerItem.cpp b/Engine/Plugins/Experimental/EditorDataStorageFeatures/Source/TedsOutliner/Private/TedsOutlinerItem.cpp
--- a/Engine/Plugins/Experimental/EditorDataStorageFeatures/Source/TedsOutliner/Private/TedsOutlinerItem.cpp
+++ b/Engine/Plugins/Experimental/EditorDataStorageFeatures/Source/TedsOutliner/Private/TedsOutlinerItem.cpp
@@ -10,7 +10,7 @@
 
 const FSceneOutlinerTreeItemType FTedsOutlinerTreeItem::Type(&ISceneOutlinerTreeItem::Type);
 
-FTedsOutlinerTreeItem::FTedsOutlinerTreeItem(const TypedElementRowHandle& InRowHandle,
+FTedsOutlinerTreeItem::FTedsOutlinerTreeItem(const UE::Editor::DataStorage::RowHandle& InRowHandle,
 	const TSharedRef<const FTedsOutlinerImpl>& InTedsOutlinerImpl)
 	: ISceneOutlinerTreeItem(Type)
 	, RowHandle(InRowHandle)
@@ -65,7 +65,7 @@ void FTedsOutlinerTreeItem::GenerateContextMenu(UToolMenu* Menu, SSceneOutliner&
 	);
 }
 
-TypedElementRowHandle FTedsOutlinerTreeItem::GetRowHandle() const
+UE::Editor::DataStorage::RowHandle FTedsOutlinerTreeItem::GetRowHandle() const
 {
 	return RowHandle;
 }
diff --git a/Engine/Plugins/Experimental/EditorDataStorageFeatures/Source/TedsOutliner/Private/Widgets/AlertWidget.cpp b/Engine/Plugins/Experimental/EditorDataStorageFeatures/Source/TedsOutliner/Private/Widgets/AlertWidget.cpp
--- a/Engine/Plugins/Experimental/EditorDataStorageFeatures/Source/TedsOutliner/Private/Widgets/AlertWidget.cpp
+++ b/Engine/Plugins/Experimental/EditorDataStorageFeatures/Source/TedsOutliner/Private/Widgets/AlertWidget.cpp
@@ -14,7 +14,7 @@
 
 #define LOCTEXT_NAMESPACE "TedsAlertWidget"
 
-namespace UE::EditorDataStorage::Widgets::Private
+namespace UE::Editor::DataStorage::Widgets::Private
 {
 	void UpdateWidget(const TSharedPtr<SWidget>& Widget, const FText& Alert, bool bIsWarning, uint16 ErrorCount, uint16 WarningCount,
 		TypedElementDataStorage::RowHandle RowWithAlertAction)
@@ -117,7 +117,7 @@ namespace UE::EditorDataStorage::Widgets::Private
 				}
 
 				// If there's an action to call, enable the invisible button, otherwise turn it off.
-				if (RowWithAlertAction != TypedElementDataStorage::InvalidRowHandle)
+				if (RowWithAlertAction != UE::Editor::DataStorage::InvalidRowHandle)
 				{
 					Background.SetVisibility(EVisibility::HitTestInvisible);
 					ActionButton.SetVisibility(EVisibility::Visible);
@@ -151,13 +151,11 @@ void UAlertWidgetFactory::RegisterWidgetConstructors(
 	ITypedElementDataStorageInterface& DataStorage,
 	ITypedElementDataStorageUiInterface& DataStorageUi) const
 {
-	using namespace TypedElementDataStorage;
-
 	DataStorageUi.RegisterWidgetFactory<FAlertWidgetConstructor>(FName(TEXT("General.Cell")),
-		FColumn<FTypedElementAlertColumn>() || FColumn<FTypedElementChildAlertColumn>());
+		TypedElementDataStorage::FColumn<FTypedElementAlertColumn>() || TypedElementDataStorage::FColumn<FTypedElementChildAlertColumn>());
 	
 	DataStorageUi.RegisterWidgetFactory<FAlertHeaderWidgetConstructor>(FName(TEXT("General.Header")),
-		FColumn<FTypedElementAlertColumn>() || FColumn<FTypedElementChildAlertColumn>());
+		TypedElementDataStorage::FColumn<FTypedElementAlertColumn>() || TypedElementDataStorage::FColumn<FTypedElementChildAlertColumn>());
 }
 
 void UAlertWidgetFactory::RegisterQueries(ITypedElementDataStorageInterface& DataStorage)
@@ -170,8 +168,9 @@ void UAlertWidgetFactory::RegisterAlertQueries(ITypedElementDataStorageInterface
 {
 	using namespace TypedElementQueryBuilder;
 	using namespace TypedElementDataStorage;
+	using namespace UE::Editor::DataStorage;
 	
-	TypedElementQueryHandle UpdateWidget_OnlyAlert = DataStorage.RegisterQuery(
+	QueryHandle UpdateWidget_OnlyAlert = DataStorage.RegisterQuery(
 		Select()
 			.ReadOnly<FTypedElementAlertColumn>()
 		.Where()
@@ -179,7 +178,7 @@ void UAlertWidgetFactory::RegisterAlertQueries(ITypedElementDataStorageInterface
 			.None<FTypedElementChildAlertColumn>()
 		.Compile());
 
-	TypedElementQueryHandle UpdateWidget_OnlyChildAlert = DataStorage.RegisterQuery(
+	QueryHandle UpdateWidget_OnlyChildAlert = DataStorage.RegisterQuery(
 		Select()
 			.ReadOnly<FTypedElementChildAlertColumn>()
 		.Where()
@@ -187,7 +186,7 @@ void UAlertWidgetFactory::RegisterAlertQueries(ITypedElementDataStorageInterface
 			.None<FTypedElementAlertColumn>()
 		.Compile());
 
-	TypedElementQueryHandle UpdateWidget_Both = DataStorage.RegisterQuery(
+	QueryHandle UpdateWidget_Both = DataStorage.RegisterQuery(
 		Select()
 			.ReadOnly<FTypedElementAlertColumn, FTypedElementChildAlertColumn>()
 		.Where()
@@ -209,14 +208,14 @@ void UAlertWidgetFactory::RegisterAlertQueries(ITypedElementDataStorageInterface
 							Alert.AlertType == FTypedElementAlertColumnType::Warning ||
 							Alert.AlertType == FTypedElementAlertColumnType::Error,
 							TEXT("Alert column has unsupported type %i"), static_cast<int>(Alert.AlertType));
-						UE::EditorDataStorage::Widgets::Private::UpdateWidget(Widget.Widget.Pin(), Alert.Message,
+						Widgets::Private::UpdateWidget(Widget.Widget.Pin(), Alert.Message,
 							Alert.AlertType == FTypedElementAlertColumnType::Warning, 0, 0,
 							Context.HasColumn<FTypedElementAlertActionColumn>() ? Row : InvalidRowHandle);
 					}));
 				Context.RunSubquery(1, ReferenceColumn.Row, CreateSubqueryCallbackBinding(
 					[&Widget](ISubqueryContext& Context, RowHandle Row, const FTypedElementChildAlertColumn& ChildAlert)
 					{
-						UE::EditorDataStorage::Widgets::Private::UpdateWidget(Widget.Widget.Pin(), FText::GetEmpty(), false,
+						Widgets::Private::UpdateWidget(Widget.Widget.Pin(), FText::GetEmpty(), false,
 							ChildAlert.Counts[static_cast<size_t>(FTypedElementAlertColumnType::Error)],
 							ChildAlert.Counts[static_cast<size_t>(FTypedElementAlertColumnType::Warning)],
 							Context.HasColumn<FTypedElementAlertActionColumn>() ? Row : InvalidRowHandle);
@@ -231,7 +230,7 @@ void UAlertWidgetFactory::RegisterAlertQueries(ITypedElementDataStorageInterface
 							Alert.AlertType == FTypedElementAlertColumnType::Warning ||
 							Alert.AlertType == FTypedElementAlertColumnType::Error,
 							TEXT("Alert column has unsupported type %i"), static_cast<int>(Alert.AlertType));
-						UE::EditorDataStorage::Widgets::Private::UpdateWidget(Widget.Widget.Pin(), Alert.Message,
+						Widgets::Private::UpdateWidget(Widget.Widget.Pin(), Alert.Message,
 							Alert.AlertType == FTypedElementAlertColumnType::Warning,
 							ChildAlert.Counts[static_cast<size_t>(FTypedElementAlertColumnType::Error)],
 							ChildAlert.Counts[static_cast<size_t>(FTypedElementAlertColumnType::Warning)],
@@ -252,8 +251,9 @@ void UAlertWidgetFactory::RegisterAlertHeaderQueries(ITypedElementDataStorageInt
 {
 	using namespace TypedElementQueryBuilder;
 	using namespace TypedElementDataStorage;
+	using namespace UE::Editor::DataStorage;
 	
-	TypedElementQueryHandle AlertCount = DataStorage.RegisterQuery(
+	QueryHandle AlertCount = DataStorage.RegisterQuery(
 		Count()
 		.Where()
 			.Any<FTypedElementAlertColumn>()
@@ -377,7 +377,7 @@ bool FAlertWidgetConstructor::FinalizeWidget(ITypedElementDataStorageInterface*
 	uint16 ErrorCount = ChildAlert ? ChildAlert->Counts[static_cast<size_t>(FTypedElementAlertColumnType::Error)] : 0;
 	uint16 WarningCount = ChildAlert ? ChildAlert->Counts[static_cast<size_t>(FTypedElementAlertColumnType::Warning)] : 0;
 
-	UE::EditorDataStorage::Widgets::Private::UpdateWidget(
+	UE::Editor::DataStorage::Widgets::Private::UpdateWidget(
 		Widget, 
 		Alert ? Alert->Message : FText::GetEmpty(), 
 		Alert ? (Alert->AlertType == FTypedElementAlertColumnType::Warning) : false, 
@@ -415,7 +415,7 @@ TConstArrayView<const UScriptStruct*> FAlertHeaderWidgetConstructor::GetAddition
 }
 
 bool FAlertHeaderWidgetConstructor::FinalizeWidget(ITypedElementDataStorageInterface* DataStorage, 
-	ITypedElementDataStorageUiInterface* DataStorageUi, TypedElementDataStorage::RowHandle Row, const TSharedPtr<SWidget>& Widget)
+	ITypedElementDataStorageUiInterface* DataStorageUi, UE::Editor::DataStorage::RowHandle Row, const TSharedPtr<SWidget>& Widget)
 {
 	DataStorage->AddColumn(Row, FUIHeaderPropertiesColumn
 		{
diff --git a/Engine/Plugins/Experimental/EditorDataStorageFeatures/Source/TedsOutliner/Private/Widgets/OverrideWidget.cpp b/Engine/Plugins/Experimental/EditorDataStorageFeatures/Source/TedsOutliner/Private/Widgets/OverrideWidget.cpp
--- a/Engine/Plugins/Experimental/EditorDataStorageFeatures/Source/TedsOutliner/Private/Widgets/OverrideWidget.cpp
+++ b/Engine/Plugins/Experimental/EditorDataStorageFeatures/Source/TedsOutliner/Private/Widgets/OverrideWidget.cpp
@@ -21,11 +21,9 @@
 void UOverrideWidgetFactory::RegisterWidgetConstructors(
 	ITypedElementDataStorageInterface& DataStorage, ITypedElementDataStorageUiInterface& DataStorageUi) const
 {
-	using namespace TypedElementDataStorage;
-
 	// The  widget is a specific widget for the Scene Outliner's item label column
 	DataStorageUi.RegisterWidgetFactory<FOverrideWidgetConstructor>(FSceneOutlinerTedsQueryBinder::ItemLabelCellWidgetPurpose,
-		FColumn<FTypedElementClassTypeInfoColumn>() || FColumn<FObjectOverrideColumn>() );
+		TypedElementDataStorage::FColumn<FTypedElementClassTypeInfoColumn>() || TypedElementDataStorage::FColumn<FObjectOverrideColumn>() );
 }
 
 void UOverrideWidgetFactory::RegisterQueries(ITypedElementDataStorageInterface& DataStorage)
@@ -33,7 +31,7 @@ void UOverrideWidgetFactory::RegisterQueries(ITypedElementDataStorageInterface&
 	using namespace TypedElementQueryBuilder;
 	using namespace TypedElementDataStorage;
 	
-	const TypedElementQueryHandle UpdateWidget = DataStorage.RegisterQuery(
+	const UE::Editor::DataStorage::QueryHandle UpdateWidget = DataStorage.RegisterQuery(
 		Select()
 		.Where()
 			.Any<FTypedElementSyncFromWorldTag, FTypedElementSyncBackToWorldTag>()
@@ -77,9 +75,9 @@ TSharedPtr<SWidget> FOverrideWidgetConstructor::CreateWidget(const TypedElementD
 }
 
 bool FOverrideWidgetConstructor::FinalizeWidget(ITypedElementDataStorageInterface* DataStorage,
-	ITypedElementDataStorageUiInterface* DataStorageUi, TypedElementRowHandle Row, const TSharedPtr<SWidget>& Widget)
+	ITypedElementDataStorageUiInterface* DataStorageUi, UE::Editor::DataStorage::RowHandle Row, const TSharedPtr<SWidget>& Widget)
 {
-	const TypedElementRowHandle TargetRow = DataStorage->GetColumn<FTypedElementRowReferenceColumn>(Row)->Row;
+	const UE::Editor::DataStorage::RowHandle TargetRow = DataStorage->GetColumn<FTypedElementRowReferenceColumn>(Row)->Row;
 
 	if (const FTypedElementClassTypeInfoColumn* TypeInfoColumn = DataStorage->GetColumn<FTypedElementClassTypeInfoColumn>(TargetRow))
 	{
@@ -165,7 +163,7 @@ void FOverrideWidgetConstructor::RemoveOverrideBadge(const TWeakPtr<SWidget>& Wi
 	}
 }
 
-void FOverrideWidgetConstructor::UpdateOverrideWidget(const TWeakPtr<SWidget>& Widget, const TypedElementRowHandle TargetRow)
+void FOverrideWidgetConstructor::UpdateOverrideWidget(const TWeakPtr<SWidget>& Widget, const UE::Editor::DataStorage::RowHandle TargetRow)
 {
 	const ITypedElementDataStorageInterface* DataStorageInterface = UTypedElementRegistry::GetInstance()->GetDataStorage();
 
diff --git a/Engine/Plugins/Experimental/EditorDataStorageFeatures/Source/TedsOutliner/Private/Widgets/OverrideWidget.h b/Engine/Plugins/Experimental/EditorDataStorageFeatures/Source/TedsOutliner/Private/Widgets/OverrideWidget.h
--- a/Engine/Plugins/Experimental/EditorDataStorageFeatures/Source/TedsOutliner/Private/Widgets/OverrideWidget.h
+++ b/Engine/Plugins/Experimental/EditorDataStorageFeatures/Source/TedsOutliner/Private/Widgets/OverrideWidget.h
@@ -48,10 +48,10 @@ public:
 	static void AddOverrideBadge(const TWeakPtr<SWidget>& Widget, EOverriddenState OverriddenState);
 	static void RemoveOverrideBadge(const TWeakPtr<SWidget>& Widget);
 
-	static void UpdateOverrideWidget(const TWeakPtr<SWidget>& Widget, const TypedElementRowHandle TargetRow);
+	static void UpdateOverrideWidget(const TWeakPtr<SWidget>& Widget, const UE::Editor::DataStorage::RowHandle TargetRow);
 
 protected:
 	TEDSOUTLINER_API TSharedPtr<SWidget> CreateWidget(const TypedElementDataStorage::FMetaDataView& Arguments) override;
 	TEDSOUTLINER_API bool FinalizeWidget(ITypedElementDataStorageInterface* DataStorage, ITypedElementDataStorageUiInterface* DataStorageUi,
-		TypedElementRowHandle Row, const TSharedPtr<SWidget>& Widget) override;
+		UE::Editor::DataStorage::RowHandle Row, const TSharedPtr<SWidget>& Widget) override;
 };
\ No newline at end of file
diff --git a/Engine/Plugins/Experimental/EditorDataStorageFeatures/Source/TedsOutliner/Private/Widgets/TransformHeadsUpWidget.cpp b/Engine/Plugins/Experimental/EditorDataStorageFeatures/Source/TedsOutliner/Private/Widgets/TransformHeadsUpWidget.cpp
--- a/Engine/Plugins/Experimental/EditorDataStorageFeatures/Source/TedsOutliner/Private/Widgets/TransformHeadsUpWidget.cpp
+++ b/Engine/Plugins/Experimental/EditorDataStorageFeatures/Source/TedsOutliner/Private/Widgets/TransformHeadsUpWidget.cpp
@@ -15,7 +15,7 @@
 
 #define LOCTEXT_NAMESPACE "TedsTransformHeadsUpWidget"
 
-namespace UE::EditorDataStorage::Widgets::Private
+namespace UE::Editor::DataStorage::Widgets::Private
 {
 	enum class EAbnormalTransformTypes : int32
 	{
@@ -56,7 +56,7 @@ namespace UE::EditorDataStorage::Widgets::Private
 
 		return AbnormalTransformFlags;
 	}
-} // namespace UE::EditorDataStorage::Widgets::Private
+} // namespace UE::Editor::DataStorage::Widgets::Private
 
 class STransformQuickDisplay : public SHorizontalBox
 {
@@ -68,7 +68,7 @@ public:
 	void Construct(const FArguments& InArgs)
 	{
 		const static FMargin IconPadding(1, 1, 0, 0);
-		using namespace UE::EditorDataStorage::Widgets::Private;
+		using namespace UE::Editor::DataStorage::Widgets::Private;
 
 		AddSlot().AutoWidth().Padding(IconPadding)
 		[
@@ -117,7 +117,7 @@ public:
 		];
 	}
 
-	void UpdateFromTransform(UE::EditorDataStorage::Widgets::Private::EAbnormalTransformTypes InAbnormalTransformFlags)
+	void UpdateFromTransform(UE::Editor::DataStorage::Widgets::Private::EAbnormalTransformTypes InAbnormalTransformFlags)
 	{
 		if (AbnormalTransformFlags != InAbnormalTransformFlags)
 		{
@@ -127,13 +127,13 @@ public:
 	}
 
 private:
-	UE::EditorDataStorage::Widgets::Private::EAbnormalTransformTypes AbnormalTransformFlags;
+	UE::Editor::DataStorage::Widgets::Private::EAbnormalTransformTypes AbnormalTransformFlags;
 };
 
 SLATE_IMPLEMENT_WIDGET(STransformQuickDisplay)
 void STransformQuickDisplay::PrivateRegisterAttributes(FSlateAttributeInitializer&) {}
 
-static void UpdateTransformHeadsUpDisplay(FTypedElementSlateWidgetReferenceColumn& Widget, UE::EditorDataStorage::Widgets::Private::EAbnormalTransformTypes AbnormalTransformFlags)
+static void UpdateTransformHeadsUpDisplay(FTypedElementSlateWidgetReferenceColumn& Widget, UE::Editor::DataStorage::Widgets::Private::EAbnormalTransformTypes AbnormalTransformFlags)
 {
 	TSharedPtr<SWidget> WidgetPointer = Widget.Widget.Pin();
 	checkf(WidgetPointer, TEXT("Referenced widget is not valid. A constructed widget may not have been cleaned up. This can "
@@ -157,7 +157,7 @@ void UTransformHeadsUpWidgetFactory::RegisterQueries(ITypedElementDataStorageInt
 	using namespace TypedElementQueryBuilder;
 	using namespace TypedElementDataStorage;
 		
-	TypedElementQueryHandle UpdateTransformWidget = DataStorage.RegisterQuery(
+	UE::Editor::DataStorage::QueryHandle UpdateTransformWidget = DataStorage.RegisterQuery(
 		Select()
 			.ReadOnly<FTypedElementLocalTransformColumn>()
 		.Where()
@@ -175,7 +175,7 @@ void UTransformHeadsUpWidgetFactory::RegisterQueries(ITypedElementDataStorageInt
 				Context.RunSubquery(0, ReferenceColumn.Row, CreateSubqueryCallbackBinding(
 					[&Widget](const FTypedElementLocalTransformColumn& Transform)
 					{
-						UpdateTransformHeadsUpDisplay(Widget, UE::EditorDataStorage::Widgets::Private::GetAbnormalTransformTypes(Transform.Transform));
+						UpdateTransformHeadsUpDisplay(Widget, UE::Editor::DataStorage::Widgets::Private::GetAbnormalTransformTypes(Transform.Transform));
 					}));
 			}
 		)
@@ -190,10 +190,8 @@ void UTransformHeadsUpWidgetFactory::RegisterQueries(ITypedElementDataStorageInt
 void UTransformHeadsUpWidgetFactory::RegisterWidgetConstructors(ITypedElementDataStorageInterface& DataStorage,
 	ITypedElementDataStorageUiInterface& DataStorageUi) const
 {
-	using namespace TypedElementDataStorage;
-
 	DataStorageUi.RegisterWidgetFactory<FTransformHeadsUpWidgetConstructor>(FName(TEXT("SceneOutliner.Cell")),
-		FColumn<FTypedElementLocalTransformColumn>());
+		TypedElementDataStorage::FColumn<FTypedElementLocalTransformColumn>());
 }
 
 
@@ -222,14 +220,14 @@ TSharedPtr<SWidget> FTransformHeadsUpWidgetConstructor::CreateWidget(const Typed
 bool FTransformHeadsUpWidgetConstructor::FinalizeWidget(
 	ITypedElementDataStorageInterface* DataStorage,
 	ITypedElementDataStorageUiInterface* DataStorageUi,
-	TypedElementRowHandle Row,
+	UE::Editor::DataStorage::RowHandle Row,
 	const TSharedPtr<SWidget>& Widget)
 {
 	FTypedElementRowReferenceColumn& RefColumn = *DataStorage->GetColumn<FTypedElementRowReferenceColumn>(Row);
 	if (const FTypedElementLocalTransformColumn* TransformColumn = DataStorage->GetColumn<FTypedElementLocalTransformColumn>(RefColumn.Row))
 	{
 		UpdateTransformHeadsUpDisplay(*DataStorage->GetColumn<FTypedElementSlateWidgetReferenceColumn>(Row), 
-			UE::EditorDataStorage::Widgets::Private::GetAbnormalTransformTypes(TransformColumn->Transform));
+			UE::Editor::DataStorage::Widgets::Private::GetAbnormalTransformTypes(TransformColumn->Transform));
 	}
 
 	return true;
diff --git a/Engine/Plugins/Experimental/EditorDataStorageFeatures/Source/TedsOutliner/Private/Widgets/TypeInfoWidget.cpp b/Engine/Plugins/Experimental/EditorDataStorageFeatures/Source/TedsOutliner/Private/Widgets/TypeInfoWidget.cpp
--- a/Engine/Plugins/Experimental/EditorDataStorageFeatures/Source/TedsOutliner/Private/Widgets/TypeInfoWidget.cpp
+++ b/Engine/Plugins/Experimental/EditorDataStorageFeatures/Source/TedsOutliner/Private/Widgets/TypeInfoWidget.cpp
@@ -62,13 +62,13 @@ TSharedPtr<SWidget> FTypeInfoWidgetConstructor::CreateWidget(
 }
 
 bool FTypeInfoWidgetConstructor::FinalizeWidget(ITypedElementDataStorageInterface* DataStorage,
-	ITypedElementDataStorageUiInterface* DataStorageUi, TypedElementRowHandle Row, const TSharedPtr<SWidget>& Widget)
+	ITypedElementDataStorageUiInterface* DataStorageUi, UE::Editor::DataStorage::RowHandle Row, const TSharedPtr<SWidget>& Widget)
 {
 	checkf(Widget, TEXT("Referenced widget is not valid. A constructed widget may not have been cleaned up. This can "
 		"also happen if this processor is running in the same phase as the processors responsible for cleaning up old "
 		"references."));
 
-	TypedElementRowHandle TargetRow = DataStorage->GetColumn<FTypedElementRowReferenceColumn>(Row)->Row;
+	UE::Editor::DataStorage::RowHandle TargetRow = DataStorage->GetColumn<FTypedElementRowReferenceColumn>(Row)->Row;
 
 	if (const FTypedElementClassTypeInfoColumn* TypeInfoColumn = DataStorage->GetColumn<FTypedElementClassTypeInfoColumn>(TargetRow))
 	{
@@ -124,7 +124,7 @@ bool FTypeInfoWidgetConstructor::FinalizeWidget(ITypedElementDataStorageInterfac
 	return true;
 }
 
-const FSlateBrush* FTypeInfoWidgetConstructor::GetIconForRow(ITypedElementDataStorageInterface* DataStorage, TypedElementRowHandle Row, const FTypedElementClassTypeInfoColumn* TypeInfoColumn)
+const FSlateBrush* FTypeInfoWidgetConstructor::GetIconForRow(ITypedElementDataStorageInterface* DataStorage, UE::Editor::DataStorage::RowHandle Row, const FTypedElementClassTypeInfoColumn* TypeInfoColumn)
 {
 	/* The logic here is very similar to SActorTreeLabel::GetIcon in ActorTreeItem.cpp which allows the actor to specify
 	 * an override for the icon, and has a fallback to the class icon if not.
diff --git a/Engine/Plugins/Experimental/EditorDataStorageFeatures/Source/TedsOutliner/Public/Compatibility/SceneOutlinerTedsBridge.h b/Engine/Plugins/Experimental/EditorDataStorageFeatures/Source/TedsOutliner/Public/Compatibility/SceneOutlinerTedsBridge.h
--- a/Engine/Plugins/Experimental/EditorDataStorageFeatures/Source/TedsOutliner/Public/Compatibility/SceneOutlinerTedsBridge.h
+++ b/Engine/Plugins/Experimental/EditorDataStorageFeatures/Source/TedsOutliner/Public/Compatibility/SceneOutlinerTedsBridge.h
@@ -39,7 +39,7 @@ public:
 
 	static FSceneOutlinerTedsQueryBinder& GetInstance();
 
-	void AssignQuery(TypedElementQueryHandle Query, const TSharedPtr<ISceneOutliner>& Widget, TConstArrayView<FName> InCellWidgetPurposes);
+	void AssignQuery(UE::Editor::DataStorage::QueryHandle Query, const TSharedPtr<ISceneOutliner>& Widget, TConstArrayView<FName> InCellWidgetPurposes);
 
 	// Register a dealiaser for a specific TEDS-Outliner to convert a row handle to an FSceneOutlinerTreeItemID
 	void RegisterTreeItemIDDealiaser(const TSharedPtr<ISceneOutliner>& Widget, const FTreeItemIDDealiaser& InDealiaser);
diff --git a/Engine/Plugins/Experimental/EditorDataStorageFeatures/Source/TedsOutliner/Public/TedsOutlinerImpl.h b/Engine/Plugins/Experimental/EditorDataStorageFeatures/Source/TedsOutliner/Public/TedsOutlinerImpl.h
--- a/Engine/Plugins/Experimental/EditorDataStorageFeatures/Source/TedsOutliner/Public/TedsOutlinerImpl.h
+++ b/Engine/Plugins/Experimental/EditorDataStorageFeatures/Source/TedsOutliner/Public/TedsOutlinerImpl.h
@@ -16,10 +16,10 @@ class SWidget;
 struct FTedsOutlinerHierarchyData
 {
 	/** A delegate used to get the parent row handle for a given row */
-	DECLARE_DELEGATE_RetVal_OneParam(TypedElementDataStorage::RowHandle, FGetParentRowHandle, void* /* InColumnData */);
+	DECLARE_DELEGATE_RetVal_OneParam(UE::Editor::DataStorage::RowHandle, FGetParentRowHandle, void* /* InColumnData */);
 	
 	/** A delegate used to set the parent row handle for a given row */
-	DECLARE_DELEGATE_TwoParams(FSetParentRowHandle, void* /* InColumnData */, TypedElementDataStorage::RowHandle /* InParentRowHandle */);
+	DECLARE_DELEGATE_TwoParams(FSetParentRowHandle, void* /* InColumnData */, UE::Editor::DataStorage::RowHandle /* InParentRowHandle */);
 
 	FTedsOutlinerHierarchyData(const UScriptStruct* InHierarchyColumn, const FGetParentRowHandle& InGetParent, const FSetParentRowHandle& InSetParent)
 		: HierarchyColumn(InHierarchyColumn)
@@ -135,7 +135,7 @@ public:
 	void SetSelection(const TArray<TypedElementDataStorage::RowHandle>& InSelectedRows);
 
 	// Helper function to create a label widget for a given row
-	TSharedRef<SWidget> CreateLabelWidgetForItem(TypedElementRowHandle InRowHandle, const STableRow<FSceneOutlinerTreeItemPtr>& InRow) const;
+	TSharedRef<SWidget> CreateLabelWidgetForItem(UE::Editor::DataStorage::RowHandle InRowHandle, const STableRow<FSceneOutlinerTreeItemPtr>& InRow) const;
 
 	// Get the hierarchy data associated with this table viewer
 	const TOptional<FTedsOutlinerHierarchyData>& GetHierarchyData();
@@ -152,23 +152,23 @@ public:
 	static void AppendQuery(TypedElementDataStorage::FQueryDescription& Query1, const TypedElementDataStorage::FQueryDescription& Query2);
 
 	// Check if the given item's parent has changed (i.e ParentRowHandle does not match what the Outliner reports as the parent)
-	bool HasItemParentChanged(TypedElementDataStorage::RowHandle ItemRowHandle, TypedElementDataStorage::RowHandle ParentRowHandle) const;
+	bool HasItemParentChanged(UE::Editor::DataStorage::RowHandle ItemRowHandle, UE::Editor::DataStorage::RowHandle ParentRowHandle) const;
 
 	// Outliner specific functionality
 	void CreateItemsFromQuery(TArray<FSceneOutlinerTreeItemPtr>& OutItems, ISceneOutlinerMode* InMode) const;
 	void CreateChildren(const FSceneOutlinerTreeItemPtr& Item, TArray<FSceneOutlinerTreeItemPtr>& OutChildren) const;
 
 	// Get the parent row for a given row
-	TypedElementDataStorage::RowHandle GetParentRow(TypedElementDataStorage::RowHandle InRowHandle);
+	UE::Editor::DataStorage::RowHandle GetParentRow(UE::Editor::DataStorage::RowHandle InRowHandle);
 
 	// Recompile all queries used by this table viewer
 	void RecompileQueries();
 
 protected:
 	
-	void OnItemAdded(TypedElementDataStorage::RowHandle ItemRowHandle);
-	void OnItemRemoved(TypedElementDataStorage::RowHandle ItemRowHandle);
-	void OnItemMoved(TypedElementDataStorage::RowHandle ItemRowHandle);
+	void OnItemAdded(UE::Editor::DataStorage::RowHandle ItemRowHandle);
+	void OnItemRemoved(UE::Editor::DataStorage::RowHandle ItemRowHandle);
+	void OnItemMoved(UE::Editor::DataStorage::RowHandle ItemRowHandle);
 
 	void UnregisterQueries() const;
 	void ClearSelection() const;
@@ -178,7 +178,7 @@ protected:
 	void CreateFilterQueries();
 
 	// Check if this row can be displayed in this table viewer
-	bool CanDisplayRow(TypedElementDataStorage::RowHandle ItemRowHandle) const;
+	bool CanDisplayRow(UE::Editor::DataStorage::RowHandle ItemRowHandle) const;
 	
 protected:
 	// TEDS Storage Constructs
@@ -189,7 +189,7 @@ protected:
 	FTedsOutlinerParams CreationParams;
 
 	// Widget constructor to create the label widget
-	TArray<TPair<TypedElementDataStorage::QueryHandle, TSharedPtr<FTypedElementWidgetConstructor>>> QueryToWidgetConstructorMap;
+	TArray<TPair<UE::Editor::DataStorage::QueryHandle, TSharedPtr<FTypedElementWidgetConstructor>>> QueryToWidgetConstructorMap;
 
 	// Widget purposes this table viewer supports
 	TArray<FName> CellWidgetPurposes;
@@ -204,20 +204,20 @@ protected:
 	TOptional<FTedsOutlinerHierarchyData> HierarchyData;
 
 	// Querys to track row handle collection, addition and removal
-	TypedElementDataStorage::QueryHandle RowHandleQuery = TypedElementDataStorage::InvalidQueryHandle;
-	TypedElementDataStorage::QueryHandle RowAdditionQuery = TypedElementDataStorage::InvalidQueryHandle;
-	TypedElementDataStorage::QueryHandle RowRemovalQuery = TypedElementDataStorage::InvalidQueryHandle;
+	UE::Editor::DataStorage::QueryHandle RowHandleQuery = UE::Editor::DataStorage::InvalidQueryHandle;
+	UE::Editor::DataStorage::QueryHandle RowAdditionQuery = UE::Editor::DataStorage::InvalidQueryHandle;
+	UE::Editor::DataStorage::QueryHandle RowRemovalQuery = UE::Editor::DataStorage::InvalidQueryHandle;
 
 	// Query to get all child rows
-	TypedElementDataStorage::QueryHandle ChildRowHandleQuery = TypedElementDataStorage::InvalidQueryHandle;
+	UE::Editor::DataStorage::QueryHandle ChildRowHandleQuery = UE::Editor::DataStorage::InvalidQueryHandle;
 
 	// Query to track when a row's parent gets changed
-	TypedElementDataStorage::QueryHandle UpdateParentQuery = TypedElementDataStorage::InvalidQueryHandle;
-	
+	UE::Editor::DataStorage::QueryHandle UpdateParentQuery = UE::Editor::DataStorage::InvalidQueryHandle;
+
 	// Query to get all selected rows, track selection added, track selection removed
-	TypedElementDataStorage::QueryHandle SelectedRowsQuery = TypedElementDataStorage::InvalidQueryHandle;
-	TypedElementDataStorage::QueryHandle SelectionAddedQuery = TypedElementDataStorage::InvalidQueryHandle;
-	TypedElementDataStorage::QueryHandle SelectionRemovedQuery = TypedElementDataStorage::InvalidQueryHandle;
+	UE::Editor::DataStorage::QueryHandle SelectedRowsQuery = UE::Editor::DataStorage::InvalidQueryHandle;
+	UE::Editor::DataStorage::QueryHandle SelectionAddedQuery = UE::Editor::DataStorage::InvalidQueryHandle;
+	UE::Editor::DataStorage::QueryHandle SelectionRemovedQuery = UE::Editor::DataStorage::InvalidQueryHandle;
 	
 	TOptional<FName> SelectionSetName;
 	bool bSelectionDirty = false;
diff --git a/Engine/Plugins/Experimental/EditorDataStorageFeatures/Source/TedsOutliner/Public/TedsOutlinerItem.h b/Engine/Plugins/Experimental/EditorDataStorageFeatures/Source/TedsOutliner/Public/TedsOutlinerItem.h
--- a/Engine/Plugins/Experimental/EditorDataStorageFeatures/Source/TedsOutliner/Public/TedsOutlinerItem.h
+++ b/Engine/Plugins/Experimental/EditorDataStorageFeatures/Source/TedsOutliner/Public/TedsOutlinerItem.h
@@ -8,7 +8,7 @@
 class FTedsOutlinerImpl;
 
 /*
- * A generic item in the TEDS driven Outliner, that uses a TypedElementRowHandle to uniquely identify the object it is
+ * A generic item in the TEDS driven Outliner, that uses a UE::Editor::DataStorage::RowHandle to uniquely identify the object it is
  * looking at. Functionality should be added through TEDS queries instead of having a different TreeItem type for each
  * type of object you are looking at (i.e. Actor vs Folder)
  * Inherits from ISceneOutlinerItem - which determines what type of item you are looking at. E.G FActorTreeItem for actors
@@ -24,7 +24,7 @@ public:
 		return Pred.Execute(RowHandle);
 	}
 
-	TEDSOUTLINER_API FTedsOutlinerTreeItem(const TypedElementRowHandle& InRowHandle, const TSharedRef<const FTedsOutlinerImpl>& InTedsOutlinerImpl);
+	TEDSOUTLINER_API FTedsOutlinerTreeItem(const UE::Editor::DataStorage::RowHandle& InRowHandle, const TSharedRef<const FTedsOutlinerImpl>& InTedsOutlinerImpl);
 
 	/* Begin ISceneOutlinerTreeItem Implementation */
 	TEDSOUTLINER_API virtual bool IsValid() const override;
@@ -38,9 +38,9 @@ public:
 
 	TEDSOUTLINER_API static const FSceneOutlinerTreeItemType Type;
 
-	TEDSOUTLINER_API TypedElementRowHandle GetRowHandle() const;
+	TEDSOUTLINER_API UE::Editor::DataStorage::RowHandle GetRowHandle() const;
 
 private:
-	const TypedElementRowHandle RowHandle;
+	const UE::Editor::DataStorage::RowHandle RowHandle;
 	const TSharedRef<const FTedsOutlinerImpl> TedsOutlinerImpl;
 };
\ No newline at end of file
diff --git a/Engine/Plugins/Experimental/EditorDataStorageFeatures/Source/TedsOutliner/Public/TedsOutlinerModule.h b/Engine/Plugins/Experimental/EditorDataStorageFeatures/Source/TedsOutliner/Public/TedsOutlinerModule.h
--- a/Engine/Plugins/Experimental/EditorDataStorageFeatures/Source/TedsOutliner/Public/TedsOutlinerModule.h
+++ b/Engine/Plugins/Experimental/EditorDataStorageFeatures/Source/TedsOutliner/Public/TedsOutlinerModule.h
@@ -34,14 +34,14 @@ public:
 	 * @return	New scene outliner widget
 	 */
 	virtual TSharedRef<ISceneOutliner> CreateTedsOutliner(
-		const FSceneOutlinerInitializationOptions& InInitOptions, const FTedsOutlinerParams& InInitTedsOptions, TypedElementDataStorage::QueryHandle ColumnQuery) const;
+		const FSceneOutlinerInitializationOptions& InInitOptions, const FTedsOutlinerParams& InInitTedsOptions, UE::Editor::DataStorage::QueryHandle ColumnQuery) const;
 
 	// IModuleInterface interface
 	virtual void StartupModule() override;
 	virtual void ShutdownModule() override;
 
 	// Get the column query the default table viewer uses
-	virtual TypedElementDataStorage::QueryHandle GetLevelEditorTedsOutlinerColumnQuery();
+	virtual UE::Editor::DataStorage::QueryHandle GetLevelEditorTedsOutlinerColumnQuery();
 
 	// The name of the tab the default table viewer is opened in
 	FName GetTedsOutlinerTabName();
@@ -56,4 +56,4 @@ private:
 private:
 	FName TedsOutlinerTabName;
 	FDelegateHandle LevelEditorTabManagerChangedHandle;
-};
\ No newline at end of file
+};
diff --git a/Engine/Plugins/Experimental/EditorDataStorageFeatures/Source/TedsOutliner/Public/Widgets/TypeInfoWidget.h b/Engine/Plugins/Experimental/EditorDataStorageFeatures/Source/TedsOutliner/Public/Widgets/TypeInfoWidget.h
--- a/Engine/Plugins/Experimental/EditorDataStorageFeatures/Source/TedsOutliner/Public/Widgets/TypeInfoWidget.h
+++ b/Engine/Plugins/Experimental/EditorDataStorageFeatures/Source/TedsOutliner/Public/Widgets/TypeInfoWidget.h
@@ -36,12 +36,12 @@ protected:
 	explicit FTypeInfoWidgetConstructor(const UScriptStruct* InTypeInfo);
 	TSharedPtr<SWidget> CreateWidget(const TypedElementDataStorage::FMetaDataView& Arguments) override;
 	bool FinalizeWidget(ITypedElementDataStorageInterface* DataStorage, ITypedElementDataStorageUiInterface* DataStorageUi,
-		TypedElementRowHandle Row, const TSharedPtr<SWidget>& Widget) override;
+		UE::Editor::DataStorage::RowHandle Row, const TSharedPtr<SWidget>& Widget) override;
 
 protected:
 
 	// Get the icon for a given row, checking the cache to see if we already have one
-	static const FSlateBrush* GetIconForRow(ITypedElementDataStorageInterface* DataStorage, TypedElementRowHandle Row, const FTypedElementClassTypeInfoColumn* TypeInfoColumn);
+	static const FSlateBrush* GetIconForRow(ITypedElementDataStorageInterface* DataStorage, UE::Editor::DataStorage::RowHandle Row, const FTypedElementClassTypeInfoColumn* TypeInfoColumn);
 	
 protected:
 
diff --git a/Engine/Plugins/Experimental/EditorDataStorageFeatures/Source/TedsPropertyEditor/Private/Widgets/SPropertyMenuTedsRowPicker.cpp b/Engine/Plugins/Experimental/EditorDataStorageFeatures/Source/TedsPropertyEditor/Private/Widgets/SPropertyMenuTedsRowPicker.cpp
--- a/Engine/Plugins/Experimental/EditorDataStorageFeatures/Source/TedsPropertyEditor/Private/Widgets/SPropertyMenuTedsRowPicker.cpp
+++ b/Engine/Plugins/Experimental/EditorDataStorageFeatures/Source/TedsPropertyEditor/Private/Widgets/SPropertyMenuTedsRowPicker.cpp
@@ -108,17 +108,17 @@ void SPropertyMenuTedsRowPicker::Construct(const FArguments& InArgs)
 
 void SPropertyMenuTedsRowPicker::OnClear()
 {
-	SetValue(TypedElementInvalidRowHandle);
+	SetValue(UE::Editor::DataStorage::InvalidRowHandle);
 	OnClose.ExecuteIfBound();
 }
 
-void SPropertyMenuTedsRowPicker::OnElementSelected(TypedElementDataStorage::RowHandle RowHandle)
+void SPropertyMenuTedsRowPicker::OnElementSelected(UE::Editor::DataStorage::RowHandle RowHandle)
 {
 	SetValue(RowHandle);
 	OnClose.ExecuteIfBound();
 }
 
-void SPropertyMenuTedsRowPicker::SetValue(TypedElementDataStorage::RowHandle RowHandle)
+void SPropertyMenuTedsRowPicker::SetValue(UE::Editor::DataStorage::RowHandle RowHandle)
 {
 	OnSet.ExecuteIfBound(RowHandle);
 }
diff --git a/Engine/Plugins/Experimental/EditorDataStorageFeatures/Source/TedsRevisionControl/Private/Processors/RevisionControlProcessors.cpp b/Engine/Plugins/Experimental/EditorDataStorageFeatures/Source/TedsRevisionControl/Private/Processors/RevisionControlProcessors.cpp
--- a/Engine/Plugins/Experimental/EditorDataStorageFeatures/Source/TedsRevisionControl/Private/Processors/RevisionControlProcessors.cpp
+++ b/Engine/Plugins/Experimental/EditorDataStorageFeatures/Source/TedsRevisionControl/Private/Processors/RevisionControlProcessors.cpp
@@ -18,145 +18,148 @@
 #include "Elements/Framework/TypedElementRegistry.h"
 #include "Elements/Interfaces/TypedElementQueryStorageInterfaces.h"
 
-extern FAutoConsoleVariableRef CVarAutoPopulateState;
-
-static bool gEnableOverlays = false;
-TAutoConsoleVariable<bool> CVarEnableOverlays(
-	TEXT("RevisionControl.Overlays.Enable"),
-	gEnableOverlays,
-	TEXT("Enables overlays."),
-	ECVF_Default);
-
-static bool gEnableOverlayCheckedOutByOtherUser = true;
-TAutoConsoleVariable<bool> CVarEnableOverlayCheckedOutByOtherUser(
-	TEXT("RevisionControl.Overlays.CheckedOutByOtherUser.Enable"),
-	gEnableOverlayCheckedOutByOtherUser,
-	TEXT("Enables overlays for files that are checked out by another user."),
-	ECVF_Default);
-
-static bool gEnableOverlayNotAtHeadRevision = true;
-TAutoConsoleVariable<bool> CVarEnableOverlayNotAtHeadRevision(
-	TEXT("RevisionControl.Overlays.NotAtHeadRevision.Enable"),
-	gEnableOverlayNotAtHeadRevision,
-	TEXT("Enables overlays for files that are not at the latest revision."),
-	ECVF_Default);
-
-static bool gEnableOverlayCheckedOut = false;
-TAutoConsoleVariable<bool> CVarEnableOverlayCheckedOut(
-	TEXT("RevisionControl.Overlays.CheckedOut.Enable"),
-	gEnableOverlayCheckedOut,
-	TEXT("Enables overlays for files that are checked out by user."),
-	ECVF_Default);
-
-static bool gEnableOverlayOpenForAdd = false;
-TAutoConsoleVariable<bool> CVarEnableOverlayOpenForAdd(
-	TEXT("RevisionControl.Overlays.OpenForAdd.Enable"),
-	gEnableOverlayOpenForAdd,
-	TEXT("Enables overlays for files that are newly added."),
-	ECVF_Default);
-
-static int32 gOverlayAlpha = 20; // [0..100]
-TAutoConsoleVariable<int32> CVarOverlayAlpha(
-	TEXT("RevisionControl.Overlays.Alpha"),
-	gOverlayAlpha,
-	TEXT("Configures overlay opacity."),
-	ECVF_Default);
-
-#if UE_BUILD_SHIPPING
-#define ENABLE_OVERLAY_DEBUG 0
-#else
-#define ENABLE_OVERLAY_DEBUG 1
-#endif
-
-#if ENABLE_OVERLAY_DEBUG
-static int32 gDefaultDebugForceColorOnAllValue = false; // [0..100]
-TAutoConsoleVariable<int32> CVarDebugForceColorOnAll(
-	TEXT("RevisionControl.Overlays.Debug.ForceColorOnAll"),
-	gDefaultDebugForceColorOnAllValue,
-	TEXT("Debug to force overlay color on everything. 1 = Red, 2 = Green, 3 = Blue, 4 = White. 0 = off  ."),
-	ECVF_Default);
-#endif
-
-static FColor DetermineOverlayColor(const TypedElementDataStorage::IQueryContext& ObjectContext, const TypedElementDataStorage::ICommonQueryContext& SCCContext, const FTypedElementUObjectColumn& Actor)
+namespace UE::Editor::RevisionControl::Private
 {
-	check(IsInGameThread());
-
-#if ENABLE_OVERLAY_DEBUG
-	if (int32 Force = CVarDebugForceColorOnAll.GetValueOnGameThread())
+	extern FAutoConsoleVariableRef CVarAutoPopulateState;
+
+	static bool gEnableOverlays = false;
+	TAutoConsoleVariable<bool> CVarEnableOverlays(
+		TEXT("RevisionControl.Overlays.Enable"),
+		gEnableOverlays,
+		TEXT("Enables overlays."),
+		ECVF_Default);
+
+	static bool gEnableOverlayCheckedOutByOtherUser = true;
+	TAutoConsoleVariable<bool> CVarEnableOverlayCheckedOutByOtherUser(
+		TEXT("RevisionControl.Overlays.CheckedOutByOtherUser.Enable"),
+		gEnableOverlayCheckedOutByOtherUser,
+		TEXT("Enables overlays for files that are checked out by another user."),
+		ECVF_Default);
+
+	static bool gEnableOverlayNotAtHeadRevision = true;
+	TAutoConsoleVariable<bool> CVarEnableOverlayNotAtHeadRevision(
+		TEXT("RevisionControl.Overlays.NotAtHeadRevision.Enable"),
+		gEnableOverlayNotAtHeadRevision,
+		TEXT("Enables overlays for files that are not at the latest revision."),
+		ECVF_Default);
+
+	static bool gEnableOverlayCheckedOut = false;
+	TAutoConsoleVariable<bool> CVarEnableOverlayCheckedOut(
+		TEXT("RevisionControl.Overlays.CheckedOut.Enable"),
+		gEnableOverlayCheckedOut,
+		TEXT("Enables overlays for files that are checked out by user."),
+		ECVF_Default);
+
+	static bool gEnableOverlayOpenForAdd = false;
+	TAutoConsoleVariable<bool> CVarEnableOverlayOpenForAdd(
+		TEXT("RevisionControl.Overlays.OpenForAdd.Enable"),
+		gEnableOverlayOpenForAdd,
+		TEXT("Enables overlays for files that are newly added."),
+		ECVF_Default);
+
+	static int32 gOverlayAlpha = 20; // [0..100]
+	TAutoConsoleVariable<int32> CVarOverlayAlpha(
+		TEXT("RevisionControl.Overlays.Alpha"),
+		gOverlayAlpha,
+		TEXT("Configures overlay opacity."),
+		ECVF_Default);
+
+	#if UE_BUILD_SHIPPING
+	#define ENABLE_OVERLAY_DEBUG 0
+	#else
+	#define ENABLE_OVERLAY_DEBUG 1
+	#endif
+
+	#if ENABLE_OVERLAY_DEBUG
+	static int32 gDefaultDebugForceColorOnAllValue = false; // [0..100]
+	TAutoConsoleVariable<int32> CVarDebugForceColorOnAll(
+		TEXT("RevisionControl.Overlays.Debug.ForceColorOnAll"),
+		gDefaultDebugForceColorOnAllValue,
+		TEXT("Debug to force overlay color on everything. 1 = Red, 2 = Green, 3 = Blue, 4 = White. 0 = off  ."),
+		ECVF_Default);
+	#endif
+
+	static FColor DetermineOverlayColor(const TypedElementDataStorage::IQueryContext& ObjectContext, const TypedElementDataStorage::ICommonQueryContext& SCCContext, const FTypedElementUObjectColumn& Actor)
 	{
-		int32 Alpha = FMath::Lerp<float>(0.f, 255.f, gOverlayAlpha / 100.f);
-		switch(Force)
-		{
-		case 1:
-			return FColor(255, 0, 0, Alpha);
-		case 2:
-			return FColor(0, 255, 0, Alpha);
-		case 3:
-			return FColor(0, 0, 255, Alpha);
-		case 4:
-			return FColor(255, 255, 255, Alpha);
-		// Do normal determination for higher than 4
-		}
-	}
-#endif
+		check(IsInGameThread());
 
-	bool bExternal = Actor.Object.IsValid() ? Cast<AActor>(Actor.Object)->IsPackageExternal() : false;
-	bool bIgnored = !bExternal;
-	bool bSelected = ObjectContext.HasColumn<FTypedElementSelectionColumn>();
-	if (!bIgnored && !bSelected)
-	{
-		// Convert CVar value from [0..100] to [0..255] range.
-		int32 Alpha = FMath::Lerp<float>(0.f, 255.f, CVarOverlayAlpha.GetValueOnGameThread() / 100.f);
-
-		// Check if the package is outdated because there is a newer version available.
-		if (SCCContext.HasColumn<FSCCNotCurrentTag>())
+	#if ENABLE_OVERLAY_DEBUG
+		if (int32 Force = CVarDebugForceColorOnAll.GetValueOnGameThread())
 		{
-			if (CVarEnableOverlayNotAtHeadRevision.GetValueOnGameThread())
+			int32 Alpha = FMath::Lerp<float>(0.f, 255.f, gOverlayAlpha / 100.f);
+			switch(Force)
 			{
-				// Yellow.
-				return FColor(225, 255, 61, Alpha);
+			case 1:
+				return FColor(255, 0, 0, Alpha);
+			case 2:
+				return FColor(0, 255, 0, Alpha);
+			case 3:
+				return FColor(0, 0, 255, Alpha);
+			case 4:
+				return FColor(255, 255, 255, Alpha);
+			// Do normal determination for higher than 4
 			}
 		}
+	#endif
 
-		// Check if the package is locked by someone else.
-		if (SCCContext.HasColumn<FSCCExternallyLockedColumn>())
+		bool bExternal = Actor.Object.IsValid() ? Cast<AActor>(Actor.Object)->IsPackageExternal() : false;
+		bool bIgnored = !bExternal;
+		bool bSelected = ObjectContext.HasColumn<FTypedElementSelectionColumn>();
+		if (!bIgnored && !bSelected)
 		{
-			if (CVarEnableOverlayCheckedOutByOtherUser.GetValueOnGameThread())
+			// Convert CVar value from [0..100] to [0..255] range.
+			int32 Alpha = FMath::Lerp<float>(0.f, 255.f, CVarOverlayAlpha.GetValueOnGameThread() / 100.f);
+
+			// Check if the package is outdated because there is a newer version available.
+			if (SCCContext.HasColumn<FSCCNotCurrentTag>())
 			{
-				// Red.
-				return FColor(239, 53, 53, Alpha);
+				if (CVarEnableOverlayNotAtHeadRevision.GetValueOnGameThread())
+				{
+					// Yellow.
+					return FColor(225, 255, 61, Alpha);
+				}
 			}
-		}
 
-		// Check if the package is added locally.
-		if (SCCContext.HasColumn<FSCCStatusColumn>())
-		{
-			if (CVarEnableOverlayOpenForAdd.GetValueOnGameThread())
+			// Check if the package is locked by someone else.
+			if (SCCContext.HasColumn<FSCCExternallyLockedColumn>())
+			{
+				if (CVarEnableOverlayCheckedOutByOtherUser.GetValueOnGameThread())
+				{
+					// Red.
+					return FColor(239, 53, 53, Alpha);
+				}
+			}
+
+			// Check if the package is added locally.
+			if (SCCContext.HasColumn<FSCCStatusColumn>())
 			{
-				if (const FSCCStatusColumn* StatusColumn = SCCContext.GetColumn<FSCCStatusColumn>())
+				if (CVarEnableOverlayOpenForAdd.GetValueOnGameThread())
 				{
-					if (StatusColumn->Modification == ESCCModification::Added)
+					if (const FSCCStatusColumn* StatusColumn = SCCContext.GetColumn<FSCCStatusColumn>())
 					{
-						// Blue.
-						return FColor(0, 112, 224, Alpha);
+						if (StatusColumn->Modification == ESCCModification::Added)
+						{
+							// Blue.
+							return FColor(0, 112, 224, Alpha);
+						}
 					}
 				}
 			}
-		}
 
-		// Check if the package is locked by self.
-		if (SCCContext.HasColumn<FSCCLockedTag>())
-		{
-			if (CVarEnableOverlayCheckedOut.GetValueOnGameThread())
+			// Check if the package is locked by self.
+			if (SCCContext.HasColumn<FSCCLockedTag>())
 			{
-				// Green.
-				return FColor(31, 228, 75, Alpha);
+				if (CVarEnableOverlayCheckedOut.GetValueOnGameThread())
+				{
+					// Green.
+					return FColor(31, 228, 75, Alpha);
+				}
 			}
 		}
-	}
 
-	return FColor(ForceInitToZero);
-}
+		return FColor(ForceInitToZero);
+	}
+} // namespace UE::Editor::RevisionControl::Private
 
 void URevisionControlDataStorageFactory::RegisterTables(ITypedElementDataStorageInterface& DataStorage)
 {
@@ -173,6 +176,9 @@ void URevisionControlDataStorageFactory::RegisterTables(ITypedElementDataStorage
 
 void URevisionControlDataStorageFactory::RegisterQueries(ITypedElementDataStorageInterface& DataStorage)
 {
+	using namespace UE::Editor::DataStorage;
+	using namespace UE::Editor::RevisionControl::Private;
+	
 	CVarAutoPopulateState->AsVariable()->OnChangedDelegate().AddLambda(
 		[this, &DataStorage](IConsoleVariable* AutoPopulate)
 		{
@@ -183,7 +189,7 @@ void URevisionControlDataStorageFactory::RegisterQueries(ITypedElementDataStorag
 			else
 			{
 				DataStorage.UnregisterQuery(FetchUpdates);
-				FetchUpdates = TypedElementInvalidQueryHandle;
+				FetchUpdates = InvalidQueryHandle;
 			}
 		}
 	);
@@ -194,19 +200,19 @@ void URevisionControlDataStorageFactory::RegisterQueries(ITypedElementDataStorag
 			if (EnableOverlays->GetBool())
 			{
 				DataStorage.UnregisterQuery(RemoveOverlays);
-				RemoveOverlays = TypedElementInvalidQueryHandle;
+				RemoveOverlays = InvalidQueryHandle;
 
 				RegisterApplyOverlays(DataStorage);
 			}
 			else
 			{
 				DataStorage.UnregisterQuery(ApplyNewOverlays);
-				ApplyNewOverlays = TypedElementInvalidQueryHandle;
+				ApplyNewOverlays = InvalidQueryHandle;
 				DataStorage.UnregisterQuery(ChangeOverlay);
-				ChangeOverlay = TypedElementInvalidQueryHandle;
+				ChangeOverlay = InvalidQueryHandle;
 
 				DataStorage.UnregisterQuery(ApplyOverlaysObjectToSCC);
-				ApplyOverlaysObjectToSCC = TypedElementInvalidQueryHandle;
+				ApplyOverlaysObjectToSCC = InvalidQueryHandle;
 
 				RegisterRemoveOverlays(DataStorage);
 			}
@@ -235,7 +241,7 @@ void URevisionControlDataStorageFactory::RegisterFetchUpdates(ITypedElementDataS
 	
 	FSourceControlFileStatusMonitor& FileStatusMonitor = ISourceControlModule::Get().GetSourceControlFileStatusMonitor();
 
-	if (FetchUpdates == TypedElementDataStorage::InvalidQueryHandle)
+	if (FetchUpdates == InvalidQueryHandle)
 	{
 		FetchUpdates = DataStorage.RegisterQuery(
 			Select(
@@ -263,6 +269,8 @@ void URevisionControlDataStorageFactory::RegisterApplyOverlays(ITypedElementData
 {
 	using namespace TypedElementDataStorage;
 	using namespace TypedElementQueryBuilder;
+	using namespace UE::Editor::DataStorage;
+	using namespace UE::Editor::RevisionControl::Private;
 	
 	if (ApplyOverlaysObjectToSCC == InvalidQueryHandle)
 	{
@@ -347,7 +355,7 @@ void URevisionControlDataStorageFactory::RegisterApplyOverlays(ITypedElementData
 					//
 					// Note: Remove and re-add will trigger observer in TypedElementActorViewportProcessors to SetOverlayColor on the primitive components
 					Context.RunSubquery(EChangeOverlay, Update.ObjectRow, CreateSubqueryCallbackBinding(
-						[&Context](TypedElementRowHandle ObjectRow, const FTypedElementUObjectColumn& Actor, const FTypedElementPackageReference& PackageReference, const FTypedElementViewportOverlayColorColumn& OverlayColorColumn)
+						[&Context](RowHandle ObjectRow, const FTypedElementUObjectColumn& Actor, const FTypedElementPackageReference& PackageReference, const FTypedElementViewportOverlayColorColumn& OverlayColorColumn)
 						{
 							Context.RunSubquery(EApplyOverlaysObjectToSCC, PackageReference.Row, CreateSubqueryCallbackBinding(
 								[&Context, &ObjectRow, &Actor, &OverlayColorColumn](ISubqueryContext& SubQueryContext)
@@ -381,8 +389,9 @@ void URevisionControlDataStorageFactory::RegisterRemoveOverlays(ITypedElementDat
 {
 	using namespace TypedElementQueryBuilder;
 	using namespace TypedElementDataStorage;
+	using namespace UE::Editor::DataStorage;
 	
-	if (RemoveOverlays == TypedElementInvalidQueryHandle)
+	if (RemoveOverlays == InvalidQueryHandle)
 	{
 		// Query:
 		// For all actors WITH an overlay color column AND having a package reference:
@@ -395,7 +404,7 @@ void URevisionControlDataStorageFactory::RegisterRemoveOverlays(ITypedElementDat
 				// This is in PrePhysics because the overlay->actor query is in DuringPhysics and contexts don't flush changes between tick groups
 				FProcessor(EQueryTickPhase::PrePhysics, DataStorage.GetQueryTickGroupName(EQueryTickGroups::SyncExternalToDataStorage))
 					.SetExecutionMode(EExecutionMode::GameThread),
-				[](IQueryContext& Context, TypedElementRowHandle ObjectRow, FTypedElementUObjectColumn& Actor, const FTypedElementViewportOverlayColorColumn& ViewportColor)
+				[](IQueryContext& Context, RowHandle ObjectRow, FTypedElementUObjectColumn& Actor, const FTypedElementViewportOverlayColorColumn& ViewportColor)
 				{
 					Context.RemoveColumns<FTypedElementViewportOverlayColorColumn>(ObjectRow);
 				}
diff --git a/Engine/Plugins/Experimental/EditorDataStorageFeatures/Source/TedsRevisionControl/Private/Processors/RevisionControlProcessors.h b/Engine/Plugins/Experimental/EditorDataStorageFeatures/Source/TedsRevisionControl/Private/Processors/RevisionControlProcessors.h
--- a/Engine/Plugins/Experimental/EditorDataStorageFeatures/Source/TedsRevisionControl/Private/Processors/RevisionControlProcessors.h
+++ b/Engine/Plugins/Experimental/EditorDataStorageFeatures/Source/TedsRevisionControl/Private/Processors/RevisionControlProcessors.h
@@ -24,10 +24,10 @@ private:
 	void RegisterFetchUpdates(ITypedElementDataStorageInterface& DataStorage);
 	void RegisterApplyOverlays(ITypedElementDataStorageInterface& DataStorage);
 	void RegisterRemoveOverlays(ITypedElementDataStorageInterface& DataStorage);
-	TypedElementDataStorage::QueryHandle FetchUpdates = TypedElementDataStorage::InvalidQueryHandle;
-	TypedElementDataStorage::QueryHandle ApplyNewOverlays = TypedElementDataStorage::InvalidQueryHandle;
-	TypedElementDataStorage::QueryHandle ChangeOverlay = TypedElementDataStorage::InvalidQueryHandle;
-	TypedElementDataStorage::QueryHandle ApplyOverlaysObjectToSCC = TypedElementDataStorage::InvalidQueryHandle;
-	TypedElementDataStorage::QueryHandle RemoveOverlays = TypedElementDataStorage::InvalidQueryHandle;
-	TypedElementDataStorage::QueryHandle FlushPackageUpdates = TypedElementDataStorage::InvalidQueryHandle;
+	UE::Editor::DataStorage::QueryHandle FetchUpdates = UE::Editor::DataStorage::InvalidQueryHandle;
+	UE::Editor::DataStorage::QueryHandle ApplyNewOverlays = UE::Editor::DataStorage::InvalidQueryHandle;
+	UE::Editor::DataStorage::QueryHandle ChangeOverlay = UE::Editor::DataStorage::InvalidQueryHandle;
+	UE::Editor::DataStorage::QueryHandle ApplyOverlaysObjectToSCC = UE::Editor::DataStorage::InvalidQueryHandle;
+	UE::Editor::DataStorage::QueryHandle RemoveOverlays = UE::Editor::DataStorage::InvalidQueryHandle;
+	UE::Editor::DataStorage::QueryHandle FlushPackageUpdates = UE::Editor::DataStorage::InvalidQueryHandle;
 };
diff --git a/Engine/Plugins/Experimental/EditorDataStorageFeatures/Source/TedsRevisionControl/Private/Queries/ObjectPackagePathToColumnQueries.cpp b/Engine/Plugins/Experimental/EditorDataStorageFeatures/Source/TedsRevisionControl/Private/Queries/ObjectPackagePathToColumnQueries.cpp
--- a/Engine/Plugins/Experimental/EditorDataStorageFeatures/Source/TedsRevisionControl/Private/Queries/ObjectPackagePathToColumnQueries.cpp
+++ b/Engine/Plugins/Experimental/EditorDataStorageFeatures/Source/TedsRevisionControl/Private/Queries/ObjectPackagePathToColumnQueries.cpp
@@ -11,39 +11,45 @@
 #include "Misc/Paths.h"
 #include "UObject/Package.h"
 
-static bool bAutoPopulateRevisionControlState = false;
-FAutoConsoleVariableRef CVarAutoPopulateState(
-	TEXT("TEDS.RevisionControl.AutoPopulateState"),
-	bAutoPopulateRevisionControlState,
-	TEXT("Automatically query revision control provider and fill information into TEDS")
-);
-
-static void ResolvePackageReference(ITypedElementDataStorageInterface::IQueryContext& Context, const UPackage* Package, TypedElementRowHandle Row, TypedElementRowHandle PackageRow)
+namespace UE::Editor::RevisionControl::Private
 {
-	FTypedElementPackageReference PackageReference;
-	PackageReference.Row = PackageRow;
-	Context.AddColumn(Row, MoveTemp(PackageReference));
-				
-	FTypedElementPackagePathColumn PathColumn;
-	FTypedElementPackageLoadedPathColumn LoadedPathColumn;
-		
-	Package->GetPathName(nullptr, PathColumn.Path);
-	LoadedPathColumn.LoadedPath = Package->GetLoadedPath();
-				
-	Context.AddColumn(PackageRow, MoveTemp(PathColumn));
-	Context.AddColumn(PackageRow, MoveTemp(LoadedPathColumn));
+	using namespace UE::Editor::DataStorage;
+
+	static bool bAutoPopulateRevisionControlState = false;
+	FAutoConsoleVariableRef CVarAutoPopulateState(
+		TEXT("TEDS.RevisionControl.AutoPopulateState"),
+		bAutoPopulateRevisionControlState,
+		TEXT("Automatically query revision control provider and fill information into TEDS")
+	);
+
+	static void ResolvePackageReference(ITypedElementDataStorageInterface::IQueryContext& Context, const UPackage* Package, RowHandle Row, RowHandle PackageRow)
+	{
+		FTypedElementPackageReference PackageReference;
+		PackageReference.Row = PackageRow;
+		Context.AddColumn(Row, MoveTemp(PackageReference));
+
+		FTypedElementPackagePathColumn PathColumn;
+		FTypedElementPackageLoadedPathColumn LoadedPathColumn;
+
+		Package->GetPathName(nullptr, PathColumn.Path);
+		LoadedPathColumn.LoadedPath = Package->GetLoadedPath();
+
+		Context.AddColumn(PackageRow, MoveTemp(PathColumn));
+		Context.AddColumn(PackageRow, MoveTemp(LoadedPathColumn));
 
-	FTypedElementPackageReference PackageBackReference;
-	PackageBackReference.Row = Row;
-	Context.AddColumn(PackageRow, MoveTemp(PackageBackReference));
-};
+		FTypedElementPackageReference PackageBackReference;
+		PackageBackReference.Row = Row;
+		Context.AddColumn(PackageRow, MoveTemp(PackageBackReference));
+	};
+} // namespace UE::Editor::DataStorage
 
 void UTypedElementUObjectPackagePathFactory::RegisterQueries(ITypedElementDataStorageInterface& DataStorage)
 {
 	using namespace TypedElementQueryBuilder;
 	using namespace TypedElementDataStorage;
+	using namespace UE::Editor::DataStorage;
 	
-	CVarAutoPopulateState->AsVariable()->OnChangedDelegate().AddLambda(
+	UE::Editor::RevisionControl::Private::CVarAutoPopulateState->AsVariable()->OnChangedDelegate().AddLambda(
 		[this, &DataStorage](IConsoleVariable* AutoPopulate)
 		{
 			if (AutoPopulate->GetBool())
@@ -60,9 +66,9 @@ void UTypedElementUObjectPackagePathFactory::RegisterQueries(ITypedElementDataSt
 		Select(
 			TEXT("Resolve package references"),
 			FProcessor(EQueryTickPhase::FrameEnd, DataStorage.GetQueryTickGroupName(EQueryTickGroups::SyncExternalToDataStorage)),
-			[](IQueryContext& Context, TypedElementRowHandle Row, const FTypedElementUObjectColumn& Object, const FTypedElementPackageUnresolvedReference& UnresolvedPackageReference)
+			[](IQueryContext& Context, RowHandle Row, const FTypedElementUObjectColumn& Object, const FTypedElementPackageUnresolvedReference& UnresolvedPackageReference)
 			{
-				TypedElementRowHandle PackageRow = Context.FindIndexedRow(UnresolvedPackageReference.Index);
+				RowHandle PackageRow = Context.FindIndexedRow(UnresolvedPackageReference.Index);
 				if (!Context.IsRowAvailable(PackageRow))
 				{
 					return;
@@ -75,13 +81,13 @@ void UTypedElementUObjectPackagePathFactory::RegisterQueries(ITypedElementDataSt
 				const UPackage* Package = ObjectInstance->GetPackage();
 				Context.RemoveColumns(Row, { FTypedElementPackageUnresolvedReference::StaticStruct() });
 
-				ResolvePackageReference(Context, Package, Row, PackageRow);
+				UE::Editor::RevisionControl::Private::ResolvePackageReference(Context, Package, Row, PackageRow);
 			}
 		)
 		.Compile()
 	);
 
-	if (CVarAutoPopulateState->GetBool())
+	if (UE::Editor::RevisionControl::Private::CVarAutoPopulateState->GetBool())
 	{
 		RegisterTryAddPackageRef(DataStorage);
 	}
@@ -91,13 +97,14 @@ void UTypedElementUObjectPackagePathFactory::RegisterTryAddPackageRef(ITypedElem
 {
 	using namespace TypedElementQueryBuilder;
 	using namespace TypedElementDataStorage;
+	using namespace UE::Editor::DataStorage;
 	
 	TryAddPackageRef = DataStorage.RegisterQuery(
 		Select(
 			TEXT("Sync UObject package info to columns"),
 			FObserver::OnAdd<FTypedElementUObjectColumn>()
 				.SetExecutionMode(EExecutionMode::GameThread),
-			[](IQueryContext& Context, TypedElementRowHandle Row, const FTypedElementUObjectColumn& Object)
+			[](IQueryContext& Context, RowHandle Row, const FTypedElementUObjectColumn& Object)
 			{
 				if (const UObject* ObjectInstance = Object.Object.Get(); ObjectInstance != nullptr)
 				{
@@ -109,11 +116,11 @@ void UTypedElementUObjectPackagePathFactory::RegisterTryAddPackageRef(ITypedElem
 						FPaths::NormalizeFilename(PackageFilename);
 						FString FullPackageFilename = FPaths::ConvertRelativePathToFull(PackageFilename);
 						TypedElementDataStorage::IndexHash Index = TypedElementDataStorage::GenerateIndexHash(FullPackageFilename);
-						TypedElementRowHandle PackageRow = Context.FindIndexedRow(Index);
+						RowHandle PackageRow = Context.FindIndexedRow(Index);
 						if (Context.IsRowAvailable(PackageRow))
 						{
 							const UPackage* Package = ObjectInstance->GetPackage();
-							ResolvePackageReference(Context, Package, Row, PackageRow);
+							UE::Editor::RevisionControl::Private::ResolvePackageReference(Context, Package, Row, PackageRow);
 						}
 						else
 						{
diff --git a/Engine/Plugins/Experimental/EditorDataStorageFeatures/Source/TedsRevisionControl/Private/Queries/ObjectPackagePathToColumnQueries.h b/Engine/Plugins/Experimental/EditorDataStorageFeatures/Source/TedsRevisionControl/Private/Queries/ObjectPackagePathToColumnQueries.h
--- a/Engine/Plugins/Experimental/EditorDataStorageFeatures/Source/TedsRevisionControl/Private/Queries/ObjectPackagePathToColumnQueries.h
+++ b/Engine/Plugins/Experimental/EditorDataStorageFeatures/Source/TedsRevisionControl/Private/Queries/ObjectPackagePathToColumnQueries.h
@@ -21,5 +21,5 @@ public:
 
 private:
 	void RegisterTryAddPackageRef(ITypedElementDataStorageInterface& DataStorage);
-	TypedElementDataStorage::QueryHandle TryAddPackageRef = TypedElementDataStorage::InvalidQueryHandle;
+	UE::Editor::DataStorage::QueryHandle TryAddPackageRef = UE::Editor::DataStorage::InvalidQueryHandle;
 };
diff --git a/Engine/Plugins/Experimental/EditorDataStorageFeatures/Source/TedsSettings/Private/TedsSettingsManager.cpp b/Engine/Plugins/Experimental/EditorDataStorageFeatures/Source/TedsSettings/Private/TedsSettingsManager.cpp
--- a/Engine/Plugins/Experimental/EditorDataStorageFeatures/Source/TedsSettings/Private/TedsSettingsManager.cpp
+++ b/Engine/Plugins/Experimental/EditorDataStorageFeatures/Source/TedsSettings/Private/TedsSettingsManager.cpp
@@ -195,6 +195,8 @@ void FTedsSettingsManager::UnregisterSettings()
 
 void FTedsSettingsManager::UpdateSettingsCategory(TSharedPtr<ISettingsCategory> SettingsCategory, const FName& ContainerName, const bool bQueryExistingRows)
 {
+	using namespace UE::Editor::DataStorage;
+
 	TRACE_CPUPROFILER_EVENT_SCOPE(TedsSettingsManager.UpdateSettingsCategory);
 
 	const FName& CategoryName = SettingsCategory->GetName();
@@ -210,7 +212,7 @@ void FTedsSettingsManager::UpdateSettingsCategory(TSharedPtr<ISettingsCategory>
 	ITypedElementDataStorageCompatibilityInterface* DataStorageCompatibility = TypedElementRegistry->GetMutableDataStorageCompatibility();
 	check(DataStorageCompatibility);
 
-	TArray<TypedElementDataStorage::RowHandle> OldRowHandles;
+	TArray<RowHandle> OldRowHandles;
 	TArray<FName> OldSectionNames;
 
 	// Gather all existing rows for the given { ContainerName, CategoryName } pair.
@@ -265,7 +267,7 @@ void FTedsSettingsManager::UpdateSettingsCategory(TSharedPtr<ISettingsCategory>
 				continue;
 			}
 
-			TypedElementDataStorage::RowHandle NewRow = DataStorageCompatibility->AddCompatibleObject(SettingsObjectPtr);
+			RowHandle NewRow = DataStorageCompatibility->AddCompatibleObject(SettingsObjectPtr);
 
 			DataStorage->AddColumn<FSettingsContainerColumn>(NewRow, { .ContainerName = ContainerName });
 			DataStorage->AddColumn<FSettingsCategoryColumn>(NewRow, { .CategoryName = CategoryName });
@@ -286,8 +288,8 @@ void FTedsSettingsManager::UpdateSettingsCategory(TSharedPtr<ISettingsCategory>
 			continue;
 		}
 
-		TypedElementDataStorage::RowHandle OldRowHandle = OldRowHandles[RowIndex];
-		check(OldRowHandle != TypedElementInvalidRowHandle);
+		RowHandle OldRowHandle = OldRowHandles[RowIndex];
+		check(OldRowHandle != InvalidRowHandle);
 
 		DataStorage->RemoveRow(OldRowHandle);
 
diff --git a/Engine/Plugins/Experimental/EditorDataStorageFeatures/Source/TedsSettings/Tests/TedsSettingsTests.cpp b/Engine/Plugins/Experimental/EditorDataStorageFeatures/Source/TedsSettings/Tests/TedsSettingsTests.cpp
--- a/Engine/Plugins/Experimental/EditorDataStorageFeatures/Source/TedsSettings/Tests/TedsSettingsTests.cpp
+++ b/Engine/Plugins/Experimental/EditorDataStorageFeatures/Source/TedsSettings/Tests/TedsSettingsTests.cpp
@@ -16,30 +16,33 @@
 #include "TedsSettingsEditorSubsystem.h"
 #include "TestSettings.h"
 
-BEGIN_DEFINE_SPEC(FTedsSettingsTestFixture, "TedsSettings", EAutomationTestFlags::EditorContext | EAutomationTestFlags::EngineFilter)
+namespace UE::Editor::Settings::Tests
+{
+	using namespace UE::Editor;
 
-ISettingsModule* SettingsModule = nullptr;
-UTypedElementRegistry* TypedElementRegistry = nullptr;
-ITypedElementDataStorageInterface* DataStorage = nullptr;
-ITypedElementDataStorageCompatibilityInterface* DataStorageCompatibility = nullptr;
-TypedElementQueryHandle CountAllSettingsQuery = TypedElementDataStorage::InvalidQueryHandle;
+	BEGIN_DEFINE_SPEC(FTedsSettingsTestFixture, "TedsSettings", EAutomationTestFlags::EditorContext | EAutomationTestFlags::EngineFilter)
+	ISettingsModule* SettingsModule = nullptr;
+	UTypedElementRegistry* TypedElementRegistry = nullptr;
+	ITypedElementDataStorageInterface* DataStorage = nullptr;
+	ITypedElementDataStorageCompatibilityInterface* DataStorageCompatibility = nullptr;
+	DataStorage::QueryHandle CountAllSettingsQuery = DataStorage::InvalidQueryHandle;
 
-uint32 BeforeRowCount = 0;
-TArray<TypedElementDataStorage::RowHandle> TestRowHandles{};
+	uint32 BeforeRowCount = 0;
+	TArray<DataStorage::RowHandle> TestRowHandles{};
 
-uint32 CountSettingsRowsInDataStorage()
-{
-	using DSI = ITypedElementDataStorageInterface;
+	uint32 CountSettingsRowsInDataStorage()
+	{
+		using DSI = ITypedElementDataStorageInterface;
 
-	DSI::FQueryResult Result = DataStorage->RunQuery(CountAllSettingsQuery);
+		DSI::FQueryResult Result = DataStorage->RunQuery(CountAllSettingsQuery);
 
-	return Result.Count;
-}
+		return Result.Count;
+	}
 
-template<typename Func>
-void AwaitRowHandleThenVerify(TypedElementDataStorage::RowHandle RowHandle, const FDoneDelegate& Done, Func&& OnVerify)
-{
-	auto OnTick = [this, RowHandle, Done, OnVerify = Forward<Func>(OnVerify)](float FrameTime)
+	template<typename Func>
+	void AwaitRowHandleThenVerify(DataStorage::RowHandle RowHandle, const FDoneDelegate& Done, Func&& OnVerify)
+	{
+		auto OnTick = [this, RowHandle, Done, OnVerify = Forward<Func>(OnVerify)](float FrameTime)
 		{
 			if (DataStorage->IsRowAssigned(RowHandle))
 			{
@@ -52,141 +55,142 @@ void AwaitRowHandleThenVerify(TypedElementDataStorage::RowHandle RowHandle, cons
 			return true;
 		};
 
-	FTSTicker::GetCoreTicker().AddTicker(FTickerDelegate::CreateLambda(OnTick));
-}
-
-END_DEFINE_SPEC(FTedsSettingsTestFixture)
-
-void FTedsSettingsTestFixture::Define()
-{
-	check(GEditor);
-	UTedsSettingsEditorSubsystem* SettingsEditorSubsystem = GEditor->GetEditorSubsystem<UTedsSettingsEditorSubsystem>();
+		FTSTicker::GetCoreTicker().AddTicker(FTickerDelegate::CreateLambda(OnTick));
+	}
 
-	SettingsEditorSubsystem->OnEnabledChanged().RemoveAll(this);
-	SettingsEditorSubsystem->OnEnabledChanged().AddRaw(this, &FTedsSettingsTestFixture::Redefine);
+	END_DEFINE_SPEC(FTedsSettingsTestFixture)
 
-	if (!SettingsEditorSubsystem->IsEnabled())
+	void FTedsSettingsTestFixture::Define()
 	{
-		return;
-	}
+		check(GEditor);
+		UTedsSettingsEditorSubsystem* SettingsEditorSubsystem = GEditor->GetEditorSubsystem<UTedsSettingsEditorSubsystem>();
 
-	BeforeEach([this]()
-	{
-		SettingsModule = FModuleManager::GetModulePtr<ISettingsModule>("Settings");
-		check(SettingsModule != nullptr);
-	
-		TypedElementRegistry = UTypedElementRegistry::GetInstance();
-		check(TypedElementRegistry != nullptr);
-	
-		DataStorage = TypedElementRegistry->GetMutableDataStorage();
-		check(DataStorage != nullptr);
-	
-		DataStorageCompatibility = TypedElementRegistry->GetMutableDataStorageCompatibility();
-		check(DataStorageCompatibility != nullptr);
+		SettingsEditorSubsystem->OnEnabledChanged().RemoveAll(this);
+		SettingsEditorSubsystem->OnEnabledChanged().AddRaw(this, &FTedsSettingsTestFixture::Redefine);
 
+		if (!SettingsEditorSubsystem->IsEnabled())
 		{
-			using namespace TypedElementQueryBuilder;
-			CountAllSettingsQuery = DataStorage->RegisterQuery(
-				Count()
-				.Where()
-					.All<FTypedElementUObjectColumn, FSettingsContainerColumn, FSettingsCategoryColumn, FSettingsSectionColumn>()
-				.Compile());
+			return;
 		}
 
-		BeforeRowCount = CountSettingsRowsInDataStorage();
-	});
-
-	AfterEach([this]()
-	{
-		for (TypedElementDataStorage::RowHandle RowHandle : TestRowHandles)
+		BeforeEach([this]()
 		{
-			DataStorage->RemoveRow(RowHandle);
-		}
-		TestRowHandles.Empty();
-		CountAllSettingsQuery = TypedElementDataStorage::InvalidQueryHandle;
-		SettingsModule = nullptr;
-		TypedElementRegistry = nullptr;
-		DataStorage = nullptr;
-		DataStorageCompatibility = nullptr;
-	});
-
-	Describe("RegisterSettings", [this]()
-	{
-		LatentIt("Should add a row to editor data storage", [this](const FDoneDelegate& Done)
-		{
-			const FName& ContainerName = FName(TEXT("TestContainer"));
-			const FName& CategoryName = FName(TEXT("TestCategory"));
-			const FName& SectionName = FName(TEXT("TestSection"));
-				
-			TObjectPtr<UTestSettings> TestSettingsObject = NewObject<UTestSettings>();
-			
-			SettingsModule->RegisterSettings(ContainerName, CategoryName, SectionName, FText(), FText(), TestSettingsObject);
-
-			TypedElementDataStorage::RowHandle RowHandle = DataStorageCompatibility->FindRowWithCompatibleObject(TestSettingsObject);
-			TestNotEqual(TEXT("RowHandle"), RowHandle, TypedElementDataStorage::InvalidRowHandle);
-
-			if (RowHandle == TypedElementDataStorage::InvalidRowHandle)
+			SettingsModule = FModuleManager::GetModulePtr<ISettingsModule>("Settings");
+			check(SettingsModule != nullptr);
+
+			TypedElementRegistry = UTypedElementRegistry::GetInstance();
+			check(TypedElementRegistry != nullptr);
+
+			DataStorage = TypedElementRegistry->GetMutableDataStorage();
+			check(DataStorage != nullptr);
+
+			DataStorageCompatibility = TypedElementRegistry->GetMutableDataStorageCompatibility();
+			check(DataStorageCompatibility != nullptr);
+
 			{
-				Done.Execute();
-				return;
+				using namespace TypedElementQueryBuilder;
+				CountAllSettingsQuery = DataStorage->RegisterQuery(
+					Count()
+					.Where()
+					.All<FTypedElementUObjectColumn, FSettingsContainerColumn, FSettingsCategoryColumn, FSettingsSectionColumn>()
+					.Compile());
 			}
 
-			TestRowHandles.Push(RowHandle);
+			BeforeRowCount = CountSettingsRowsInDataStorage();
+		});
 
-			AwaitRowHandleThenVerify(RowHandle, Done, [this, RowHandle, ContainerName, CategoryName, SectionName]()
+		AfterEach([this]()
+		{
+			for (DataStorage::RowHandle RowHandle : TestRowHandles)
 			{
-				uint32 AfterRowCount = CountSettingsRowsInDataStorage();
+				DataStorage->RemoveRow(RowHandle);
+			}
+			TestRowHandles.Empty();
+			CountAllSettingsQuery = DataStorage::InvalidQueryHandle;
+			SettingsModule = nullptr;
+			TypedElementRegistry = nullptr;
+			DataStorage = nullptr;
+			DataStorageCompatibility = nullptr;
+		});
+
+		Describe("RegisterSettings", [this]()
+		{
+			LatentIt("Should add a row to editor data storage", [this](const FDoneDelegate& Done)
+			{
+				const FName& ContainerName = FName(TEXT("TestContainer"));
+				const FName& CategoryName = FName(TEXT("TestCategory"));
+				const FName& SectionName = FName(TEXT("TestSection"));
+
+				TObjectPtr<UTestSettings> TestSettingsObject = NewObject<UTestSettings>();
+
+				SettingsModule->RegisterSettings(ContainerName, CategoryName, SectionName, FText(), FText(), TestSettingsObject);
+
+				DataStorage::RowHandle RowHandle = DataStorageCompatibility->FindRowWithCompatibleObject(TestSettingsObject);
+				TestNotEqual(TEXT("RowHandle"), RowHandle, DataStorage::InvalidRowHandle);
 
-				TestEqual(TEXT("RowCount"), AfterRowCount, BeforeRowCount + 1);
-				
-				TestEqual(TEXT("ContainerName"), DataStorage->GetColumn<FSettingsContainerColumn>(RowHandle)->ContainerName, ContainerName);
-				TestEqual(TEXT("CategoryName"), DataStorage->GetColumn<FSettingsCategoryColumn>(RowHandle)->CategoryName, CategoryName);
-				TestEqual(TEXT("SectionName"), DataStorage->GetColumn<FSettingsSectionColumn>(RowHandle)->SectionName, SectionName);
+				if (RowHandle == DataStorage::InvalidRowHandle)
+				{
+					Done.Execute();
+					return;
+				}
+
+				TestRowHandles.Push(RowHandle);
+
+				AwaitRowHandleThenVerify(RowHandle, Done, [this, RowHandle, ContainerName, CategoryName, SectionName]()
+				{
+					uint32 AfterRowCount = CountSettingsRowsInDataStorage();
+
+					TestEqual(TEXT("RowCount"), AfterRowCount, BeforeRowCount + 1);
+
+					TestEqual(TEXT("ContainerName"), DataStorage->GetColumn<FSettingsContainerColumn>(RowHandle)->ContainerName, ContainerName);
+					TestEqual(TEXT("CategoryName"), DataStorage->GetColumn<FSettingsCategoryColumn>(RowHandle)->CategoryName, CategoryName);
+					TestEqual(TEXT("SectionName"), DataStorage->GetColumn<FSettingsSectionColumn>(RowHandle)->SectionName, SectionName);
+				});
 			});
 		});
-	});
 
-	Describe("UnregisterSettings", [this]()
-	{
-		LatentIt("Should remove a row from editor data storage", [this](const FDoneDelegate& Done)
+		Describe("UnregisterSettings", [this]()
 		{
-			const FName& ContainerName = FName(TEXT("TestContainer"));
-			const FName& CategoryName = FName(TEXT("TestCategory"));
-			const FName& SectionName = FName(TEXT("TestSection"));
+			LatentIt("Should remove a row from editor data storage", [this](const FDoneDelegate& Done)
+			{
+				const FName& ContainerName = FName(TEXT("TestContainer"));
+				const FName& CategoryName = FName(TEXT("TestCategory"));
+				const FName& SectionName = FName(TEXT("TestSection"));
 
-			TObjectPtr<UTestSettings> TestSettingsObject = NewObject<UTestSettings>();
+				TObjectPtr<UTestSettings> TestSettingsObject = NewObject<UTestSettings>();
 
-			SettingsModule->RegisterSettings(ContainerName, CategoryName, SectionName, FText(), FText(), TestSettingsObject);
+				SettingsModule->RegisterSettings(ContainerName, CategoryName, SectionName, FText(), FText(), TestSettingsObject);
 
-			TypedElementDataStorage::RowHandle RowHandle = DataStorageCompatibility->FindRowWithCompatibleObject(TestSettingsObject);
-			TestNotEqual(TEXT("RowHandle"), RowHandle, TypedElementDataStorage::InvalidRowHandle);
+				DataStorage::RowHandle RowHandle = DataStorageCompatibility->FindRowWithCompatibleObject(TestSettingsObject);
+				TestNotEqual(TEXT("RowHandle"), RowHandle, DataStorage::InvalidRowHandle);
 
-			if (RowHandle == TypedElementDataStorage::InvalidRowHandle)
-			{
-				Done.Execute();
-				return;
-			}
+				if (RowHandle == DataStorage::InvalidRowHandle)
+				{
+					Done.Execute();
+					return;
+				}
 
-			TestRowHandles.Push(RowHandle);
+				TestRowHandles.Push(RowHandle);
 
-			AwaitRowHandleThenVerify(RowHandle, Done, [this, RowHandle, TestSettingsObject, ContainerName, CategoryName, SectionName]()
-			{
-				uint32 AfterRegisterRowCount = CountSettingsRowsInDataStorage();
+				AwaitRowHandleThenVerify(RowHandle, Done, [this, RowHandle, TestSettingsObject, ContainerName, CategoryName, SectionName]()
+				{
+					uint32 AfterRegisterRowCount = CountSettingsRowsInDataStorage();
 
-				TestEqual(TEXT("RowCount"), AfterRegisterRowCount, BeforeRowCount + 1);
+					TestEqual(TEXT("RowCount"), AfterRegisterRowCount, BeforeRowCount + 1);
 
-				SettingsModule->UnregisterSettings(ContainerName, CategoryName, SectionName);
+					SettingsModule->UnregisterSettings(ContainerName, CategoryName, SectionName);
 
-				uint32 AfterUnregisterRowCount = CountSettingsRowsInDataStorage();
-				TestEqual(TEXT("RowCount"), AfterUnregisterRowCount, BeforeRowCount);
+					uint32 AfterUnregisterRowCount = CountSettingsRowsInDataStorage();
+					TestEqual(TEXT("RowCount"), AfterUnregisterRowCount, BeforeRowCount);
 
-				TestFalse(TEXT("IsRowAssigned"), DataStorage->IsRowAssigned(RowHandle));
+					TestFalse(TEXT("IsRowAssigned"), DataStorage->IsRowAssigned(RowHandle));
 
-				TypedElementDataStorage::RowHandle InvalidRowHandle = DataStorageCompatibility->FindRowWithCompatibleObject(TestSettingsObject);
-				TestEqual(TEXT("InvalidRowHandle"), InvalidRowHandle, TypedElementDataStorage::InvalidRowHandle);
+					DataStorage::RowHandle TestInvalidRowHandle = DataStorageCompatibility->FindRowWithCompatibleObject(TestSettingsObject);
+					TestEqual(TEXT("InvalidRowHandle"), TestInvalidRowHandle, DataStorage::InvalidRowHandle);
+				});
 			});
 		});
-	});
-}
+	}
+} // namespace UE::Editor::Settings::Tests
 
 #endif // WITH_AUTOMATION_TESTS
diff --git a/Engine/Plugins/Experimental/EditorDataStorageFeatures/Source/TedsTableViewer/Private/QueryStack/FQueryStackNode_RowView.cpp b/Engine/Plugins/Experimental/EditorDataStorageFeatures/Source/TedsTableViewer/Private/QueryStack/FQueryStackNode_RowView.cpp
--- a/Engine/Plugins/Experimental/EditorDataStorageFeatures/Source/TedsTableViewer/Private/QueryStack/FQueryStackNode_RowView.cpp
+++ b/Engine/Plugins/Experimental/EditorDataStorageFeatures/Source/TedsTableViewer/Private/QueryStack/FQueryStackNode_RowView.cpp
@@ -2,9 +2,9 @@
 
 #include "QueryStack/FQueryStackNode_RowView.h"
 
-namespace UE::EditorDataStorage
+namespace UE::Editor::DataStorage
 {
-	TConstArrayView<TypedElementDataStorage::RowHandle> FQueryStackNode_RowView::GetOrderedRowList()
+	TConstArrayView<RowHandle> FQueryStackNode_RowView::GetOrderedRowList()
 	{
 		return *Rows;
 	}
@@ -14,7 +14,7 @@ namespace UE::EditorDataStorage
 		return RevisionId;
 	}
 
-	FQueryStackNode_RowView::FQueryStackNode_RowView(TArray<TypedElementDataStorage::RowHandle>* InRows)
+	FQueryStackNode_RowView::FQueryStackNode_RowView(TArray<RowHandle>* InRows)
 		: Rows(InRows)
 	{
 		
@@ -24,6 +24,4 @@ namespace UE::EditorDataStorage
 	{
 		++RevisionId;
 	}
-}
-
-
+} // namespace UE::Editor::DataStorage
diff --git a/Engine/Plugins/Experimental/EditorDataStorageFeatures/Source/TedsTableViewer/Private/TedsTableViewerColumn.cpp b/Engine/Plugins/Experimental/EditorDataStorageFeatures/Source/TedsTableViewer/Private/TedsTableViewerColumn.cpp
--- a/Engine/Plugins/Experimental/EditorDataStorageFeatures/Source/TedsTableViewer/Private/TedsTableViewerColumn.cpp
+++ b/Engine/Plugins/Experimental/EditorDataStorageFeatures/Source/TedsTableViewer/Private/TedsTableViewerColumn.cpp
@@ -13,7 +13,7 @@
 #include "Widgets/SBoxPanel.h"
 #include "Widgets/Text/STextBlock.h"
 
-namespace UE::EditorDataStorage
+namespace UE::Editor::DataStorage
 {
 	FTedsTableViewerColumn::FTedsTableViewerColumn(const FName& InColumnName, const TSharedPtr<FTypedElementWidgetConstructor>& InCellWidgetConstructor,
 		const TArray<TWeakObjectPtr<const UScriptStruct>>& InMatchedColumns, const TSharedPtr<FTypedElementWidgetConstructor>& InHeaderWidgetConstructor,
@@ -43,13 +43,13 @@ namespace UE::EditorDataStorage
 		UnRegisterQueries();
 	}
 
-	TSharedPtr<SWidget> FTedsTableViewerColumn::ConstructRowWidget(TypedElementDataStorage::RowHandle RowHandle) const
+	TSharedPtr<SWidget> FTedsTableViewerColumn::ConstructRowWidget(RowHandle InRowHandle) const
 	{
 		TSharedPtr<SWidget> RowWidget;
 		
-		if(Storage->IsRowAssigned(RowHandle))
+		if(Storage->IsRowAssigned(InRowHandle))
 		{
-			const TypedElementRowHandle UiRowHandle = Storage->AddRow(Storage->FindTable(TableViewerUtils::GetWidgetTableName()));
+			const RowHandle UiRowHandle = Storage->AddRow(Storage->FindTable(TableViewerUtils::GetWidgetTableName()));
 
 			const TConstArrayView<TWeakObjectPtr<const UScriptStruct>> ColumnTypes = GetMatchedColumns();
 			if (ColumnTypes.Num() == 1)
@@ -59,7 +59,7 @@ namespace UE::EditorDataStorage
 			
 			if (FTypedElementRowReferenceColumn* RowReference = Storage->GetColumn<FTypedElementRowReferenceColumn>(UiRowHandle))
 			{
-				RowReference->Row = RowHandle;
+				RowReference->Row = InRowHandle;
 			}
 
 			if(FTypedElementSlateWidgetReferenceColumn* WidgetReferenceColumn = Storage->GetColumn<FTypedElementSlateWidgetReferenceColumn>(UiRowHandle))
@@ -88,7 +88,7 @@ namespace UE::EditorDataStorage
 		}
 
 		TSharedPtr<SWidget> Widget;
-		TypedElementDataStorage::RowHandle UiRowHandle = TypedElementDataStorage::InvalidRowHandle;
+		RowHandle UiRowHandle = InvalidRowHandle;
 		if (HeaderWidgetConstructor)
 		{
 			UiRowHandle = Storage->AddRow(Storage->FindTable(FName(TEXT("Editor_WidgetTable"))));
@@ -170,13 +170,13 @@ namespace UE::EditorDataStorage
 
 			// TEDS-Outliner TODO: Long term if we move this into TypedElementOutlinerMode or similar we can get access to the exact
 			// types the Outliner is looking at and specify them on .Where() to cut down on the things we are observing
-			TypedElementDataStorage::QueryHandle AddQueryHandle = Storage->RegisterQuery(
+			QueryHandle AddQueryHandle = Storage->RegisterQuery(
 				Select(
 					ColumnAddObserverName,
 					AddObserver,
-					[this](IQueryContext& Context, TypedElementRowHandle Row)
+					[this](IQueryContext& Context, RowHandle Row)
 						{
-							RowsToUpdate.Add(TPair<TypedElementRowHandle, bool>(Row, true));
+							RowsToUpdate.Add(TPair<RowHandle, bool>(Row, true));
 						})
 				.Where()
 					.All(ColumnType.Get())
@@ -191,13 +191,13 @@ namespace UE::EditorDataStorage
 
 			// Table Viewer TODO: We might be able to cut down on the rows we are querying for in the future by getting the rows from the query stack
 			// but we currently have to use a generic query so we can support the TEDS-Outliner as well
-			TypedElementDataStorage::QueryHandle RemoveQueryHandle = Storage->RegisterQuery(
+			QueryHandle RemoveQueryHandle = Storage->RegisterQuery(
 				Select(
 					ColumnRemoveObserverName,
 					RemoveObserver,
-					[this](IQueryContext& Context, TypedElementRowHandle Row)
+					[this](IQueryContext& Context, RowHandle Row)
 						{
-							RowsToUpdate.Add(TPair<TypedElementRowHandle, bool>(Row, false));
+							RowsToUpdate.Add(TPair<RowHandle, bool>(Row, false));
 						})
 				.Where()
 					.All(ColumnType.Get())
@@ -227,7 +227,7 @@ namespace UE::EditorDataStorage
 
 	void FTedsTableViewerColumn::UnRegisterQueries()
 	{
-		for(const TypedElementDataStorage::QueryHandle Query : InternalObserverQueries)
+		for(const QueryHandle Query : InternalObserverQueries)
 		{
 			Storage->UnregisterQuery(Query);
 		}
@@ -236,7 +236,7 @@ namespace UE::EditorDataStorage
 
 	}
 
-	bool FTedsTableViewerColumn::IsRowVisible(const TypedElementDataStorage::RowHandle InRowHandle) const
+	bool FTedsTableViewerColumn::IsRowVisible(const RowHandle InRowHandle) const
 	{
 		if(IsRowVisibleDelegate.IsBound())
 		{
@@ -251,7 +251,7 @@ namespace UE::EditorDataStorage
 	void FTedsTableViewerColumn::UpdateWidgets()
 	{
 		// Remove any widget rows that don't actually need an update
-		RowsToUpdate = RowsToUpdate.FilterByPredicate([this](const TPair<TypedElementDataStorage::RowHandle, bool>& Pair) -> bool
+		RowsToUpdate = RowsToUpdate.FilterByPredicate([this](const TPair<RowHandle, bool>& Pair) -> bool
 		{
 			// We don't have a widget for this item visible, so there is nothing to update
 			if(!IsRowVisible(Pair.Key))
@@ -276,7 +276,7 @@ namespace UE::EditorDataStorage
 		{
 			const FTypedElementSlateWidgetReferenceColumn* WidgetsIt = ContainerWidgetReferenceColumns;
 			const FTypedElementRowReferenceColumn* RowRefsIt = RowReferenceColumns;
-			const TConstArrayView<TypedElementRowHandle> Rows = Context.GetRowHandles();
+			const TConstArrayView<RowHandle> Rows = Context.GetRowHandles();
 
 			for(unsigned RowIndex = 0; RowIndex < Context.GetRowCount(); ++RowIndex, ++WidgetsIt, ++RowRefsIt)
 			{
diff --git a/Engine/Plugins/Experimental/EditorDataStorageFeatures/Source/TedsTableViewer/Private/TedsTableViewerModel.cpp b/Engine/Plugins/Experimental/EditorDataStorageFeatures/Source/TedsTableViewer/Private/TedsTableViewerModel.cpp
--- a/Engine/Plugins/Experimental/EditorDataStorageFeatures/Source/TedsTableViewer/Private/TedsTableViewerModel.cpp
+++ b/Engine/Plugins/Experimental/EditorDataStorageFeatures/Source/TedsTableViewer/Private/TedsTableViewerModel.cpp
@@ -9,7 +9,7 @@
 #include "Elements/Columns/TypedElementMiscColumns.h"
 #include "Elements/Columns/TypedElementSlateWidgetColumns.h"
 
-namespace UE::EditorDataStorage
+namespace UE::Editor::DataStorage
 {
 	FTedsTableViewerModel::FTedsTableViewerModel(const TSharedPtr<IQueryStackNode_Row>& InRowQueryStack,
 		const TArray<TWeakObjectPtr<const UScriptStruct>>& InRequestedColumns, const TArray<FName>& InCellWidgetPurposes,
@@ -271,6 +271,4 @@ namespace UE::EditorDataStorage
 
 		return true;
 	}
-
-
-}
+} // namespace UE::Editor::DataStorage
diff --git a/Engine/Plugins/Experimental/EditorDataStorageFeatures/Source/TedsTableViewer/Private/TedsTableViewerModule.cpp b/Engine/Plugins/Experimental/EditorDataStorageFeatures/Source/TedsTableViewer/Private/TedsTableViewerModule.cpp
--- a/Engine/Plugins/Experimental/EditorDataStorageFeatures/Source/TedsTableViewer/Private/TedsTableViewerModule.cpp
+++ b/Engine/Plugins/Experimental/EditorDataStorageFeatures/Source/TedsTableViewer/Private/TedsTableViewerModule.cpp
@@ -86,19 +86,19 @@ public:
 		Rows.Empty();
 		
 		TypedElementDataStorage::FQueryResult QueryResult = DataStorage->RunQuery(QueryHandle,
-			CreateDirectQueryCallbackBinding([this](const ITypedElementDataStorageInterface::IDirectQueryContext& Context, const TypedElementDataStorage::RowHandle* RowHandles)
+			CreateDirectQueryCallbackBinding([this](const ITypedElementDataStorageInterface::IDirectQueryContext& Context, const UE::Editor::DataStorage::RowHandle* RowHandles)
 		{
 			Rows.Append(RowHandles, Context.GetRowCount());
 		}));
 
-		return SNew(UE::EditorDataStorage::STedsTableViewer)
-				.QueryStack(MakeShared<UE::EditorDataStorage::FQueryStackNode_RowView>(&Rows))
+		return SNew(UE::Editor::DataStorage::STedsTableViewer)
+				.QueryStack(MakeShared<UE::Editor::DataStorage::FQueryStackNode_RowView>(&Rows))
 				.Columns({FTypedElementLabelColumn::StaticStruct(), FTypedElementSelectionColumn::StaticStruct(),
 					FTypedElementAlertColumn::StaticStruct(), FTypedElementChildAlertColumn::StaticStruct()});
 	}
 	
 	FDelegateHandle LevelEditorTabManagerChangedHandle;
-	TArray<TypedElementDataStorage::RowHandle> Rows;
+	TArray<UE::Editor::DataStorage::RowHandle> Rows;
 };
 
 // CVar to summon a test table viewer that views a snapshot of all actors at the moment when the Cvar is used
diff --git a/Engine/Plugins/Experimental/EditorDataStorageFeatures/Source/TedsTableViewer/Private/TedsTableViewerUtils.cpp b/Engine/Plugins/Experimental/EditorDataStorageFeatures/Source/TedsTableViewer/Private/TedsTableViewerUtils.cpp
--- a/Engine/Plugins/Experimental/EditorDataStorageFeatures/Source/TedsTableViewer/Private/TedsTableViewerUtils.cpp
+++ b/Engine/Plugins/Experimental/EditorDataStorageFeatures/Source/TedsTableViewer/Private/TedsTableViewerUtils.cpp
@@ -6,7 +6,7 @@
 #include "Elements/Common/TypedElementDataStorageLog.h"
 #include "Elements/Interfaces/TypedElementDataStorageUiInterface.h"
 
-namespace UE::EditorDataStorage::TableViewerUtils
+namespace UE::Editor::DataStorage::TableViewerUtils
 {
 	static FName TableViewerWidgetTableName("Editor_TableViewerWidgetTable");
 
@@ -143,13 +143,14 @@ namespace UE::EditorDataStorage::TableViewerUtils
 
 void UTypedElementTableViewerFactory::RegisterTables(ITypedElementDataStorageInterface& DataStorage)
 {
-	const TypedElementTableHandle BaseWidgetTable = DataStorage.FindTable(FName(TEXT("Editor_WidgetTable")));
-	if (BaseWidgetTable != TypedElementInvalidTableHandle)
+	using namespace UE::Editor::DataStorage;
+	const TableHandle BaseWidgetTable = DataStorage.FindTable(FName(TEXT("Editor_WidgetTable")));
+	if (BaseWidgetTable != InvalidTableHandle)
 	{
 		DataStorage.RegisterTable(
 			BaseWidgetTable,
 			{
 				FTypedElementRowReferenceColumn::StaticStruct()
-			}, UE::EditorDataStorage::TableViewerUtils::TableViewerWidgetTableName);
+			}, TableViewerUtils::TableViewerWidgetTableName);
 	}
 }
diff --git a/Engine/Plugins/Experimental/EditorDataStorageFeatures/Source/TedsTableViewer/Private/Widgets/RowHandleWidget.cpp b/Engine/Plugins/Experimental/EditorDataStorageFeatures/Source/TedsTableViewer/Private/Widgets/RowHandleWidget.cpp
--- a/Engine/Plugins/Experimental/EditorDataStorageFeatures/Source/TedsTableViewer/Private/Widgets/RowHandleWidget.cpp
+++ b/Engine/Plugins/Experimental/EditorDataStorageFeatures/Source/TedsTableViewer/Private/Widgets/RowHandleWidget.cpp
@@ -41,7 +41,7 @@ TSharedPtr<SWidget> FRowHandleWidgetConstructor::CreateWidget(const TypedElement
 		.VAlign(VAlign_Center)
 		.Padding(8, 0, 0, 0);}
 
-bool FRowHandleWidgetConstructor::FinalizeWidget(ITypedElementDataStorageInterface* DataStorage, ITypedElementDataStorageUiInterface* DataStorageUi, TypedElementDataStorage::RowHandle Row, const TSharedPtr<SWidget>& Widget)
+bool FRowHandleWidgetConstructor::FinalizeWidget(ITypedElementDataStorageInterface* DataStorage, ITypedElementDataStorageUiInterface* DataStorageUi, UE::Editor::DataStorage::RowHandle Row, const TSharedPtr<SWidget>& Widget)
 {
 
 	checkf(Widget->GetType() == SBox::StaticWidgetClass().GetWidgetType(),
diff --git a/Engine/Plugins/Experimental/EditorDataStorageFeatures/Source/TedsTableViewer/Private/Widgets/SRowDetails.cpp b/Engine/Plugins/Experimental/EditorDataStorageFeatures/Source/TedsTableViewer/Private/Widgets/SRowDetails.cpp
--- a/Engine/Plugins/Experimental/EditorDataStorageFeatures/Source/TedsTableViewer/Private/Widgets/SRowDetails.cpp
+++ b/Engine/Plugins/Experimental/EditorDataStorageFeatures/Source/TedsTableViewer/Private/Widgets/SRowDetails.cpp
@@ -12,16 +12,16 @@
 
 #define LOCTEXT_NAMESPACE "SRowDetails"
 
-namespace UE::SRowDetails::Local
+namespace UE::Editor::DataStorage
 {
-	static const FName NameColumn(TEXT("Name"));
-	static const FName DataColumn(TEXT("Data"));
-	static TArray<FName> DefaultWidgetPurposes = { FName(TEXT("RowDetails.Cell.Large")),
-		FName(TEXT("RowDetails.Cell")), FName(TEXT("General.Cell.Large")), FName(TEXT("General.Cell")) };
-}
+	namespace Widgets::Private
+	{
+		static const FName NameColumn(TEXT("Name"));
+		static const FName DataColumn(TEXT("Data"));
+		static TArray<FName> DefaultWidgetPurposes = { FName(TEXT("RowDetails.Cell.Large")),
+			FName(TEXT("RowDetails.Cell")), FName(TEXT("General.Cell.Large")), FName(TEXT("General.Cell")) };
+	} // namespace Widgets::Private
 
-namespace UE::EditorDataStorage
-{
 	//
 	// SRowDetails
 	//
@@ -33,7 +33,7 @@ namespace UE::EditorDataStorage
 
 		if(WidgetPurposes.IsEmpty())
 		{
-			WidgetPurposes = UE::SRowDetails::Local::DefaultWidgetPurposes;
+			WidgetPurposes = Widgets::Private::DefaultWidgetPurposes;
 		}
 
 		UTypedElementRegistry* Registry = UTypedElementRegistry::GetInstance();
@@ -56,17 +56,17 @@ namespace UE::EditorDataStorage
 				.HeaderRow
 				(
 					SNew(SHeaderRow)
-					+SHeaderRow::Column(UE::SRowDetails::Local::NameColumn)
+					+SHeaderRow::Column(Widgets::Private::NameColumn)
 						.DefaultLabel(FText::FromString(TEXT("Name")))
 						.FillWidth(0.3f)
-					+SHeaderRow::Column(UE::SRowDetails::Local::DataColumn)
+					+SHeaderRow::Column(Widgets::Private::DataColumn)
 						.DefaultLabel(FText::FromString(TEXT("Value")))
 						.FillWidth(0.7f)
 				)
 		];
 	}
 	
-	void SRowDetails::SetRow(TypedElementDataStorage::RowHandle Row)
+	void SRowDetails::SetRow(RowHandle Row)
 	{
 		if (DataStorage->IsRowAssigned(Row))
 		{
@@ -142,9 +142,7 @@ namespace UE::EditorDataStorage
 	}
 	
 	TSharedRef<SWidget> SRowDetailsRow::GenerateWidgetForColumn(const FName& ColumnName)
-	{
-		using namespace TypedElementDataStorage;
-		
+	{		
 		if (!DataStorage->IsRowAvailable(Item->WidgetRow))
 		{
 			Item->WidgetRow = DataStorage->AddRow(DataStorage->FindTable(FName("Editor_WidgetTable")));
@@ -163,13 +161,13 @@ namespace UE::EditorDataStorage
 					});
 			}
 		}
-		if (ColumnName == UE::SRowDetails::Local::NameColumn)
+		if (ColumnName == Widgets::Private::NameColumn)
 		{
 			return SNew(STextBlock)
 					.Text(FText::FromString(Item->WidgetConstructor->CreateWidgetDisplayName(
 							DataStorage, Item->WidgetRow)));
 		}
-		else if (ColumnName == UE::SRowDetails::Local::DataColumn)
+		else if (ColumnName == Widgets::Private::DataColumn)
 		{
 			return DataStorageUi->ConstructWidget(Item->WidgetRow, *(Item->WidgetConstructor), {}).ToSharedRef();
 		}
@@ -182,7 +180,7 @@ namespace UE::EditorDataStorage
 
 	// FRowDetailsItem
 	FRowDetailsItem::FRowDetailsItem(const TWeakObjectPtr<const UScriptStruct>& InColumnType,
-		TUniquePtr<FTypedElementWidgetConstructor> InWidgetConstructor, TypedElementDataStorage::RowHandle InRow)
+		TUniquePtr<FTypedElementWidgetConstructor> InWidgetConstructor, RowHandle InRow)
 		: ColumnType(InColumnType)
 		, WidgetConstructor(MoveTemp(InWidgetConstructor))
 		, Row(InRow)
diff --git a/Engine/Plugins/Experimental/EditorDataStorageFeatures/Source/TedsTableViewer/Private/Widgets/STedsTableViewer.cpp b/Engine/Plugins/Experimental/EditorDataStorageFeatures/Source/TedsTableViewer/Private/Widgets/STedsTableViewer.cpp
--- a/Engine/Plugins/Experimental/EditorDataStorageFeatures/Source/TedsTableViewer/Private/Widgets/STedsTableViewer.cpp
+++ b/Engine/Plugins/Experimental/EditorDataStorageFeatures/Source/TedsTableViewer/Private/Widgets/STedsTableViewer.cpp
@@ -9,7 +9,7 @@
 
 #define LOCTEXT_NAMESPACE "STedsTableViewer"
 
-namespace UE::EditorDataStorage
+namespace UE::Editor::DataStorage
 {
 	void STedsTableViewer::Construct(const FArguments& InArgs)
 	{
@@ -122,7 +122,7 @@ namespace UE::EditorDataStorage
 		return SNew(STedsTableViewerRow, OwnerTable, Model.ToSharedRef())
 				.Item(InItem);
 	}
-}
+} // namespace UE::Editor::DataStorage
 
 #undef LOCTEXT_NAMESPACE //"STedsTableViewer"
 
