TEDS-Content Browser prototype (1/2): Core TEDS/TEDS UI changes

TEDS: Add a strongly typed wrapper around row handles
TEDS: Add a tag for rows that don't want to be displayed in the table viewer
TEDS UI: STedsWidgets are now automatically registered with TEDS if the system creating the widget already hasn't done so
TEDS UI: Add slate delegate columns that are used by STedsTableViewer by binding them to slate attributes on the internal list view
TEDS UI: Add attribute bindings to bind delegates inside TEDS columns to slate events

diff --git a/Engine/Plugins/Experimental/Animation/HierarchyTable/Source/Editor/Private/HierarchyTableWidgetConstructor.cpp b/Engine/Plugins/Experimental/Animation/HierarchyTable/Source/Editor/Private/HierarchyTableWidgetConstructor.cpp
--- a/Engine/Plugins/Experimental/Animation/HierarchyTable/Source/Editor/Private/HierarchyTableWidgetConstructor.cpp
+++ b/Engine/Plugins/Experimental/Animation/HierarchyTable/Source/Editor/Private/HierarchyTableWidgetConstructor.cpp
@@ -11,7 +11,7 @@ FHierarchyTableWidgetConstructor::FHierarchyTableWidgetConstructor(const UScript
 {
 }
 
-TSharedRef<SWidget> FHierarchyTableWidgetConstructor::CreateWidget(FHierarchyTableEntryData* EntryData)
+TSharedRef<SWidget> FHierarchyTableWidgetConstructor::CreateInternalWidget(FHierarchyTableEntryData* EntryData)
 {
 	return SNullWidget::NullWidget;
 }
@@ -32,7 +32,7 @@ bool FHierarchyTableWidgetConstructor::FinalizeWidget(ITypedElementDataStorageIn
 	TypedElementDataStorage::RowHandle TargetRow = DataStorage->GetColumn<FTypedElementRowReferenceColumn>(Row)->Row;
 	FTypedElementOverrideColumn* OverrideColumn = DataStorage->GetColumn<FTypedElementOverrideColumn>(TargetRow);
 
-	TSharedRef<SWidget> ActualWidget = CreateWidget(OverrideColumn->OwnerEntry);
+	TSharedRef<SWidget> ActualWidget = CreateInternalWidget(OverrideColumn->OwnerEntry);
 	WidgetInstance->SetContent(ActualWidget);
 
 	return true;
diff --git a/Engine/Plugins/Experimental/Animation/HierarchyTable/Source/Editor/Public/HierarchyTableWidgetConstructor.h b/Engine/Plugins/Experimental/Animation/HierarchyTable/Source/Editor/Public/HierarchyTableWidgetConstructor.h
--- a/Engine/Plugins/Experimental/Animation/HierarchyTable/Source/Editor/Public/HierarchyTableWidgetConstructor.h
+++ b/Engine/Plugins/Experimental/Animation/HierarchyTable/Source/Editor/Public/HierarchyTableWidgetConstructor.h
@@ -19,7 +19,7 @@ public:
 
 	virtual ~FHierarchyTableWidgetConstructor() override = default;
 
-	virtual TSharedRef<SWidget> CreateWidget(FHierarchyTableEntryData* EntryData);
+	virtual TSharedRef<SWidget> CreateInternalWidget(FHierarchyTableEntryData* EntryData);
 
 protected:
 	// Begin FHierarchyTableWidgetConstructor
diff --git a/Engine/Plugins/Experimental/Animation/HierarchyTableBuiltin/Source/Editor/Private/MaskProfile/MaskProfileWidgetConstructor.cpp b/Engine/Plugins/Experimental/Animation/HierarchyTableBuiltin/Source/Editor/Private/MaskProfile/MaskProfileWidgetConstructor.cpp
--- a/Engine/Plugins/Experimental/Animation/HierarchyTableBuiltin/Source/Editor/Private/MaskProfile/MaskProfileWidgetConstructor.cpp
+++ b/Engine/Plugins/Experimental/Animation/HierarchyTableBuiltin/Source/Editor/Private/MaskProfile/MaskProfileWidgetConstructor.cpp
@@ -5,7 +5,7 @@
 #include "MaskProfile/HierarchyTableTypeMask.h"
 #include "HierarchyTable.h"
 
-TSharedRef<SWidget> FHierarchyTableMaskWidgetConstructor_Value::CreateWidget(FHierarchyTableEntryData* EntryData)
+TSharedRef<SWidget> FHierarchyTableMaskWidgetConstructor_Value::CreateInternalWidget(FHierarchyTableEntryData* EntryData)
 {
 	return SNew(SSpinBox<float>)
 		.IsEnabled_Lambda([EntryData]() { return EntryData->IsOverridden(); })
diff --git a/Engine/Plugins/Experimental/Animation/HierarchyTableBuiltin/Source/Editor/Private/MaskProfile/MaskProfileWidgetConstructor.h b/Engine/Plugins/Experimental/Animation/HierarchyTableBuiltin/Source/Editor/Private/MaskProfile/MaskProfileWidgetConstructor.h
--- a/Engine/Plugins/Experimental/Animation/HierarchyTableBuiltin/Source/Editor/Private/MaskProfile/MaskProfileWidgetConstructor.h
+++ b/Engine/Plugins/Experimental/Animation/HierarchyTableBuiltin/Source/Editor/Private/MaskProfile/MaskProfileWidgetConstructor.h
@@ -17,5 +17,5 @@ public:
 	~FHierarchyTableMaskWidgetConstructor_Value() override = default;
 
 protected:
-	virtual TSharedRef<SWidget> CreateWidget(FHierarchyTableEntryData* EntryData) override;
+	virtual TSharedRef<SWidget> CreateInternalWidget(FHierarchyTableEntryData* EntryData) override;
 };
\ No newline at end of file
diff --git a/Engine/Plugins/Experimental/Animation/HierarchyTableBuiltin/Source/Editor/Private/TimeProfile/TimeProfileWidgetConstructor.cpp b/Engine/Plugins/Experimental/Animation/HierarchyTableBuiltin/Source/Editor/Private/TimeProfile/TimeProfileWidgetConstructor.cpp
--- a/Engine/Plugins/Experimental/Animation/HierarchyTableBuiltin/Source/Editor/Private/TimeProfile/TimeProfileWidgetConstructor.cpp
+++ b/Engine/Plugins/Experimental/Animation/HierarchyTableBuiltin/Source/Editor/Private/TimeProfile/TimeProfileWidgetConstructor.cpp
@@ -5,7 +5,7 @@
 #include "TimeProfile/HierarchyTableTypeTime.h"
 #include "HierarchyTable.h"
 
-TSharedRef<SWidget> FHierarchyTableTimeWidgetConstructor_StartTime::CreateWidget(FHierarchyTableEntryData* EntryData)
+TSharedRef<SWidget> FHierarchyTableTimeWidgetConstructor_StartTime::CreateInternalWidget(FHierarchyTableEntryData* EntryData)
 {
 	return SNew(SSpinBox<float>)
 		.IsEnabled_Lambda([EntryData]() { return EntryData->IsOverridden(); })
@@ -22,7 +22,7 @@ TSharedRef<SWidget> FHierarchyTableTimeWidgetConstructor_StartTime::CreateWidget
 			});
 }
 
-TSharedRef<SWidget> FHierarchyTableTimeWidgetConstructor_EndTime::CreateWidget(FHierarchyTableEntryData* EntryData)
+TSharedRef<SWidget> FHierarchyTableTimeWidgetConstructor_EndTime::CreateInternalWidget(FHierarchyTableEntryData* EntryData)
 {
 	return SNew(SSpinBox<float>)
 		.IsEnabled_Lambda([EntryData]() { return EntryData->IsOverridden(); })
@@ -39,7 +39,7 @@ TSharedRef<SWidget> FHierarchyTableTimeWidgetConstructor_EndTime::CreateWidget(F
 			});
 }
 
-TSharedRef<SWidget> FHierarchyTableTimeWidgetConstructor_TimeFactor::CreateWidget(FHierarchyTableEntryData* EntryData)
+TSharedRef<SWidget> FHierarchyTableTimeWidgetConstructor_TimeFactor::CreateInternalWidget(FHierarchyTableEntryData* EntryData)
 {
 	return SNew(SSpinBox<float>)
 		.IsEnabled_Lambda([EntryData]() { return EntryData->IsOverridden(); })
@@ -56,7 +56,7 @@ TSharedRef<SWidget> FHierarchyTableTimeWidgetConstructor_TimeFactor::CreateWidge
 			});
 }
 
-TSharedRef<SWidget> FHierarchyTableTimeWidgetConstructor_Preview::CreateWidget(FHierarchyTableEntryData* EntryData)
+TSharedRef<SWidget> FHierarchyTableTimeWidgetConstructor_Preview::CreateInternalWidget(FHierarchyTableEntryData* EntryData)
 {
 	return SNew(STextBlock)
 		.Text(INVTEXT("PREVIEW"));
diff --git a/Engine/Plugins/Experimental/Animation/HierarchyTableBuiltin/Source/Editor/Private/TimeProfile/TimeProfileWidgetConstructor.h b/Engine/Plugins/Experimental/Animation/HierarchyTableBuiltin/Source/Editor/Private/TimeProfile/TimeProfileWidgetConstructor.h
--- a/Engine/Plugins/Experimental/Animation/HierarchyTableBuiltin/Source/Editor/Private/TimeProfile/TimeProfileWidgetConstructor.h
+++ b/Engine/Plugins/Experimental/Animation/HierarchyTableBuiltin/Source/Editor/Private/TimeProfile/TimeProfileWidgetConstructor.h
@@ -17,7 +17,7 @@ public:
 	~FHierarchyTableTimeWidgetConstructor_StartTime() override = default;
 
 protected:
-	virtual TSharedRef<SWidget> CreateWidget(FHierarchyTableEntryData* EntryData) override;
+	virtual TSharedRef<SWidget> CreateInternalWidget(FHierarchyTableEntryData* EntryData) override;
 };
 
 USTRUCT()
@@ -30,7 +30,7 @@ public:
 	~FHierarchyTableTimeWidgetConstructor_EndTime() override = default;
 
 protected:
-	virtual TSharedRef<SWidget> CreateWidget(FHierarchyTableEntryData* EntryData) override;
+	virtual TSharedRef<SWidget> CreateInternalWidget(FHierarchyTableEntryData* EntryData) override;
 };
 
 USTRUCT()
@@ -43,7 +43,7 @@ public:
 	~FHierarchyTableTimeWidgetConstructor_TimeFactor() override = default;
 
 protected:
-	virtual TSharedRef<SWidget> CreateWidget(FHierarchyTableEntryData* EntryData) override;
+	virtual TSharedRef<SWidget> CreateInternalWidget(FHierarchyTableEntryData* EntryData) override;
 };
 
 USTRUCT()
@@ -56,5 +56,5 @@ public:
 	~FHierarchyTableTimeWidgetConstructor_Preview() override = default;
 
 protected:
-	virtual TSharedRef<SWidget> CreateWidget(FHierarchyTableEntryData* EntryData) override;
+	virtual TSharedRef<SWidget> CreateInternalWidget(FHierarchyTableEntryData* EntryData) override;
 };
\ No newline at end of file
diff --git a/Engine/Plugins/Experimental/EditorDataStorageFeatures/Source/TedsTableViewer/Private/TedsTableViewerModel.cpp b/Engine/Plugins/Experimental/EditorDataStorageFeatures/Source/TedsTableViewer/Private/TedsTableViewerModel.cpp
--- a/Engine/Plugins/Experimental/EditorDataStorageFeatures/Source/TedsTableViewer/Private/TedsTableViewerModel.cpp
+++ b/Engine/Plugins/Experimental/EditorDataStorageFeatures/Source/TedsTableViewer/Private/TedsTableViewerModel.cpp
@@ -2,12 +2,12 @@
 
 #include "TedsTableViewerModel.h"
 
+
+#include "Elements/Columns/TypedElementUIColumns.h"
 #include "Elements/Framework/TypedElementRegistry.h"
 #include "QueryStack/IQueryStackNode_Row.h"
 #include "TedsTableViewerColumn.h"
 #include "TedsTableViewerUtils.h"
-#include "Elements/Columns/TypedElementMiscColumns.h"
-#include "Elements/Columns/TypedElementSlateWidgetColumns.h"
 
 namespace UE::Editor::DataStorage
 {
@@ -48,7 +48,8 @@ namespace UE::Editor::DataStorage
 		{
 			if(IsRowDisplayable(RowHandle))
 			{
-				Items.Add(RowHandle);
+				FTedsRowHandle TedsRowHandle{ .RowHandle = RowHandle };
+				Items.Add(TedsRowHandle);
 			}
 		}
 
@@ -59,20 +60,7 @@ namespace UE::Editor::DataStorage
 
 	bool FTedsTableViewerModel::IsRowDisplayable(TypedElementDataStorage::RowHandle InRowHandle) const
 	{
-		// We don't want to display any second level widgets (widgets for widgets and so on...) because they will keep cause the table viewer to
-		// infinitely grow as you keep scrolling (which creates new widgets)
-		if(Storage->HasColumns<FTypedElementSlateWidgetReferenceColumn>(InRowHandle))
-		{
-			if(const FTypedElementRowReferenceColumn* RowReferenceColumn = Storage->GetColumn<FTypedElementRowReferenceColumn>(InRowHandle))
-			{
-				if(Storage->HasColumns<FTypedElementSlateWidgetReferenceColumn>(RowReferenceColumn->Row))
-				{
-					return false;
-				}
-			}
-		}
-
-		return true;
+		return !Storage->HasColumns<FHideRowFromUITag>(InRowHandle);
 	}
 
 	bool FTedsTableViewerModel::Tick(float DeltaTime)
@@ -149,6 +137,11 @@ namespace UE::Editor::DataStorage
 		ColumnsView.Insert(InColumn, 0);
 	}
 
+	ITypedElementDataStorageInterface* FTedsTableViewerModel::GetDataStorageInterface() const
+	{
+		return Storage;
+	}
+
 	void FTedsTableViewerModel::GenerateColumns()
 	{
 		using MatchApproach = ITypedElementDataStorageUiInterface::EMatchApproach;
diff --git a/Engine/Plugins/Experimental/EditorDataStorageFeatures/Source/TedsTableViewer/Private/Widgets/STedsTableViewer.cpp b/Engine/Plugins/Experimental/EditorDataStorageFeatures/Source/TedsTableViewer/Private/Widgets/STedsTableViewer.cpp
--- a/Engine/Plugins/Experimental/EditorDataStorageFeatures/Source/TedsTableViewer/Private/Widgets/STedsTableViewer.cpp
+++ b/Engine/Plugins/Experimental/EditorDataStorageFeatures/Source/TedsTableViewer/Private/Widgets/STedsTableViewer.cpp
@@ -3,6 +3,10 @@
 #include "Widgets/STedsTableViewer.h"
 
 #include "TedsTableViewerColumn.h"
+#include "Columns/SlateDelegateColumns.h"
+#include "Elements/Columns/TypedElementUIColumns.h"
+#include "Elements/Framework/TypedElementAttributeBinding.h"
+#include "Elements/Framework/TypedElementDataStorageWidget.h"
 #include "Widgets/STedsTableViewerRow.h"
 #include "Widgets/Views/SListView.h"
 #include "Widgets/Text/STextBlock.h"
@@ -21,13 +25,22 @@ namespace UE::Editor::DataStorage
 		HeaderRowWidget = SNew( SHeaderRow )
 							.CanSelectGeneratedColumn(true);
 
+		// Create the internal Teds Widget
+		CreateTedsWidget();
+
+		// Attribute binder to bind widget columns to attributes on the ListView
+		EditorDataStorage::FAttributeBinder Binder(TedsWidget->GetRowHandle());
+
 		ListView = SNew(SListView<TableViewerItemPtr>)
 			.HeaderRow(HeaderRowWidget)
 			.ListItemsSource(&Model->GetItems())
 			.OnGenerateRow(this, &STedsTableViewer::MakeTableRowWidget)
 			.OnSelectionChanged(this, &STedsTableViewer::OnListSelectionChanged)
-			.SelectionMode(ESelectionMode::Single); // We only support single selection for now in the table viewer
-		
+			.SelectionMode(InArgs._ListSelectionMode)
+			.OnContextMenuOpening(Binder.BindEvent(&FWidgetContextMenuColumn::OnContextMenuOpening))
+			.OnItemScrolledIntoView(Binder.BindEvent(&FWidgetRowScrolledIntoView::OnItemScrolledIntoView))
+			.OnMouseButtonDoubleClick(Binder.BindEvent(&FWidgetDoubleClickedColumn::OnMouseButtonDoubleClick));
+
 		AssignChildSlot();
 		
 		// Add each Teds column from the model to our header row widget
@@ -44,41 +57,64 @@ namespace UE::Editor::DataStorage
 		});
 	}
 
+	void STedsTableViewer::CreateTedsWidget()
+	{
+		TedsWidget = SNew(STedsWidget);
+
+		if(ITypedElementDataStorageInterface* DataStorage = Model->GetDataStorageInterface())
+		{
+			const TypedElementDataStorage::RowHandle WidgetRowHandle = TedsWidget->GetRowHandle();
+		
+			if(DataStorage->IsRowAvailable(WidgetRowHandle))
+			{
+				// The table viewer should not show up as a row in a table viewer because that will cause all sorts of recursion issues
+				DataStorage->AddColumn(WidgetRowHandle, FHideRowFromUITag::StaticStruct());
+
+				// Columns we are going to bind to attributes on SListView
+				DataStorage->AddColumn(WidgetRowHandle, FWidgetContextMenuColumn::StaticStruct());
+				DataStorage->AddColumn(WidgetRowHandle, FWidgetRowScrolledIntoView::StaticStruct());
+				DataStorage->AddColumn(WidgetRowHandle, FWidgetDoubleClickedColumn::StaticStruct());
+			}
+		}
+
+		ChildSlot
+		[
+			TedsWidget.ToSharedRef()
+		];
+	}
+
 	void STedsTableViewer::AssignChildSlot()
 	{
+		TSharedPtr<SWidget> ContentWidget;
+		
 		if(Model->GetRowCount() == 0)
 		{
-			ChildSlot
-			[
+			ContentWidget = 
 				SNew(SBox)
 					.HAlign(HAlign_Center)
 					.VAlign(VAlign_Center)
 					[
 						SNew(STextBlock)
 							.Text(LOCTEXT("EmptyTableViewerQueryText", "The input query has no results"))
-					]
-			];
+					];
 		}
 		else if(Model->GetColumnCount() == 0)
 		{
-			ChildSlot
-			[
+			ContentWidget =
 				SNew(SBox)
 					.HAlign(HAlign_Center)
 					.VAlign(VAlign_Center)
 					[
 						SNew(STextBlock)
 							.Text(LOCTEXT("EmptyTableViewerColumnsText", "There were no columns specified to display"))
-					]
-			];
+					];
 		}
 		else
 		{
-			ChildSlot
-			[
-				ListView.ToSharedRef()	
-			];
+			ContentWidget = ListView.ToSharedRef();
 		}
+
+		TedsWidget->SetContent(ContentWidget.ToSharedRef());
 	}
 
 	void STedsTableViewer::RefreshColumnWidgets()
@@ -112,6 +148,22 @@ namespace UE::Editor::DataStorage
 		RefreshColumnWidgets();
 	}
 
+	void STedsTableViewer::ForEachSelectedRow(TFunctionRef<void(TypedElementDataStorage::RowHandle)> InCallback) const
+	{
+		TArray<TableViewerItemPtr> SelectedRows;
+		ListView->GetSelectedItems(SelectedRows);
+
+		for(TableViewerItemPtr& Row : SelectedRows)
+		{
+			InCallback(Row);
+		}
+	}
+
+	TypedElementDataStorage::RowHandle STedsTableViewer::GetWidgetRowHandle() const
+	{
+		return TedsWidget->GetRowHandle();
+	}
+
 	bool STedsTableViewer::IsItemVisible(TableViewerItemPtr InItem) const
 	{
 		return ListView->IsItemVisible(InItem);
diff --git a/Engine/Plugins/Experimental/EditorDataStorageFeatures/Source/TedsTableViewer/Public/TedsTableViewerModel.h b/Engine/Plugins/Experimental/EditorDataStorageFeatures/Source/TedsTableViewer/Public/TedsTableViewerModel.h
--- a/Engine/Plugins/Experimental/EditorDataStorageFeatures/Source/TedsTableViewer/Public/TedsTableViewerModel.h
+++ b/Engine/Plugins/Experimental/EditorDataStorageFeatures/Source/TedsTableViewer/Public/TedsTableViewerModel.h
@@ -19,7 +19,7 @@ namespace UE::Editor::DataStorage
 	class IQueryStackNode_Row;
 
 	// Typedef for an item in the table viewer
-	using TableViewerItemPtr = UIRowType;
+	using TableViewerItemPtr = FTedsRowHandle;
 
 	// Model class for the TEDS Table Viewer that can be plugged into any widget that is a UI representation of data in TEDS
 	// @see STedsTableViewer
@@ -59,6 +59,8 @@ namespace UE::Editor::DataStorage
 		// Add a custom column to display in the table viewer, that doesn't necessarily map to a Teds column
 		TEDSTABLEVIEWER_API void AddCustomColumn(const TSharedRef<FTedsTableViewerColumn>& InColumn);
 
+		TEDSTABLEVIEWER_API ITypedElementDataStorageInterface* GetDataStorageInterface() const;
+
 	protected:
 
 		// Generate the actual columns to display in the UI using TEDS UI
diff --git a/Engine/Plugins/Experimental/EditorDataStorageFeatures/Source/TedsTableViewer/Public/TypedElementUITypes.h b/Engine/Plugins/Experimental/EditorDataStorageFeatures/Source/TedsTableViewer/Public/TypedElementUITypes.h
--- a/Engine/Plugins/Experimental/EditorDataStorageFeatures/Source/TedsTableViewer/Public/TypedElementUITypes.h
+++ b/Engine/Plugins/Experimental/EditorDataStorageFeatures/Source/TedsTableViewer/Public/TypedElementUITypes.h
@@ -5,84 +5,53 @@
 #include "Elements/Common/TypedElementHandles.h"
 #include "Framework/Views/TableViewTypeTraits.h"
 
-namespace UE::Editor::DataStorage
-{
-	// Wrapper struct around RowHandle so we can template specialize for it without also specializing for uint64
-	struct UIRowType
-	{
-		RowHandle Row;
-		
-		UIRowType()
-			: Row(InvalidRowHandle)
-		{}
-
-		UIRowType(RowHandle InRowHandle)
-			: Row(InRowHandle)
-		{}
-
-		operator RowHandle() const
-		{
-			return Row;
-		}
 
-		UIRowType& operator=(RowHandle InRowHandle)
-		{
-			Row = InRowHandle;
-			return *this;
-		}
-		
-		friend uint32 GetTypeHash(const UIRowType& Key)
-		{
-			return GetTypeHash(Key.Row);
-		}
-
-	};
-}
 
 /* Template declaration to describe how a row handle behaves as a type for slate widgets like SListView, STreeView etc
  * This allows you to use Row Handles with slate widgets that work on pointers by using the wrapper struct e.g
- * SListView<UE::EditorDataStorage::UIRowType>
+ * SListView<FTedsRowHandle>
  */
 template <>
-struct TListTypeTraits<UE::Editor::DataStorage::UIRowType>
+struct TListTypeTraits<FTedsRowHandle>
 {
-	using NullableType = UE::Editor::DataStorage::UIRowType;
+	typedef FTedsRowHandle NullableType;
 
-	using MapKeyFuncs = TDefaultMapHashableKeyFuncs<UE::Editor::DataStorage::UIRowType, TSharedRef<ITableRow>, false>;
-	using MapKeyFuncsSparse = TDefaultMapHashableKeyFuncs<UE::Editor::DataStorage::UIRowType, FSparseItemInfo, false>;
-	using SetKeyFuncs = DefaultKeyFuncs<UE::Editor::DataStorage::UIRowType>;
+	using MapKeyFuncs = TDefaultMapHashableKeyFuncs<FTedsRowHandle, TSharedRef<ITableRow>, false>;
+	using MapKeyFuncsSparse = TDefaultMapHashableKeyFuncs<FTedsRowHandle, FSparseItemInfo, false>;
+	using SetKeyFuncs = DefaultKeyFuncs<FTedsRowHandle>;
 
 	template<typename U>
 	static void AddReferencedObjects(FReferenceCollector&,
-		TArray<UE::Editor::DataStorage::UIRowType>&,
-		TSet<UE::Editor::DataStorage::UIRowType>&,
-		TMap<const U*, UE::Editor::DataStorage::UIRowType>&)
+		TArray<FTedsRowHandle>&,
+		TSet<FTedsRowHandle>&,
+		TMap<const U*, FTedsRowHandle>&)
 	{
 	}
 
-	static bool IsPtrValid(const UE::Editor::DataStorage::UIRowType& InPtr)
+	static bool IsPtrValid(const FTedsRowHandle& InPtr)
 	{
 		return InPtr != UE::Editor::DataStorage::InvalidRowHandle;
 	}
 
-	static void ResetPtr(UE::Editor::DataStorage::UIRowType& InPtr)
+	static void ResetPtr(FTedsRowHandle& InPtr)
 	{
 		InPtr = UE::Editor::DataStorage::InvalidRowHandle;
 	}
 
-	static UE::Editor::DataStorage::UIRowType MakeNullPtr()
+	static FTedsRowHandle MakeNullPtr()
 	{
-		return UE::Editor::DataStorage::InvalidRowHandle;
+		static FTedsRowHandle InvalidRowHandle{ .RowHandle = UE::Editor::DataStorage::InvalidRowHandle };
+		return InvalidRowHandle;
 	}
 
-	static UE::Editor::DataStorage::UIRowType NullableItemTypeConvertToItemType(const UE::Editor::DataStorage::UIRowType& InPtr)
+	static FTedsRowHandle NullableItemTypeConvertToItemType(const FTedsRowHandle& InPtr)
 	{
 		return InPtr;
 	}
 
-	static FString DebugDump(UE::Editor::DataStorage::UIRowType InPtr)
+	static FString DebugDump(FTedsRowHandle InPtr)
 	{
-		return FString::Printf(TEXT("%llu"), InPtr.Row);
+		return FString::Printf(TEXT("%llu"), InPtr.RowHandle);
 	}
 
 	class SerializerType {};
@@ -90,10 +59,10 @@ struct TListTypeTraits<UE::Editor::DataStorage::UIRowType>
 
 // Template declaration to enable using row handles inside of slate widgets like SListView
 template <>
-struct TIsValidListItem<UE::Editor::DataStorage::UIRowType>
+struct TIsValidListItem<FTedsRowHandle>
 {
 	enum
 	{
 		Value = true
 	};
-};
+};
\ No newline at end of file
diff --git a/Engine/Plugins/Experimental/EditorDataStorageFeatures/Source/TedsTableViewer/Public/Widgets/STedsTableViewer.h b/Engine/Plugins/Experimental/EditorDataStorageFeatures/Source/TedsTableViewer/Public/Widgets/STedsTableViewer.h
--- a/Engine/Plugins/Experimental/EditorDataStorageFeatures/Source/TedsTableViewer/Public/Widgets/STedsTableViewer.h
+++ b/Engine/Plugins/Experimental/EditorDataStorageFeatures/Source/TedsTableViewer/Public/Widgets/STedsTableViewer.h
@@ -11,6 +11,7 @@
 class ITableRow;
 class STableViewBase;
 class SHeaderRow;
+class STedsWidget;
 
 namespace UE::Editor::DataStorage
 {
@@ -30,10 +31,11 @@ namespace UE::Editor::DataStorage
 	public:
 		
 		// Delegate fired when the selection in the table viewer changes
-		DECLARE_DELEGATE_OneParam(FOnSelectionChanged, RowHandle)
+		DECLARE_DELEGATE_OneParam(FOnSelectionChanged, UE::Editor::DataStorage::RowHandle)
 
 		SLATE_BEGIN_ARGS(STedsTableViewer)
 			: _CellWidgetPurposes({TEXT("General.Cell")})
+			, _ListSelectionMode(ESelectionMode::Type::Single)
 		{
 			
 		}
@@ -50,7 +52,10 @@ namespace UE::Editor::DataStorage
 		SLATE_ARGUMENT(TArray<FName>, CellWidgetPurposes)
 		
 		// Delegate called when the selection changes
-		SLATE_ARGUMENT(FOnSelectionChanged, OnSelectionChanged)
+		SLATE_EVENT(FOnSelectionChanged, OnSelectionChanged)
+
+		// The selection mode for the table viewer (single/multi etc)
+		SLATE_ARGUMENT(ESelectionMode::Type, ListSelectionMode)
 
 		SLATE_END_ARGS()
 
@@ -64,6 +69,11 @@ namespace UE::Editor::DataStorage
 		// Add a custom column to display in the table viewer, that doesn't necessarily map to a Teds column
 		TEDSTABLEVIEWER_API void AddCustomColumn(const TSharedRef<FTedsTableViewerColumn>& InColumn);
 
+		// Execute the given callback for each row that is selected in the table viewer
+		TEDSTABLEVIEWER_API void ForEachSelectedRow(TFunctionRef<void(UE::Editor::DataStorage::RowHandle)> InCallback) const;
+
+		// Get the row handle for the widget row the table viewer's contents are stored in
+		TEDSTABLEVIEWER_API UE::Editor::DataStorage::RowHandle GetWidgetRowHandle() const;
 
 	protected:
 		
@@ -77,6 +87,8 @@ namespace UE::Editor::DataStorage
 
 		void OnListSelectionChanged(TableViewerItemPtr Item, ESelectInfo::Type SelectInfo);
 
+		void CreateTedsWidget();
+
 	private:
 
 		// The actual ListView widget that displays the rows
@@ -90,5 +102,8 @@ namespace UE::Editor::DataStorage
 
 		// Delegate fired when the selection changes
 		FOnSelectionChanged OnSelectionChanged;
+
+		// Wrapper Teds Widget around our contents so we can use Teds columns to specify behavior
+		TSharedPtr<STedsWidget> TedsWidget;
 	};
-} // namespace UE::Editor::DataStorage
\ No newline at end of file
+}
\ No newline at end of file
diff --git a/Engine/Plugins/Experimental/EditorDataStorageFeatures/Source/TedsTableViewer/TedsTableViewer.Build.cs b/Engine/Plugins/Experimental/EditorDataStorageFeatures/Source/TedsTableViewer/TedsTableViewer.Build.cs
--- a/Engine/Plugins/Experimental/EditorDataStorageFeatures/Source/TedsTableViewer/TedsTableViewer.Build.cs
+++ b/Engine/Plugins/Experimental/EditorDataStorageFeatures/Source/TedsTableViewer/TedsTableViewer.Build.cs
@@ -27,7 +27,8 @@ public class TedsTableViewer : ModuleRules
 					"InputCore",
 					"Slate",
 					"SlateCore",
-					"WorkspaceMenuStructure"
+					"WorkspaceMenuStructure",
+					"ToolWidgets"
 				});
 		}
 	}
diff --git a/Engine/Source/Developer/ToolWidgets/Public/Columns/SlateDelegateColumns.h b/Engine/Source/Developer/ToolWidgets/Public/Columns/SlateDelegateColumns.h
--- a/Engine/Source/Developer/ToolWidgets/Public/Columns/SlateDelegateColumns.h
+++ b/Engine/Source/Developer/ToolWidgets/Public/Columns/SlateDelegateColumns.h
@@ -23,4 +23,38 @@ struct FExternalWidgetSelectionColumn : public FEditorDataStorageColumn
 	 * @see SInlineEditableTextBlock::IsSelected
 	 */
 	FIsSelected IsSelected;
+};
+
+/**
+ * Column added to a widget row to allow systems to provide a context menu for the widget
+ */
+USTRUCT(meta = (DisplayName = "Widget with context menu"))
+struct FWidgetContextMenuColumn : public FEditorDataStorageColumn
+{
+	GENERATED_BODY()
+	FOnContextMenuOpening OnContextMenuOpening;
+};
+
+/**
+ * Column added to a widget row for table views to allow systems to specify behavior when a row is scrolled into view
+ */
+USTRUCT(meta = (DisplayName = "Widget with row scrolled into view delegate"))
+struct FWidgetRowScrolledIntoView : public FEditorDataStorageColumn
+{
+	using FOnItemScrolledIntoView = TSlateDelegates<FTedsRowHandle>::FOnItemScrolledIntoView;
+
+	GENERATED_BODY()
+	FOnItemScrolledIntoView OnItemScrolledIntoView;
+};
+
+/**
+ * Column added to a widget row to allow systems to specify behavior when the widget is double clicked
+ */
+USTRUCT(meta = (DisplayName = "Widget with item scrolled into view delegate"))
+struct FWidgetDoubleClickedColumn : public FEditorDataStorageColumn
+{
+	using FOnMouseButtonDoubleClick = TSlateDelegates<FTedsRowHandle>::FOnMouseButtonDoubleClick;
+	
+	GENERATED_BODY()
+	FOnMouseButtonDoubleClick OnMouseButtonDoubleClick;
 };
\ No newline at end of file
diff --git a/Engine/Source/Runtime/TypedElementFramework/Private/Elements/Framework/TypedElementAttributeBinding.cpp b/Engine/Source/Runtime/TypedElementFramework/Private/Elements/Framework/TypedElementAttributeBinding.cpp
--- a/Engine/Source/Runtime/TypedElementFramework/Private/Elements/Framework/TypedElementAttributeBinding.cpp
+++ b/Engine/Source/Runtime/TypedElementFramework/Private/Elements/Framework/TypedElementAttributeBinding.cpp
@@ -10,9 +10,17 @@ namespace UE::EditorDataStorage
 		: TargetRow(InTargetRow)
 	{
 		if (ensureMsgf(UTypedElementRegistry::GetInstance()->AreDataStorageInterfacesSet(), 
-				TEXT("The TypedElementsDataStorage plugin needs to be enabled to use attribute bindings.")))
+				TEXT("The Editor Data Storage plugin needs to be enabled to use attribute bindings.")))
 		{
 			DataStorage = UTypedElementRegistry::GetInstance()->GetMutableDataStorage();
 		}
 	}
+
+	FAttributeBinder::FAttributeBinder(
+		TypedElementDataStorage::RowHandle InTargetRow, ITypedElementDataStorageInterface* InDataStorage)
+		: TargetRow(InTargetRow)
+		, DataStorage(InDataStorage)
+	{
+		ensureMsgf(DataStorage, TEXT("The Editor Data Storage plugin needs to be enabled to use attribute bindings."));
+	}
 }
diff --git a/Engine/Source/Runtime/TypedElementFramework/Private/Elements/Framework/TypedElementDataStorageWidget.cpp b/Engine/Source/Runtime/TypedElementFramework/Private/Elements/Framework/TypedElementDataStorageWidget.cpp
--- a/Engine/Source/Runtime/TypedElementFramework/Private/Elements/Framework/TypedElementDataStorageWidget.cpp
+++ b/Engine/Source/Runtime/TypedElementFramework/Private/Elements/Framework/TypedElementDataStorageWidget.cpp
@@ -1,7 +1,8 @@
 // Copyright Epic Games, Inc. All Rights Reserved.
 
 #include "Elements/Framework/TypedElementDataStorageWidget.h"
-#include "Elements/Interfaces/TypedElementDataStorageUiInterface.h"
+#include "Elements/Columns/TypedElementSlateWidgetColumns.h"
+#include "Elements/Framework/TypedElementRegistry.h"
 
 STedsWidget::STedsWidget()
 	: UiRowHandle(TypedElementDataStorage::InvalidRowHandle)
@@ -13,16 +14,72 @@ void STedsWidget::Construct(const FArguments& InArgs)
 {
 	UiRowHandle = InArgs._UiRowHandle;
 
+	// If the Ui row wasn't already registered externally, register it with Teds
+	if(UiRowHandle == TypedElementDataStorage::InvalidRowHandle)
+	{
+		RegisterTedsWidget(InArgs._Content.Widget);
+	}
+
 	ChildSlot
 	[
 		InArgs._Content.Widget
 	];
 }
 
+void STedsWidget::RegisterTedsWidget(const TSharedPtr<SWidget>& InContentWidget)
+{
+	ITypedElementDataStorageInterface* Storage = GetStorageIfAvailable();
+
+	// If TEDS is not enabled, STedsWidget will just behave like a regular widget
+	if(!Storage)
+	{
+		return;
+	}
+	
+	const TypedElementDataStorage::TableHandle WidgetTable = Storage->FindTable(TEXT("Editor_WidgetTable"));
+	if(WidgetTable == TypedElementDataStorage::InvalidTableHandle)
+	{
+		return;
+	}
+	
+	UiRowHandle = Storage->AddRow(WidgetTable);
+	
+	if(FTypedElementSlateWidgetReferenceColumn* WidgetReferenceColumn = Storage->GetColumn<FTypedElementSlateWidgetReferenceColumn>(UiRowHandle))
+	{
+		WidgetReferenceColumn->TedsWidget = SharedThis(this);
+		WidgetReferenceColumn->Widget = InContentWidget;
+	}
+}
+
 void STedsWidget::SetContent(const TSharedRef< SWidget >& InContent)
 {
+	if(ITypedElementDataStorageInterface* Storage = GetStorageIfAvailable())
+	{
+		if(FTypedElementSlateWidgetReferenceColumn* WidgetReferenceColumn = Storage->GetColumn<FTypedElementSlateWidgetReferenceColumn>(UiRowHandle))
+		{
+			WidgetReferenceColumn->Widget = InContent;
+		}
+	}
+	
 	ChildSlot
 	[
 		InContent
 	];
-}
\ No newline at end of file
+}
+
+TypedElementDataStorage::RowHandle STedsWidget::GetRowHandle() const
+{
+	return UiRowHandle;
+}
+
+ITypedElementDataStorageInterface* STedsWidget::GetStorageIfAvailable()
+{
+	UTypedElementRegistry* Registry = UTypedElementRegistry::GetInstance();
+
+	if(!Registry || !Registry->AreDataStorageInterfacesSet())
+	{
+		return nullptr;
+	}
+
+	return Registry->GetMutableDataStorage();
+}
diff --git a/Engine/Source/Runtime/TypedElementFramework/Private/Elements/Interfaces/TypedElementDataStorageUiInterface.cpp b/Engine/Source/Runtime/TypedElementFramework/Private/Elements/Interfaces/TypedElementDataStorageUiInterface.cpp
--- a/Engine/Source/Runtime/TypedElementFramework/Private/Elements/Interfaces/TypedElementDataStorageUiInterface.cpp
+++ b/Engine/Source/Runtime/TypedElementFramework/Private/Elements/Interfaces/TypedElementDataStorageUiInterface.cpp
@@ -4,8 +4,8 @@
 
 #include "Elements/Columns/TypedElementLabelColumns.h"
 #include "Elements/Columns/TypedElementMiscColumns.h"
-#include "Elements/Framework/TypedElementColumnUtils.h"
 #include "Elements/Columns/TypedElementSlateWidgetColumns.h"
+#include "Elements/Columns/TypedElementUIColumns.h"
 #include "Elements/Framework/TypedElementDataStorageWidget.h"
 
 FTypedElementWidgetConstructor::FTypedElementWidgetConstructor(const UScriptStruct* InTypeInfo)
@@ -141,7 +141,7 @@ TSharedPtr<SWidget> FTypedElementWidgetConstructor::Construct(
 	ITypedElementDataStorageUiInterface* DataStorageUi,
 	const TypedElementDataStorage::FMetaDataView& Arguments)
 {
-	TSharedPtr<SWidget> Widget = CreateWidget(Arguments);
+	TSharedPtr<SWidget> Widget = CreateWidget(DataStorage, DataStorageUi, Row, Arguments);
 	if (Widget)
 	{
 		DataStorage->GetColumn<FTypedElementSlateWidgetReferenceColumn>(Row)->Widget = Widget;
@@ -162,6 +162,15 @@ TSharedPtr<SWidget> FTypedElementWidgetConstructor::CreateWidget(const TypedElem
 	return nullptr;
 }
 
+TSharedPtr<SWidget> FTypedElementWidgetConstructor::CreateWidget(
+	ITypedElementDataStorageInterface* DataStorage,
+	ITypedElementDataStorageUiInterface* DataStorageUi,
+	TypedElementDataStorage::RowHandle UiRow, 
+	const TypedElementDataStorage::FMetaDataView& Arguments)
+{
+	return CreateWidget(Arguments);
+}
+
 bool FTypedElementWidgetConstructor::SetColumns(ITypedElementDataStorageInterface* DataStorage, RowHandle Row)
 {
 	return true;
@@ -197,4 +206,17 @@ void FTypedElementWidgetConstructor::AddDefaultWidgetColumns(RowHandle Row, ITyp
 {
 	const FString WidgetLabel(CreateWidgetDisplayName(DataStorage, Row));
 	DataStorage->AddColumn(Row, FTypedElementLabelColumn{.Label = WidgetLabel} );
+
+	// We don't want to display any second level widgets (widgets for widgets and so on...) in UI because they will cause the table viewer to
+	// infinitely grow as you keep scrolling (which creates new widgets)
+	if(DataStorage->HasColumns<FTypedElementSlateWidgetReferenceColumn>(Row))
+	{
+		if(const FTypedElementRowReferenceColumn* RowReferenceColumn = DataStorage->GetColumn<FTypedElementRowReferenceColumn>(Row))
+		{
+			if(DataStorage->HasColumns<FTypedElementSlateWidgetReferenceColumn>(RowReferenceColumn->Row))
+			{
+				DataStorage->AddColumn(Row, FHideRowFromUITag::StaticStruct());
+			}
+		}
+	}
 }
diff --git a/Engine/Source/Runtime/TypedElementFramework/Public/Elements/Columns/TypedElementUIColumns.h b/Engine/Source/Runtime/TypedElementFramework/Public/Elements/Columns/TypedElementUIColumns.h
new file mode 100644
--- /dev/null
+++ b/Engine/Source/Runtime/TypedElementFramework/Public/Elements/Columns/TypedElementUIColumns.h
@@ -0,0 +1,17 @@
+// Copyright Epic Games, Inc. All Rights Reserved.
+
+#pragma once
+
+#include "Elements/Common/TypedElementCommonTypes.h"
+#include "UObject/ObjectMacros.h"
+
+#include "TypedElementUIColumns.generated.h"
+
+/**
+ * Tag to indicate that this row should not show up in any TEDS UI that displays rows (e.g table viewer)
+ */
+USTRUCT(meta = (DisplayName = "Hide row from UI"))
+struct FHideRowFromUITag final : public FEditorDataStorageTag
+{
+	GENERATED_BODY()
+};
\ No newline at end of file
diff --git a/Engine/Source/Runtime/TypedElementFramework/Public/Elements/Common/TypedElementHandles.h b/Engine/Source/Runtime/TypedElementFramework/Public/Elements/Common/TypedElementHandles.h
--- a/Engine/Source/Runtime/TypedElementFramework/Public/Elements/Common/TypedElementHandles.h
+++ b/Engine/Source/Runtime/TypedElementFramework/Public/Elements/Common/TypedElementHandles.h
@@ -3,6 +3,9 @@
 #pragma once
 
 #include "Math/NumericLimits.h"
+#include "Templates/TypeHash.h"
+
+#include "TypedElementHandles.generated.h"
 
 namespace UE::Editor::DataStorage
 {
@@ -15,3 +18,35 @@ namespace UE::Editor::DataStorage
 	using QueryHandle = uint64;
 	static constexpr QueryHandle InvalidQueryHandle = TNumericLimits<QueryHandle>::Max();
 } // namespace TypedElementDataStorage
+
+/*
+ * FTedsRowHandle is a strongly typed wrapper around TypedElementDataStorage::RowHandle and should only be used in cases where you need the extra info.
+ * E.g for reflection/UHT or for template specializing something that needs to know the semantics of the row handle.
+ * For all other cases, you should use the regular typedef TypedElementDataStorage::RowHandle
+ */
+USTRUCT()
+struct FTedsRowHandle
+{
+	GENERATED_BODY()
+
+	UPROPERTY()
+	uint64 RowHandle;
+
+	operator UE::Editor::DataStorage::RowHandle () const
+	{
+		return RowHandle;
+	}
+
+	FTedsRowHandle& operator=(UE::Editor::DataStorage::RowHandle InRowHandle)
+	{
+		RowHandle = InRowHandle;
+		return *this;
+	}
+		
+	friend uint32 GetTypeHash(const FTedsRowHandle& Key)
+	{
+		return GetTypeHash(Key.RowHandle);
+	}
+};
+
+static_assert(sizeof(FTedsRowHandle::RowHandle) == sizeof(UE::Editor::DataStorage::RowHandle), "RowHandle and RowHandle wrapper sizes should match");
\ No newline at end of file
diff --git a/Engine/Source/Runtime/TypedElementFramework/Public/Elements/Framework/TypedElementAttributeBinding.h b/Engine/Source/Runtime/TypedElementFramework/Public/Elements/Framework/TypedElementAttributeBinding.h
--- a/Engine/Source/Runtime/TypedElementFramework/Public/Elements/Framework/TypedElementAttributeBinding.h
+++ b/Engine/Source/Runtime/TypedElementFramework/Public/Elements/Framework/TypedElementAttributeBinding.h
@@ -29,8 +29,11 @@ namespace UE::EditorDataStorage
 	{
 	public:
 
-		// Create an attribute binder for a given row
-		TYPEDELEMENTFRAMEWORK_API FAttributeBinder(TypedElementDataStorage::RowHandle InTargetRow);
+		/* Create an attribute binder for a given row. */
+		TYPEDELEMENTFRAMEWORK_API explicit FAttributeBinder(TypedElementDataStorage::RowHandle InTargetRow);
+		
+		/* Create an attribute binder for a given row. */
+		TYPEDELEMENTFRAMEWORK_API FAttributeBinder(TypedElementDataStorage::RowHandle InTargetRow, ITypedElementDataStorageInterface* InDataStorage);
 
 		/**
 		 * Bind a specific data member inside a TEDS column to an attribute of the same type as the data
@@ -78,6 +81,14 @@ namespace UE::EditorDataStorage
 		template <typename DataType, TDataColumnType ColumnType, typename FunctionType>
 			requires AttributeBinderInvocable<FunctionType, DataType>
 		auto BindData(DataType ColumnType::* InVariable, FunctionType InConverter, const DataType& InDefaultValue = DataType());
+
+		/**
+		 * Bind a delegate inside a Teds column to a SLATE_EVENT macro on a widget
+		 * @param InVariable The delegate inside the Teds column
+		 * @return A delegate that can be provided to an event on a slate widget
+		 */
+		template <typename InRetValType, typename... ParamTypes, typename ColumnType>
+		TDelegate<InRetValType(ParamTypes...)> BindEvent(TDelegate<InRetValType(ParamTypes...)> ColumnType::* InVariable);
 		
 	private:
 
@@ -89,5 +100,4 @@ namespace UE::EditorDataStorage
 	};
 }
 
-
 #include "TypedElementAttributeBinding.inl"
\ No newline at end of file
diff --git a/Engine/Source/Runtime/TypedElementFramework/Public/Elements/Framework/TypedElementAttributeBinding.inl b/Engine/Source/Runtime/TypedElementFramework/Public/Elements/Framework/TypedElementAttributeBinding.inl
--- a/Engine/Source/Runtime/TypedElementFramework/Public/Elements/Framework/TypedElementAttributeBinding.inl
+++ b/Engine/Source/Runtime/TypedElementFramework/Public/Elements/Framework/TypedElementAttributeBinding.inl
@@ -134,7 +134,7 @@ namespace UE::EditorDataStorage
 		Prop.Bind(InVariable);
 
 		// We don't want any references to this in the lambda because binders are designed to be used and destructed on the stack
-		return TAttribute<AttributeType>::CreateLambda([Property = Prop, Storage = DataStorage, Row = TargetRow, DefaultValue = InDefaultValue]()
+		return TAttribute<AttributeType>::CreateLambda([Property = MoveTemp(Prop), Storage = DataStorage, Row = TargetRow, DefaultValue = InDefaultValue]()
 		{
 			// Get the column from the given row and use that to return the stored property
 			if(ColumnType* Column = Storage->GetColumn<ColumnType>(Row))
@@ -157,7 +157,7 @@ namespace UE::EditorDataStorage
 		Property<AttributeType> Prop;
 		Prop.Bind(InVariable, InConverter);
 	
-		return TAttribute<AttributeType>::CreateLambda([Property = Prop, Storage = DataStorage, Row = TargetRow, DefaultValue = InDefaultValue, Converter = InConverter]()
+		return TAttribute<AttributeType>::CreateLambda([Property = MoveTemp(Prop), Storage = DataStorage, Row = TargetRow, DefaultValue = InDefaultValue, Converter = InConverter]()
 		{
 			if(ColumnType* Column = Storage->GetColumn<ColumnType>(Row))
 			{
@@ -177,6 +177,30 @@ namespace UE::EditorDataStorage
 		
 		return BindData<AttributeType, DataType>(InVariable, TFunction<AttributeType(const DataType&)>(InConverter), InDefaultValue);
 	}
+
+	template <typename InRetValType, typename... ParamTypes, typename ColumnType>
+		TDelegate<InRetValType(ParamTypes...)> FAttributeBinder::BindEvent(TDelegate<InRetValType(ParamTypes...)> ColumnType::* InVariable)
+	{
+		// Create a property for the delegate
+		Property<TDelegate<InRetValType(ParamTypes...)>> Prop;
+		Prop.Bind(InVariable);
+	
+		return TDelegate<InRetValType(ParamTypes...)>::CreateLambda([Property = MoveTemp(Prop), Storage = DataStorage, Row = TargetRow](ParamTypes&&... Params)
+		{
+			if(ColumnType* Column = Storage->GetColumn<ColumnType>(Row))
+			{
+				// Get the delegate in the bound column for the specified row
+				TDelegate<InRetValType(ParamTypes...)> Delegate = Property.Get(Column, ColumnType::StaticStruct());
+
+				// Execute the delegate if it is bound
+				if(Delegate.IsBound())
+				{
+					return Delegate.Execute(Forward<ParamTypes>(Params)...);
+				}
+			}
+			return InRetValType();
+		});
+	}
 	
 }
 
diff --git a/Engine/Source/Runtime/TypedElementFramework/Public/Elements/Framework/TypedElementDataStorageWidget.h b/Engine/Source/Runtime/TypedElementFramework/Public/Elements/Framework/TypedElementDataStorageWidget.h
--- a/Engine/Source/Runtime/TypedElementFramework/Public/Elements/Framework/TypedElementDataStorageWidget.h
+++ b/Engine/Source/Runtime/TypedElementFramework/Public/Elements/Framework/TypedElementDataStorageWidget.h
@@ -33,13 +33,20 @@ public:
 	
 	SLATE_END_ARGS()
 	
-	STedsWidget();
+	TYPEDELEMENTFRAMEWORK_API STedsWidget();
 
-	void Construct( const FArguments& InArgs );
+	TYPEDELEMENTFRAMEWORK_API void Construct( const FArguments& InArgs );
 	
 	TYPEDELEMENTFRAMEWORK_API void SetContent(const TSharedRef< SWidget >& InContent);
 
+	TYPEDELEMENTFRAMEWORK_API TypedElementDataStorage::RowHandle GetRowHandle() const;
+
+private:
+
+	void RegisterTedsWidget(const TSharedPtr<SWidget>& InContentWidget);
 	
+	static ITypedElementDataStorageInterface* GetStorageIfAvailable();
+
 private:
 	
     TypedElementDataStorage::RowHandle UiRowHandle;
diff --git a/Engine/Source/Runtime/TypedElementFramework/Public/Elements/Interfaces/TypedElementDataStorageUiInterface.h b/Engine/Source/Runtime/TypedElementFramework/Public/Elements/Interfaces/TypedElementDataStorageUiInterface.h
--- a/Engine/Source/Runtime/TypedElementFramework/Public/Elements/Interfaces/TypedElementDataStorageUiInterface.h
+++ b/Engine/Source/Runtime/TypedElementFramework/Public/Elements/Interfaces/TypedElementDataStorageUiInterface.h
@@ -86,6 +86,12 @@ public:
 protected:
 	/** Create a new instance of the target widget. This is a required function. */
 	TYPEDELEMENTFRAMEWORK_API virtual TSharedPtr<SWidget> CreateWidget(const TypedElementDataStorage::FMetaDataView& Arguments);
+	/** Create a new instance of the target widget. This is a required function. */
+	TYPEDELEMENTFRAMEWORK_API virtual TSharedPtr<SWidget> CreateWidget(
+		ITypedElementDataStorageInterface* DataStorage,
+		ITypedElementDataStorageUiInterface* DataStorageUi,
+		TypedElementDataStorage::RowHandle UiRow, 
+		const TypedElementDataStorage::FMetaDataView& Arguments);
 	/** Set any values in columns if needed. The columns provided through GetAdditionalColumnsList() will have already been created. */
 	TYPEDELEMENTFRAMEWORK_API virtual bool SetColumns(ITypedElementDataStorageInterface* DataStorage, RowHandle Row);
 	
