Adding the AssetClassColumn to the TEDS CB tables.

diff --git a/Engine/Plugins/Experimental/EditorDataStorageFeatures/Source/TedsAssetData/Private/CB/TedsAssetDataCBDataSource.cpp b/Engine/Plugins/Experimental/EditorDataStorageFeatures/Source/TedsAssetData/Private/CB/TedsAssetDataCBDataSource.cpp
--- a/Engine/Plugins/Experimental/EditorDataStorageFeatures/Source/TedsAssetData/Private/CB/TedsAssetDataCBDataSource.cpp
+++ b/Engine/Plugins/Experimental/EditorDataStorageFeatures/Source/TedsAssetData/Private/CB/TedsAssetDataCBDataSource.cpp
@@ -91,7 +91,7 @@ FTedsAssetDataCBDataSource::FTedsAssetDataCBDataSource(ITypedElementDataStorageI
 			}
 		)
 		.Where()
-			.All<FUpdatedPathTag>()
+				.All<FUpdatedPathTag>()
 			.None<FUpdatedAssetDataTag>()
 		.Compile()
 		);
@@ -119,6 +119,11 @@ FTedsAssetDataCBDataSource::FTedsAssetDataCBDataSource(ITypedElementDataStorageI
 			InContext.AddColumn(Row, MoveTemp(DiskSizeColumn));
 		}
 
+		// Should add those or should we develop tooling to query the asset data via Teds?
+		FAssetClassColumn AssetClassColumn;
+		AssetClassColumn.ClassPath = InAssetData.AssetClassPath;
+		InContext.AddColumn(Row, MoveTemp(AssetClassColumn));
+
 		FItemNameColumn_Experimental ItemNameColumn;
 		ItemNameColumn.Name = InAssetData.AssetName;
 		InContext.AddColumn(Row, MoveTemp(ItemNameColumn));
diff --git a/Engine/Plugins/Experimental/EditorDataStorageFeatures/Source/TedsAssetData/Public/TedsAssetDataColumns.h b/Engine/Plugins/Experimental/EditorDataStorageFeatures/Source/TedsAssetData/Public/TedsAssetDataColumns.h
--- a/Engine/Plugins/Experimental/EditorDataStorageFeatures/Source/TedsAssetData/Public/TedsAssetDataColumns.h
+++ b/Engine/Plugins/Experimental/EditorDataStorageFeatures/Source/TedsAssetData/Public/TedsAssetDataColumns.h
@@ -88,6 +88,17 @@ struct FPublicAssetTag final : public FEditorDataStorageTag
 	GENERATED_BODY()
 };
 
+// Column to store the type of an asset
+USTRUCT(meta = (DisplayName = "Type Path"))
+struct FAssetClassColumn final : public FEditorDataStorageColumn
+{
+	GENERATED_BODY()
+
+	// The path of the type
+	UPROPERTY()
+	FTopLevelAssetPath ClassPath;
+};
+
 // Column to store the disk size of an asset
 USTRUCT(meta = (DisplayName = "Disk Size"))
 struct FDiskSizeColumn final : public FEditorDataStorageColumn
@@ -99,6 +110,7 @@ struct FDiskSizeColumn final : public FEditorDataStorageColumn
 	int64 DiskSize = 0;
 };
 
+
 // Column to store the verse path of an asset
 USTRUCT(meta = (DisplayName = "Verse Path"))
 struct FVersePathColumn final : public FEditorDataStorageColumn
