Universal Search prototype - Removed searching asset paths.

Removed the ability to search assets paths. The current use of asset paths don't provide any information that users will find usable or actionable as the data is used behind the scenes.

diff --git a/Engine/Plugins/Experimental/EditorDataStorageFeatures/Source/TedsAssetData/Public/TedsAssetDataColumns.h b/Engine/Plugins/Experimental/EditorDataStorageFeatures/Source/TedsAssetData/Public/TedsAssetDataColumns.h
--- a/Engine/Plugins/Experimental/EditorDataStorageFeatures/Source/TedsAssetData/Public/TedsAssetDataColumns.h
+++ b/Engine/Plugins/Experimental/EditorDataStorageFeatures/Source/TedsAssetData/Public/TedsAssetDataColumns.h
@@ -14,7 +14,7 @@ struct FAssetPathColumn_Experimental : public FEditorDataStorageColumn
 {
 	GENERATED_BODY()
 	
-	UPROPERTY(meta = (Searchable))
+	UPROPERTY()
 	FName Path;
 };
 
