TEDS Debugger:
Remove "Table Viewer" tab completely
Add better messaging when there are no columns present

diff --git a/Engine/Plugins/Experimental/EditorDataStorageFeatures/Source/TedsDebugger/Private/STedsDebugger.cpp b/Engine/Plugins/Experimental/EditorDataStorageFeatures/Source/TedsDebugger/Private/STedsDebugger.cpp
--- a/Engine/Plugins/Experimental/EditorDataStorageFeatures/Source/TedsDebugger/Private/STedsDebugger.cpp
+++ b/Engine/Plugins/Experimental/EditorDataStorageFeatures/Source/TedsDebugger/Private/STedsDebugger.cpp
@@ -21,7 +21,6 @@ namespace UE::Editor::DataStorage::Debug
 	namespace Private
 	{
 		FName QueryEditorToolTabName = TEXT("TEDS Query Editor");
-		FName TableViewerToolTabName = TEXT("TEDS Table Viewer");
 		FName ToolbarTabName = TEXT("TEDS Debugger Toolbar");
 	}
 
@@ -64,7 +63,6 @@ void STedsDebugger::Construct(const FArguments& InArgs, const TSharedRef<SDockTa
 			(
 				FTabManager::NewStack()
 					->AddTab(Private::QueryEditorToolTabName, ETabState::OpenedTab)
-					->AddTab(Private::TableViewerToolTabName, ETabState::OpenedTab)
 			)
 		)
 	);
@@ -157,53 +155,6 @@ TSharedRef<SDockTab> STedsDebugger::SpawnQueryEditorTab(const FSpawnTabArgs& Arg
 	return DockTab;
 }
 
-TSharedRef<SDockTab> STedsDebugger::SpawnTableViewerTab(const FSpawnTabArgs& Args)
-{
-	UTypedElementRegistry* Registry = UTypedElementRegistry::GetInstance();
-	checkf(Registry, TEXT("Unable to initialize the table viewer before TEDS is initialized."));
-
-	using namespace TypedElementQueryBuilder;
-
-	// The TEDS-Debugger will show all rows with a label
-	TypedElementDataStorage::FQueryDescription RowQueryDescription =
-						Select()
-						.Where()
-							.All<FTypedElementLabelColumn>()
-						.Compile();
-
-	// TEDS-Debugger TODO: Currently uses a pre-determined initial set of columns, how can we drive this by the rows shown or let the user pick?
-	TypedElementDataStorage::FQueryDescription ColumnQueryDescription =
-						Select()
-							.ReadOnly<FTypedElementClassTypeInfoColumn, FTypedElementSelectionColumn, FTypedElementRowReferenceColumn>()
-						.Compile();
-
-	TableViewerQuery = Registry->GetMutableDataStorage()->RegisterQuery(MoveTemp(ColumnQueryDescription));
-
-	FSceneOutlinerInitializationOptions InitOptions;
-	InitOptions.bShowHeaderRow = true;
-	InitOptions.FilterBarOptions.bHasFilterBar = true;
-	InitOptions.OutlinerIdentifier = "TedsDebugger.TableViewer";
-
-	FTedsOutlinerParams Params(nullptr);
-	Params.QueryDescription = RowQueryDescription;
-	Params.bUseDefaultTedsFilters = true;
-	Params.HierarchyData = TOptional<FTedsOutlinerHierarchyData>(); // We don't want to show hierarchies in the debugger
-	Params.CellWidgetPurposes = TArray<FName>{TEXT("General.Cell")};
-	
-	FTedsOutlinerModule& TedsOutlinerModule = FModuleManager::GetModuleChecked<FTedsOutlinerModule>("TedsOutliner");
-	
-	TSharedRef<ISceneOutliner> TedsOutliner = TedsOutlinerModule.CreateTedsOutliner(InitOptions, Params, TableViewerQuery);
-
-	// Store an instance of the table viewer
-	TableViewerInstance = TedsOutliner;
-	
-	return SNew(SDockTab)
-		.TabRole(ETabRole::NomadTab)
-		[
-			TedsOutliner
-		];
-}
-
 void STedsDebugger::RegisterTabSpawners()
 {
 	const TSharedRef<FWorkspaceItem> AppMenuGroup =
@@ -223,60 +174,8 @@ void STedsDebugger::RegisterTabSpawners()
 		.SetDisplayName(LOCTEXT("TedsDebugger_QueryEditorDisplayName", "Query Editor"))
 		.SetTooltipText(LOCTEXT("TedsDebugger_QueryEditorToolTip", "Opens the TEDS Query Editor"))
 		.SetIcon(FSlateIcon(FAppStyle::GetAppStyleSetName(), "Debug"));
-	
-	TabManager->RegisterTabSpawner(
-		Private::TableViewerToolTabName,
-		FOnSpawnTab::CreateRaw(this, &STedsDebugger::SpawnTableViewerTab))
-		.SetGroup(AppMenuGroup)
-		.SetDisplayName(LOCTEXT("TedsDebugger_TableViewerDisplayName", "Table Viewer"))
-		.SetTooltipText(LOCTEXT("TedsDebugger_TableViewerToolTip", "Opens the TEDS Table Viewer"))
-		.SetIcon(FSlateIcon(FAppStyle::GetAppStyleSetName(), "LevelEditor.Tabs.Outliner"));
 }
 
-void STedsDebugger::NavigateToRow(TypedElementDataStorage::RowHandle InRow) const
-{
-	UTypedElementRegistry* Registry = UTypedElementRegistry::GetInstance();
-
-	if (!Registry)
-	{
-		return;
-	}
-	
-	// If the debugger isn't already open, open it
-	if (!TableViewerInstance.IsValid())
-	{
-		TabManager->TryInvokeTab(Private::TableViewerToolTabName);
-	}
-
-	TSharedPtr<ISceneOutliner> TableViewerPinned = TableViewerInstance.Pin();
-	if (!TableViewerPinned)
-	{
-		return;
-	}
-
-	// If the item isn't currently present in the debugger, try disabling all filters to make it show up
-	if (!TableViewerPinned->GetTreeItem(InRow))
-	{
-		TableViewerPinned->DisableAllFilterBarFilters(/** bRemove */ false);
-	}
-
-	// Defer the actual navigation by one tick to give the outliner a chance to update its items in case any filters were disabled
-	FTSTicker::GetCoreTicker().AddTicker(FTickerDelegate::CreateLambda([TableViewerPinned, InRow](float DeltaTime)
-	{
-		// Find the item for this row, select it and scroll to view it
-		if (FSceneOutlinerTreeItemPtr TreeItem = TableViewerPinned->GetTreeItem(InRow))
-		{
-			TableViewerPinned->SetSelection([TreeItem](ISceneOutlinerTreeItem& Item)
-			{
-				return Item.GetID() == TreeItem->GetID();
-			});
-			
-			TableViewerPinned->FrameSelectedItems();
-		}
-		
-		return false;
-	}));
-}
 } // namespace UE::Editor::DataStorage::Debug
 
 #undef LOCTEXT_NAMESPACE
diff --git a/Engine/Plugins/Experimental/EditorDataStorageFeatures/Source/TedsDebugger/Private/STedsDebugger.h b/Engine/Plugins/Experimental/EditorDataStorageFeatures/Source/TedsDebugger/Private/STedsDebugger.h
--- a/Engine/Plugins/Experimental/EditorDataStorageFeatures/Source/TedsDebugger/Private/STedsDebugger.h
+++ b/Engine/Plugins/Experimental/EditorDataStorageFeatures/Source/TedsDebugger/Private/STedsDebugger.h
@@ -33,13 +33,9 @@ namespace UE::Editor::DataStorage
 			* @param ConstructUnderWindow The window in which this widget is being constructed.
 			*/
 			void Construct(const FArguments& InArgs, const TSharedRef<SDockTab>& ConstructUnderMajorTab, const TSharedPtr<SWindow>& ConstructUnderWindow);
-
-			void NavigateToRow(RowHandle InRow) const;
-	
 		private:
 			TSharedRef<SDockTab> SpawnToolbar(const FSpawnTabArgs& Args);
 			TSharedRef<SDockTab> SpawnQueryEditorTab(const FSpawnTabArgs& Args);
-			TSharedRef<SDockTab> SpawnTableViewerTab(const FSpawnTabArgs& Args);
 			void FillWindowMenu( FMenuBuilder& MenuBuilder);
 
 			void RegisterTabSpawners();
@@ -51,7 +47,6 @@ namespace UE::Editor::DataStorage
 
 			// Table Viewer
 			QueryHandle TableViewerQuery;
-			TWeakPtr<ISceneOutliner> TableViewerInstance;
 
 			// Query Editor
 			TUniquePtr<QueryEditor::FTedsQueryEditorModel> QueryEditorModel;
diff --git a/Engine/Plugins/Experimental/EditorDataStorageFeatures/Source/TedsDebugger/Private/TedsDebuggerModule.cpp b/Engine/Plugins/Experimental/EditorDataStorageFeatures/Source/TedsDebugger/Private/TedsDebuggerModule.cpp
--- a/Engine/Plugins/Experimental/EditorDataStorageFeatures/Source/TedsDebugger/Private/TedsDebuggerModule.cpp
+++ b/Engine/Plugins/Experimental/EditorDataStorageFeatures/Source/TedsDebugger/Private/TedsDebuggerModule.cpp
@@ -67,30 +67,6 @@ TSharedRef<SDockTab> FTedsDebuggerModule::OpenTedsDebuggerTab(const FSpawnTabArg
 	
 	return MajorTab;
 }
-
-void FTedsDebuggerModule::NavigateToRow(RowHandle InRow) const
-{
-	UTypedElementRegistry* Registry = UTypedElementRegistry::GetInstance();
-
-	if (!Registry)
-	{
-		return;
-	}
-	
-	// If the debugger isn't already open, open it
-	if(!TedsDebuggerInstance.IsValid())
-	{
-		FGlobalTabmanager::Get()->TryInvokeTab(Private::TedsDebuggerTableName);
-	}
-
-	TSharedPtr<STedsDebugger> TedsDebuggerPinned = TedsDebuggerInstance.Pin();
-	if(!TedsDebuggerPinned)
-	{
-		return;
-	}
-
-	TedsDebuggerPinned->NavigateToRow(InRow);
-}
 } // namespace UE::Editor::DataStorage::Debug
 
 IMPLEMENT_MODULE(UE::Editor::DataStorage::Debug::FTedsDebuggerModule, TedsDebugger);
diff --git a/Engine/Plugins/Experimental/EditorDataStorageFeatures/Source/TedsDebugger/Private/TedsDebuggerModule.h b/Engine/Plugins/Experimental/EditorDataStorageFeatures/Source/TedsDebugger/Private/TedsDebuggerModule.h
--- a/Engine/Plugins/Experimental/EditorDataStorageFeatures/Source/TedsDebugger/Private/TedsDebuggerModule.h
+++ b/Engine/Plugins/Experimental/EditorDataStorageFeatures/Source/TedsDebugger/Private/TedsDebuggerModule.h
@@ -31,9 +31,6 @@ namespace UE::Editor::DataStorage
 			virtual void StartupModule() override;
 			virtual void ShutdownModule() override;
 
-			// Open the TEDS Debugger if not already open, and navigate to the given row in the table viewer tab
-			void NavigateToRow(RowHandle InRow) const;
-
 		private:
 			void RegisterTabSpawners();
 			void UnregisterTabSpawners() const;
diff --git a/Engine/Plugins/Experimental/EditorDataStorageFeatures/Source/TedsDebugger/Private/Widgets/QueryEditor/TedsQueryEditorResultsView.cpp b/Engine/Plugins/Experimental/EditorDataStorageFeatures/Source/TedsDebugger/Private/Widgets/QueryEditor/TedsQueryEditorResultsView.cpp
--- a/Engine/Plugins/Experimental/EditorDataStorageFeatures/Source/TedsDebugger/Private/Widgets/QueryEditor/TedsQueryEditorResultsView.cpp
+++ b/Engine/Plugins/Experimental/EditorDataStorageFeatures/Source/TedsDebugger/Private/Widgets/QueryEditor/TedsQueryEditorResultsView.cpp
@@ -48,6 +48,7 @@ namespace UE::Editor::DataStorage
 					[
 						SAssignNew(TableViewer, STedsTableViewer)
 						.QueryStack(RowQueryStack)
+						.EmptyRowsMessage(LOCTEXT("EmptyRowsMessage", "The provided query has no results."))
 						.OnSelectionChanged(STedsTableViewer::FOnSelectionChanged::CreateLambda(
 							[this](RowHandle SelectedRow)
 								{
diff --git a/Engine/Plugins/Experimental/EditorDataStorageFeatures/Source/TedsDebugger/Private/Widgets/RowReferenceWidget.cpp b/Engine/Plugins/Experimental/EditorDataStorageFeatures/Source/TedsDebugger/Private/Widgets/RowReferenceWidget.cpp
--- a/Engine/Plugins/Experimental/EditorDataStorageFeatures/Source/TedsDebugger/Private/Widgets/RowReferenceWidget.cpp
+++ b/Engine/Plugins/Experimental/EditorDataStorageFeatures/Source/TedsDebugger/Private/Widgets/RowReferenceWidget.cpp
@@ -49,10 +49,6 @@ namespace UE::Editor::DataStorage::Debug::Private
 
 			return;
 		}
-
-		// If it wasn't found in the table viewer owning this widget, navigate to it in the global TEDS debugger
-		Debug::FTedsDebuggerModule& TedsDebuggerModule = FModuleManager::GetModuleChecked<Debug::FTedsDebuggerModule>("TedsDebugger");
-		TedsDebuggerModule.NavigateToRow(TargetRowHandle);
 	}
 	
 	void CreateInternalWidget(const TWeakPtr<SWidget>& InWidget, RowHandle UiRow, RowHandle TargetRow)
diff --git a/Engine/Plugins/Experimental/EditorDataStorageFeatures/Source/TedsTableViewer/Private/Widgets/STedsTableViewer.cpp b/Engine/Plugins/Experimental/EditorDataStorageFeatures/Source/TedsTableViewer/Private/Widgets/STedsTableViewer.cpp
--- a/Engine/Plugins/Experimental/EditorDataStorageFeatures/Source/TedsTableViewer/Private/Widgets/STedsTableViewer.cpp
+++ b/Engine/Plugins/Experimental/EditorDataStorageFeatures/Source/TedsTableViewer/Private/Widgets/STedsTableViewer.cpp
@@ -18,6 +18,7 @@ namespace UE::Editor::DataStorage
 	void STedsTableViewer::Construct(const FArguments& InArgs)
 	{
 		OnSelectionChanged = InArgs._OnSelectionChanged;
+		EmptyRowsMessage = InArgs._EmptyRowsMessage;
 		
 		Model = MakeShared<FTedsTableViewerModel>(InArgs._QueryStack, InArgs._Columns, InArgs._CellWidgetPurposes,
 			FTedsTableViewerModel::FIsItemVisible::CreateSP(this, &STedsTableViewer::IsItemVisible));
@@ -86,8 +87,9 @@ namespace UE::Editor::DataStorage
 	void STedsTableViewer::AssignChildSlot()
 	{
 		TSharedPtr<SWidget> ContentWidget;
-		
-		if(Model->GetRowCount() == 0)
+
+		// If there are no rows and the table viewer wants to show a message
+		if(Model->GetRowCount() == 0 && EmptyRowsMessage.IsSet())
 		{
 			ContentWidget = 
 				SNew(SBox)
@@ -95,7 +97,7 @@ namespace UE::Editor::DataStorage
 					.VAlign(VAlign_Center)
 					[
 						SNew(STextBlock)
-							.Text(LOCTEXT("EmptyTableViewerQueryText", "The input query has no results"))
+							.Text(EmptyRowsMessage)
 					];
 		}
 		else if(Model->GetColumnCount() == 0)
@@ -106,7 +108,7 @@ namespace UE::Editor::DataStorage
 					.VAlign(VAlign_Center)
 					[
 						SNew(STextBlock)
-							.Text(LOCTEXT("EmptyTableViewerColumnsText", "There were no columns specified to display"))
+							.Text(LOCTEXT("EmptyTableViewerColumnsText", "No columns found to display."))
 					];
 		}
 		else
diff --git a/Engine/Plugins/Experimental/EditorDataStorageFeatures/Source/TedsTableViewer/Public/Widgets/STedsTableViewer.h b/Engine/Plugins/Experimental/EditorDataStorageFeatures/Source/TedsTableViewer/Public/Widgets/STedsTableViewer.h
--- a/Engine/Plugins/Experimental/EditorDataStorageFeatures/Source/TedsTableViewer/Public/Widgets/STedsTableViewer.h
+++ b/Engine/Plugins/Experimental/EditorDataStorageFeatures/Source/TedsTableViewer/Public/Widgets/STedsTableViewer.h
@@ -57,6 +57,10 @@ namespace UE::Editor::DataStorage
 		// The selection mode for the table viewer (single/multi etc)
 		SLATE_ARGUMENT(ESelectionMode::Type, ListSelectionMode)
 
+		// The message to show in place of the table viewer when there are no rows provided by the current query stack
+		// Empty = simply show the column headers instead of a message
+		SLATE_ATTRIBUTE(FText, EmptyRowsMessage)
+
 		SLATE_END_ARGS()
 
 	public:
@@ -105,5 +109,8 @@ namespace UE::Editor::DataStorage
 
 		// Wrapper Teds Widget around our contents so we can use Teds columns to specify behavior
 		TSharedPtr<STedsWidget> TedsWidget;
+
+		// The message to show in place of the table viewer when there are no rows provided by the current query stack
+		TAttribute<FText> EmptyRowsMessage;
 	};
 }
\ No newline at end of file
