TEDS-Content Browser prototype (2/2): Content Browser changes

CB: Add a ContentBrowserViewExtender class to allow systems to register a custom view in the content browser

CB: Make AssetViewTypes.h public

SAssetView: Add support for displaying a custom view
 - Not all actions are completely supported yet.

TEDS CB: Add a new TEDS driven content browser as a custom view behind a cvar
 - Displays data using the export text widget
 - Currently supports browsing, context menus and double clicking items to open/navigate

diff --git a/Engine/Plugins/Experimental/EditorDataStorageFeatures/EditorDataStorageFeatures.uplugin b/Engine/Plugins/Experimental/EditorDataStorageFeatures/EditorDataStorageFeatures.uplugin
--- a/Engine/Plugins/Experimental/EditorDataStorageFeatures/EditorDataStorageFeatures.uplugin
+++ b/Engine/Plugins/Experimental/EditorDataStorageFeatures/EditorDataStorageFeatures.uplugin
@@ -43,6 +43,11 @@
 			"Name": "TedsTableViewer",
 			"Type": "Editor"
 		}
+	,
+		{
+			"Name": "TedsContentBrowser",
+			"Type": "Editor"
+		}
 	],
 	"Plugins": [
 	]
diff --git a/Engine/Plugins/Experimental/EditorDataStorageFeatures/Source/TedsContentBrowser/Private/TedsContentBrowserModule.cpp b/Engine/Plugins/Experimental/EditorDataStorageFeatures/Source/TedsContentBrowser/Private/TedsContentBrowserModule.cpp
new file mode 100644
--- /dev/null
+++ b/Engine/Plugins/Experimental/EditorDataStorageFeatures/Source/TedsContentBrowser/Private/TedsContentBrowserModule.cpp
@@ -0,0 +1,269 @@
+// Copyright Epic Games, Inc. All Rights Reserved.
+
+#include "TedsContentBrowserModule.h"
+
+#include "Framework/Application/SlateApplication.h"
+#include "ContentBrowserModule.h"
+#include "Columns/SlateDelegateColumns.h"
+#include "Elements/Common/TypedElementQueryTypes.h"
+#include "Elements/Framework/TypedElementIndexHasher.h"
+#include "Elements/Framework/TypedElementRegistry.h"
+#include "Experimental/ContentBrowserViewExtender.h"
+#include "Modules/ModuleManager.h"
+#include "QueryStack/FQueryStackNode_RowView.h"
+#include "TedsAssetDataColumns.h"
+#include "Widgets/STedsTableViewer.h"
+#include "Widgets/Views/SListView.h"
+
+#define LOCTEXT_NAMESPACE "TedsContentBrowserModule"
+
+namespace UE::Editor::DataStorage::ContentBrowser
+{
+	static bool bEnableTedsContentBrowser = false;
+	
+	static FAutoConsoleVariableRef CVarUseTEDSOutliner(
+		TEXT("TEDS.UI.EnableTedsContentBrowser"),
+		bEnableTedsContentBrowser,
+		TEXT("Add the Teds Content Browser as a custom view (requires re-opening any currently open content browsers)")
+		, FConsoleVariableDelegate::CreateLambda([](IConsoleVariable* /*CVar*/)
+		{
+			FContentBrowserModule& ContentBrowserModule = FModuleManager::Get().GetModuleChecked<FContentBrowserModule>(TEXT("ContentBrowser"));
+			FTedsContentBrowserModule& TedsContentBrowserModule = FModuleManager::Get().GetModuleChecked<FTedsContentBrowserModule>(TEXT("TedsContentBrowser"));
+			
+			if(bEnableTedsContentBrowser)
+			{
+				ContentBrowserModule.SetContentBrowserViewExtender(TedsContentBrowserModule.GetContentBrowserViewExtender());
+			}
+			else
+			{
+				ContentBrowserModule.SetContentBrowserViewExtender(nullptr);
+			}
+		}));
+
+	TSharedPtr<IContentBrowserViewExtender> FTedsContentBrowserModule::GetContentBrowserViewExtender()
+	{
+		if(!ViewExtender)
+		{
+			ViewExtender = MakeShared<FTedsContentBrowserViewExtender>();
+		}
+		return ViewExtender;
+	}
+
+	void FTedsContentBrowserModule::StartupModule()
+	{
+		IModuleInterface::StartupModule();
+	}
+
+	void FTedsContentBrowserModule::ShutdownModule()
+	{
+		IModuleInterface::ShutdownModule();
+	}
+
+	void FTedsContentBrowserViewExtender::RefreshRows(TArray<TSharedPtr<FAssetViewItem>>* InItemsSource)
+	{
+		if(!InItemsSource)
+		{
+			return;
+		}
+
+		Rows.Empty();
+		ContentBrowserItemMap.Empty();
+
+		for(TSharedPtr<FAssetViewItem> Item : *InItemsSource)
+		{
+			if(!Item)
+			{
+				continue;
+			}
+
+			AddRow(Item);
+		}
+		
+		RowQueryStack->MarkDirty();
+	}
+
+	void FTedsContentBrowserViewExtender::AddRow(const TSharedPtr<FAssetViewItem>& Item)
+	{
+		FAssetData ItemAssetData;
+		FName PackagePath;
+		
+		RowHandle RowHandle = InvalidRowHandle;
+		
+		if (Item->GetItem().Legacy_TryGetAssetData(ItemAssetData))
+		{
+			TypedElementDataStorage::IndexHash IndexHash = TypedElementDataStorage::GenerateIndexHash(ItemAssetData.GetSoftObjectPath());
+			RowHandle = DataStorage->FindIndexedRow(IndexHash);
+		}
+		else if(Item->GetItem().Legacy_TryGetPackagePath(PackagePath))
+		{
+			TypedElementDataStorage::IndexHash IndexHash = TypedElementDataStorage::GenerateIndexHash(PackagePath);
+			RowHandle = DataStorage->FindIndexedRow(IndexHash);
+		}
+		
+		if(DataStorage->IsRowAssigned(RowHandle))
+		{
+			ContentBrowserItemMap.Emplace(RowHandle, Item);
+			Rows.Add(RowHandle);
+		}
+	}
+
+	TSharedPtr<FAssetViewItem> FTedsContentBrowserViewExtender::GetAssetViewItemFromRow(RowHandle Row)
+	{
+		// CB 2.0 TODO: Since FAssetViewItem was private previously, there is no good way to lookup currently aside from storing a map
+		if(TWeakPtr<FAssetViewItem>* AssetViewItem = ContentBrowserItemMap.Find(Row))
+		{
+			if(TSharedPtr<FAssetViewItem> AssetViewItemPin = AssetViewItem->Pin())
+			{
+				return AssetViewItemPin;
+			}
+		}
+
+		return nullptr;
+	}
+
+	FTedsContentBrowserViewExtender::FTedsContentBrowserViewExtender()
+	{
+		DataStorage = UTypedElementRegistry::GetInstance()->GetMutableDataStorage();
+		
+		RowQueryStack = MakeShared<FQueryStackNode_RowView>(&Rows);
+
+		// Create the table viewer widget
+		TableViewer = SNew(UE::Editor::DataStorage::STedsTableViewer)
+					.QueryStack(RowQueryStack)
+					.CellWidgetPurposes({TEXT("General.RowLabel"), TEXT("General.Cell")})
+					// Default list of columns to display
+					.Columns({ FItemNameColumn_Experimental::StaticStruct(), FAssetTag::StaticStruct(), FAssetPathColumn_Experimental::StaticStruct(),
+						FDiskSizeColumn::StaticStruct(), FVirtualPathColumn_Experimental::StaticStruct() })
+					.ListSelectionMode(ESelectionMode::Multi)
+					.OnSelectionChanged_Lambda([this](RowHandle Row)
+					{
+						if(TSharedPtr<FAssetViewItem> AssetViewItem = GetAssetViewItemFromRow(Row))
+						{
+							// CB 2.0 TODO: Does the CB use ESelectInfo and we need to propagate it from the table viewer?
+							OnSelectionChangedDelegate.Execute(AssetViewItem, ESelectInfo::Direct);
+						}
+					});
+
+		// Bind the delegates the CB view extender requires to delegates in TEDS columns on the widget row that are fired
+		// when the event occurs
+		const RowHandle WidgetRow = TableViewer->GetWidgetRowHandle();
+
+		if(FWidgetContextMenuColumn* ContextMenuColumn = DataStorage->GetColumn<FWidgetContextMenuColumn>(WidgetRow))
+		{
+			ContextMenuColumn->OnContextMenuOpening.BindLambda([this]()
+			{
+				return OnContextMenuOpenedDelegate.Execute();
+			});
+		}
+		if(FWidgetRowScrolledIntoView* ScrolledIntoViewColumn = DataStorage->GetColumn<FWidgetRowScrolledIntoView>(WidgetRow))
+		{
+			ScrolledIntoViewColumn->OnItemScrolledIntoView.BindLambda([this](FTedsRowHandle Row, const TSharedPtr<ITableRow>& TableRow)
+			{
+				if(TSharedPtr<FAssetViewItem> AssetViewItem = GetAssetViewItemFromRow(Row))
+				{
+					return OnItemScrolledIntoViewDelegate.Execute(AssetViewItem, TableRow);
+				}
+			});
+		}
+		if(FWidgetDoubleClickedColumn* DoubleClickedColumn = DataStorage->GetColumn<FWidgetDoubleClickedColumn>(WidgetRow))
+		{
+			DoubleClickedColumn->OnMouseButtonDoubleClick.BindLambda([this](FTedsRowHandle Row)
+			{
+				if(TSharedPtr<FAssetViewItem> AssetViewItem = GetAssetViewItemFromRow(Row))
+				{
+					return OnItemDoubleClickedDelegate.Execute(AssetViewItem);
+				}
+			});
+		}
+	}
+
+	TSharedRef<SWidget> FTedsContentBrowserViewExtender::CreateView(TArray<TSharedPtr<FAssetViewItem>>* InItemsSource)
+	{
+		RefreshRows(InItemsSource);
+		return TableViewer.ToSharedRef();
+	}
+	
+	void FTedsContentBrowserViewExtender::OnItemListChanged(TArray<TSharedPtr<FAssetViewItem>>* InItemsSource)
+	{
+		// CB 2.0 TODO: We might want to track individual addition/removals instead of a full refresh for perf
+		RefreshRows(InItemsSource);
+	}
+
+
+	TArray<TSharedPtr<FAssetViewItem>> FTedsContentBrowserViewExtender::GetSelectedItems()
+	{
+		// CB 2.0 TODO: Figure out selection
+		TArray<TSharedPtr<FAssetViewItem>> SelectedItems;
+
+		TableViewer->ForEachSelectedRow([this, &SelectedItems](RowHandle Row)
+		{
+			if(TSharedPtr<FAssetViewItem> AssetViewItem = GetAssetViewItemFromRow(Row))
+			{
+				SelectedItems.Add(AssetViewItem);
+			}
+		});
+		
+		return SelectedItems;
+	}
+
+	IContentBrowserViewExtender::FOnSelectionChanged& FTedsContentBrowserViewExtender::OnSelectionChanged()
+	{
+		return OnSelectionChangedDelegate;
+	}
+
+	FOnContextMenuOpening& FTedsContentBrowserViewExtender::OnContextMenuOpened()
+	{
+		return OnContextMenuOpenedDelegate;
+	}
+
+	IContentBrowserViewExtender::FOnItemScrolledIntoView& FTedsContentBrowserViewExtender::OnItemScrolledIntoView()
+	{
+		return OnItemScrolledIntoViewDelegate;
+	}
+
+	IContentBrowserViewExtender::FOnMouseButtonClick& FTedsContentBrowserViewExtender::OnItemDoubleClicked()
+	{
+		return OnItemDoubleClickedDelegate;
+	}
+
+	FText FTedsContentBrowserViewExtender::GetViewDisplayName()
+	{
+		return LOCTEXT("TedsCBViewName", "TEDS List View");
+	}
+
+	FText FTedsContentBrowserViewExtender::GetViewTooltipText()
+	{
+		return LOCTEXT("TedsCBViewTooltip", "A List view populated using TEDS UI and the asset registry data in TEDS");
+	}
+
+	void FTedsContentBrowserViewExtender::FocusList()
+	{
+		// CB 2.0 TODO: Do we need to focus the internal list? If so, implement using a Teds column
+		FSlateApplication::Get().SetKeyboardFocus(TableViewer, EFocusCause::SetDirectly);
+	}
+
+	void FTedsContentBrowserViewExtender::SetSelection(const TSharedPtr<FAssetViewItem>& Item, bool bSelected, const ESelectInfo::Type SelectInfo)
+	{
+		// CB 2.0 TODO: Implement using Teds selection
+	}
+
+	void FTedsContentBrowserViewExtender::RequestScrollIntoView(const TSharedPtr<FAssetViewItem>& Item)
+	{
+		// CB 2.0 TODO: Implement using a Teds column
+	}
+
+	void FTedsContentBrowserViewExtender::ClearSelection()
+	{
+		// CB 2.0 TODO: Implement using a Teds column
+	}
+
+	bool FTedsContentBrowserViewExtender::IsRightClickScrolling()
+	{
+		// CB 2.0 TODO: Implement using a Teds column
+		return false;
+	}
+}
+
+IMPLEMENT_MODULE(UE::Editor::DataStorage::ContentBrowser::FTedsContentBrowserModule, TedsContentBrowser);
+
+#undef LOCTEXT_NAMESPACE
\ No newline at end of file
diff --git a/Engine/Plugins/Experimental/EditorDataStorageFeatures/Source/TedsContentBrowser/Private/TedsContentBrowserModule.h b/Engine/Plugins/Experimental/EditorDataStorageFeatures/Source/TedsContentBrowser/Private/TedsContentBrowserModule.h
new file mode 100644
--- /dev/null
+++ b/Engine/Plugins/Experimental/EditorDataStorageFeatures/Source/TedsContentBrowser/Private/TedsContentBrowserModule.h
@@ -0,0 +1,100 @@
+// Copyright Epic Games, Inc. All Rights Reserved.
+
+#pragma once
+
+#include "Elements/Common/TypedElementHandles.h"
+#include "Experimental/ContentBrowserViewExtender.h"
+#include "Modules/ModuleInterface.h"
+#include "Templates/SharedPointer.h"
+
+class ITypedElementDataStorageInterface;
+
+namespace UE::Editor::DataStorage
+{
+	class FQueryStackNode_RowView;
+	class STedsTableViewer;
+}
+
+class FAssetViewItem;
+
+namespace UE::Editor::DataStorage::ContentBrowser
+{
+	// A custom view for the content browser that uses the TEDS asset registry integration to display rows with widgets using TEDS UI
+	class FTedsContentBrowserViewExtender : public IContentBrowserViewExtender
+	{
+	public:
+
+		FTedsContentBrowserViewExtender();
+
+		// IContentBrowserViewExtender interface
+		virtual TSharedRef<SWidget> CreateView(TArray<TSharedPtr<FAssetViewItem>>* InItemsSource) override;
+		virtual TArray<TSharedPtr<FAssetViewItem>> GetSelectedItems() override;
+		virtual FOnSelectionChanged& OnSelectionChanged() override;
+		virtual FOnContextMenuOpening& OnContextMenuOpened() override;
+		virtual FOnItemScrolledIntoView& OnItemScrolledIntoView() override;
+		virtual FOnMouseButtonClick& OnItemDoubleClicked() override;
+		virtual FText GetViewDisplayName() override;
+		virtual FText GetViewTooltipText() override;
+		virtual void FocusList() override;
+		virtual void SetSelection(const TSharedPtr<FAssetViewItem>& Item, bool bSelected, const ESelectInfo::Type SelectInfo)override;
+		virtual void RequestScrollIntoView(const TSharedPtr<FAssetViewItem>& Item) override;
+		virtual void ClearSelection() override;
+		virtual bool IsRightClickScrolling() override;
+		virtual void OnItemListChanged(TArray<TSharedPtr<FAssetViewItem>>* InItemsSource) override;
+		// ~IContentBrowserViewExtender interface
+
+
+		// Refresh the rows in the current view by syncing the the items source
+		void RefreshRows(TArray<TSharedPtr<FAssetViewItem>>* InItemsSource);
+
+		// Add a single row to the table viewer
+		void AddRow(const TSharedPtr<FAssetViewItem>& Item);
+
+		// Get the internal FAssetViewItem from a row handle
+		TSharedPtr<FAssetViewItem> GetAssetViewItemFromRow(UE::Editor::DataStorage::RowHandle Row);
+
+
+	private:
+
+		// Ptr to the data storage interface
+		ITypedElementDataStorageInterface* DataStorage;
+		
+		// The actual table viewer widget
+		TSharedPtr<STedsTableViewer> TableViewer;
+		
+		// Query stack used by the table viewer
+		TSharedPtr<FQueryStackNode_RowView> RowQueryStack;
+
+		// The row handles of the items currently in the list
+		TArray<UE::Editor::DataStorage::RowHandle> Rows;
+
+		// A map from row handle -> FAssetView item for lookups
+		TMap<UE::Editor::DataStorage::RowHandle, TWeakPtr<FAssetViewItem>> ContentBrowserItemMap;
+
+		// Delegates fired when specific events happen on the list
+		FOnSelectionChanged OnSelectionChangedDelegate;
+		FOnContextMenuOpening OnContextMenuOpenedDelegate;
+		FOnItemScrolledIntoView OnItemScrolledIntoViewDelegate;
+		FOnMouseButtonClick OnItemDoubleClickedDelegate;
+	};
+
+	/**
+	 * Implements the Teds Content Browser module.
+	 */
+	class FTedsContentBrowserModule
+		: public IModuleInterface
+	{
+	public:
+
+		FTedsContentBrowserModule() = default;
+		
+		TSharedPtr<IContentBrowserViewExtender> GetContentBrowserViewExtender();
+
+		// IModuleInterface interface
+		virtual void StartupModule() override;
+		virtual void ShutdownModule() override;
+
+	private:
+		TSharedPtr<FTedsContentBrowserViewExtender> ViewExtender;
+	};
+}
diff --git a/Engine/Plugins/Experimental/EditorDataStorageFeatures/Source/TedsContentBrowser/TedsContentBrowser.Build.cs b/Engine/Plugins/Experimental/EditorDataStorageFeatures/Source/TedsContentBrowser/TedsContentBrowser.Build.cs
new file mode 100644
--- /dev/null
+++ b/Engine/Plugins/Experimental/EditorDataStorageFeatures/Source/TedsContentBrowser/TedsContentBrowser.Build.cs
@@ -0,0 +1,40 @@
+// Copyright Epic Games, Inc. All Rights Reserved.
+
+using UnrealBuildTool;
+
+public class TedsContentBrowser : ModuleRules
+{
+	public TedsContentBrowser(ReadOnlyTargetRules Target) : base(Target)
+	{
+		PCHUsage = ModuleRules.PCHUsageMode.UseExplicitOrSharedPCHs;
+
+		if (Target.bBuildEditor)
+		{
+			PublicIncludePaths.AddRange(new string[] {});
+			PrivateIncludePaths.AddRange(new string[] {});
+
+			PublicDependencyModuleNames.AddRange(
+				new string[]
+				{
+					"Core",
+					"CoreUObject"
+				});
+			
+			PrivateDependencyModuleNames.AddRange(
+            	new string[]
+            	{
+		            "TypedElementFramework",
+		            "SlateCore",
+		            "Slate",
+		            "TedsTableViewer",
+		            "ContentBrowser",
+		            "ContentBrowserData",
+		            "InputCore",
+		            "TedsAssetData",
+		            "ToolWidgets"
+            	});
+			
+			DynamicallyLoadedModuleNames.AddRange(new string[] {});
+		}
+	}
+}
diff --git a/Engine/Source/Editor/ContentBrowser/Private/AssetViewTypes.h b/Engine/Source/Editor/ContentBrowser/Private/AssetViewTypes.h
deleted file mode 100644
diff --git a/Engine/Source/Editor/ContentBrowser/Private/ContentBrowserModule.cpp b/Engine/Source/Editor/ContentBrowser/Private/ContentBrowserModule.cpp
--- a/Engine/Source/Editor/ContentBrowser/Private/ContentBrowserModule.cpp
+++ b/Engine/Source/Editor/ContentBrowser/Private/ContentBrowserModule.cpp
@@ -48,6 +48,16 @@ IContentBrowserSingleton& FContentBrowserModule::Get() const
 	return *ContentBrowserSingleton;
 }
 
+void FContentBrowserModule::SetContentBrowserViewExtender(const TSharedPtr<IContentBrowserViewExtender>& InViewExtender)
+{
+	ContentBrowserViewExtender = InViewExtender;
+}
+
+TSharedPtr<IContentBrowserViewExtender> FContentBrowserModule::GetContentBrowserViewExtender()
+{
+	return ContentBrowserViewExtender;
+}
+
 FMainMRUFavoritesList* FContentBrowserModule::GetRecentlyOpenedAssets() const
 {
 	return GEditor->GetEditorSubsystem<UAssetEditorSubsystem>()->GetRecentlyOpenedAssets();
diff --git a/Engine/Source/Editor/ContentBrowser/Private/SAssetView.cpp b/Engine/Source/Editor/ContentBrowser/Private/SAssetView.cpp
--- a/Engine/Source/Editor/ContentBrowser/Private/SAssetView.cpp
+++ b/Engine/Source/Editor/ContentBrowser/Private/SAssetView.cpp
@@ -1227,6 +1227,21 @@ void SAssetView::Construct( const FArguments& InArgs )
 	FolderPermissionList = AssetToolsModule.Get().GetFolderPermissionList();
 	WritableFolderPermissionList = AssetToolsModule.Get().GetWritableFolderPermissionList();
 
+	if(InArgs._AllowCustomView)
+	{
+		FContentBrowserModule& ContentBrowserModule = FModuleManager::GetModuleChecked<FContentBrowserModule>(TEXT("ContentBrowser"));
+		ViewExtender = ContentBrowserModule.GetContentBrowserViewExtender();
+
+		// Bind the delegates the custom view is responsible for firing
+		if(ViewExtender)
+		{
+			ViewExtender->OnSelectionChanged().BindSP(this, &SAssetView::AssetSelectionChanged);
+			ViewExtender->OnContextMenuOpened().BindSP(this, &SAssetView::OnGetContextMenuContent);
+			ViewExtender->OnItemScrolledIntoView().BindSP(this, &SAssetView::ItemScrolledIntoView);
+			ViewExtender->OnItemDoubleClicked().BindSP(this, &SAssetView::OnListMouseButtonDoubleClick);
+		}
+	}
+	
 	FEditorWidgetsModule& EditorWidgetsModule = FModuleManager::LoadModuleChecked<FEditorWidgetsModule>("EditorWidgets");
 	TSharedRef<SWidget> AssetDiscoveryIndicator = EditorWidgetsModule.CreateAssetDiscoveryIndicator(EAssetDiscoveryIndicatorScaleMode::Scale_Vertical);
 
@@ -1709,6 +1724,7 @@ TArray<TSharedPtr<FAssetViewItem>> SAssetView::GetSelectedViewItems() const
 		case EAssetViewType::List: return ListView->GetSelectedItems();
 		case EAssetViewType::Tile: return TileView->GetSelectedItems();
 		case EAssetViewType::Column: return ColumnView->GetSelectedItems();
+		case EAssetViewType::Custom: return ViewExtender->GetSelectedItems();
 		default:
 		ensure(0); // Unknown list type
 		return TArray<TSharedPtr<FAssetViewItem>>();
@@ -3388,6 +3404,16 @@ void SAssetView::ExecutePaste()
 	}
 }
 
+bool SAssetView::IsCustomViewSet() const
+{
+	return ViewExtender.IsValid();
+}
+
+TSharedRef<SWidget> SAssetView::CreateCustomView()
+{
+	return IsCustomViewSet() ? ViewExtender->CreateView(&FilteredAssetItems) : SNullWidget::NullWidget;
+}
+
 void SAssetView::ToggleShowAllFolder()
 {
 	const bool bNewValue = !IsShowingAllFolder();
@@ -3847,6 +3873,23 @@ void SAssetView::PopulateViewButtonMenu(UToolMenu* Menu)
 				),
 			EUserInterfaceActionType::RadioButton
 			);
+
+		const FText CustomViewLabel = ViewExtender ? ViewExtender->GetViewDisplayName() : LOCTEXT("CustomViewOption", "Custom");
+		const FText CustomViewTooltip = ViewExtender ? ViewExtender->GetViewTooltipText() : LOCTEXT("CustomViewOptionToolTip", "A user specified custom view.");
+
+		Section.AddMenuEntry(
+			"CustomView",
+			CustomViewLabel,
+			CustomViewTooltip,
+			FSlateIcon(),
+			FUIAction(
+				FExecuteAction::CreateSP( this, &SAssetView::SetCurrentViewTypeFromMenu, EAssetViewType::Custom ),
+				FCanExecuteAction(),
+				FIsActionChecked::CreateSP( this, &SAssetView::IsCurrentViewType, EAssetViewType::Custom ),
+				FIsActionButtonVisible::CreateSP(this, &SAssetView::IsCustomViewSet)
+				),
+			EUserInterfaceActionType::RadioButton
+			);
 	}
 
 	if(TSharedPtr<SFilterList> FilterBarPinned = FilterBar.Pin())
@@ -4594,6 +4637,15 @@ void SAssetView::SetCurrentViewType(EAssetViewType::Type NewType)
 {
 	if ( ensure(NewType != EAssetViewType::MAX) && NewType != CurrentViewType )
 	{
+		// If we are setting to the custom type, but the view extender does not exist for some reason - default back to tile
+		if(NewType == EAssetViewType::Custom)
+		{
+			if(!IsCustomViewSet())
+			{
+				NewType = EAssetViewType::Tile;
+			}
+		}
+		
 		ResetQuickJump();
 
 		CurrentViewType = NewType;
@@ -4668,6 +4720,11 @@ void SAssetView::CreateCurrentView()
 			ColumnView = CreateColumnView();
 			NewView = CreateShadowOverlay(ColumnView.ToSharedRef());
 			break;
+		case EAssetViewType::Custom:
+			// The custom view does not necessarily have an accessible list, so we create a generic scroll border
+			CustomView = CreateCustomView();
+			NewView = CustomView.ToSharedRef();
+			break;
 	}
 	
 	ViewContainer->SetContent( NewView );
@@ -4713,6 +4770,7 @@ void SAssetView::RefreshList()
 		case EAssetViewType::List: ListView->RequestListRefresh(); break;
 		case EAssetViewType::Tile: TileView->RequestListRefresh(); break;
 		case EAssetViewType::Column: ColumnView->RequestListRefresh(); break;
+		case EAssetViewType::Custom: ViewExtender->OnItemListChanged(&FilteredAssetItems); break;
 	}
 }
 
@@ -4723,6 +4781,7 @@ void SAssetView::SetSelection(const TSharedPtr<FAssetViewItem>& Item)
 		case EAssetViewType::List: ListView->SetSelection(Item); break;
 		case EAssetViewType::Tile: TileView->SetSelection(Item); break;
 		case EAssetViewType::Column: ColumnView->SetSelection(Item); break;
+		case EAssetViewType::Custom: ViewExtender->SetSelection(Item, true, ESelectInfo::Direct); break;
 	}
 }
 
@@ -4733,6 +4792,7 @@ void SAssetView::SetItemSelection(const TSharedPtr<FAssetViewItem>& Item, bool b
 		case EAssetViewType::List: ListView->SetItemSelection(Item, bSelected, SelectInfo); break;
 		case EAssetViewType::Tile: TileView->SetItemSelection(Item, bSelected, SelectInfo); break;
 		case EAssetViewType::Column: ColumnView->SetItemSelection(Item, bSelected, SelectInfo); break;
+		case EAssetViewType::Custom: ViewExtender->SetSelection(Item, bSelected, SelectInfo); break;
 	}
 }
 
@@ -4743,6 +4803,7 @@ void SAssetView::RequestScrollIntoView(const TSharedPtr<FAssetViewItem>& Item)
 		case EAssetViewType::List: ListView->RequestScrollIntoView(Item); break;
 		case EAssetViewType::Tile: TileView->RequestScrollIntoView(Item); break;
 		case EAssetViewType::Column: ColumnView->RequestScrollIntoView(Item); break;
+		case EAssetViewType::Custom: ViewExtender->RequestScrollIntoView(Item); break;
 	}
 }
 
@@ -4775,6 +4836,7 @@ void SAssetView::ClearSelection(bool bForceSilent)
 		case EAssetViewType::List: ListView->ClearSelection(); break;
 		case EAssetViewType::Tile: TileView->ClearSelection(); break;
 		case EAssetViewType::Column: ColumnView->ClearSelection(); break;
+		case EAssetViewType::Custom: ViewExtender->ClearSelection(); break;
 	}
 }
 
@@ -5414,6 +5476,10 @@ bool SAssetView::ShouldAllowToolTips() const
 		case EAssetViewType::Column:
 			bIsRightClickScrolling = ColumnView->IsRightClickScrolling();
 			break;
+		
+		case EAssetViewType::Custom:
+			bIsRightClickScrolling = ViewExtender->IsRightClickScrolling();
+			break;
 
 		default:
 			bIsRightClickScrolling = false;
@@ -5589,7 +5655,7 @@ void SAssetView::UpdateThumbnailSizeValue()
 
 bool SAssetView::IsThumbnailScalingAllowed() const
 {
-	return GetCurrentViewType() != EAssetViewType::Column;
+	return GetCurrentViewType() != EAssetViewType::Column && GetCurrentViewType() != EAssetViewType::Custom;
 }
 
 float SAssetView::GetTileViewTypeNameHeight() const
diff --git a/Engine/Source/Editor/ContentBrowser/Private/SContentBrowser.cpp b/Engine/Source/Editor/ContentBrowser/Private/SContentBrowser.cpp
--- a/Engine/Source/Editor/ContentBrowser/Private/SContentBrowser.cpp
+++ b/Engine/Source/Editor/ContentBrowser/Private/SContentBrowser.cpp
@@ -485,7 +485,8 @@ void SContentBrowser::Construct( const FArguments& InArgs, const FName& InInstan
 			.OnSearchOptionsChanged(this, &SContentBrowser::HandleAssetViewSearchOptionsChanged)
 			.bShowPathViewFilters(true)
 			.FillEmptySpaceInTileView(true)
-			.ShowDisallowedAssetClassAsUnsupportedItems(true);
+			.ShowDisallowedAssetClassAsUnsupportedItems(true)
+			.AllowCustomView(true);
 
 	TSharedRef<SWidget> ViewOptions = SNullWidget::NullWidget;
 
diff --git a/Engine/Source/Editor/ContentBrowser/Public/AssetViewTypes.h b/Engine/Source/Editor/ContentBrowser/Public/AssetViewTypes.h
new file mode 100644
--- /dev/null
+++ b/Engine/Source/Editor/ContentBrowser/Public/AssetViewTypes.h
@@ -0,0 +1,113 @@
+// Copyright Epic Games, Inc. All Rights Reserved.
+
+#pragma once
+
+#include "Containers/ArrayView.h"
+#include "Containers/Map.h"
+#include "ContentBrowserItem.h"
+#include "CoreMinimal.h"
+#include "Delegates/Delegate.h"
+#include "Templates/Tuple.h"
+#include "UObject/NameTypes.h"
+#include "UObject/Object.h"
+
+#include <atomic>
+
+class FContentBrowserItemData;
+class FString;
+class FText;
+struct FAssetViewCustomColumn;
+
+/** An item (folder or file) displayed in the asset view */
+class FAssetViewItem
+{
+public:
+	FAssetViewItem(int32 Index);
+
+	explicit FAssetViewItem(int32 Index, FContentBrowserItem&& InItem);
+	explicit FAssetViewItem(int32 Index, const FContentBrowserItem& InItem);
+
+	explicit FAssetViewItem(int32 Index, FContentBrowserItemData&& InItemData);
+	explicit FAssetViewItem(int32 Index, const FContentBrowserItemData& InItemData);
+
+	// When recycling an object, clear the item data and replace it with the given data.
+	void ResetItemData(int32 OldIndex, int32 Index, FContentBrowserItemData InItemData);
+
+	void AppendItemData(const FContentBrowserItem& InItem);
+
+	void AppendItemData(const FContentBrowserItemData& InItemData);
+
+	void RemoveItemData(const FContentBrowserItem& InItem);
+
+	void RemoveItemData(const FContentBrowserItemData& InItemData);
+
+	void RemoveItemData(const FContentBrowserMinimalItemData& InItemKey);
+
+	/** Clear cached custom column data */
+	void ClearCachedCustomColumns();
+
+	/**
+	 * Updates cached custom column data (only does something for files) 
+	 * @param bUpdateExisting If true, only updates existing columns, if false only adds missing columns 
+	 */
+	void CacheCustomColumns(TArrayView<const FAssetViewCustomColumn> CustomColumns, const bool bUpdateSortData, const bool bUpdateDisplayText, const bool bUpdateExisting);
+	
+	/** Get the display value of a custom column on this item */
+	bool GetCustomColumnDisplayValue(const FName ColumnName, FText& OutText) const;
+
+	/** Get the value (and optionally also the type) of a custom column on this item */
+	bool GetCustomColumnValue(const FName ColumnName, FString& OutString, UObject::FAssetRegistryTag::ETagType* OutType = nullptr) const;
+
+	/** Get the value (and optionally also the type) of a named tag on this item */
+	bool GetTagValue(const FName Tag, FString& OutString, UObject::FAssetRegistryTag::ETagType* OutType = nullptr) const;
+
+	/** Get the underlying Content Browser item */
+	CONTENTBROWSER_API const FContentBrowserItem& GetItem() const;
+
+	bool IsFolder() const;
+
+	bool IsFile() const;
+
+	bool IsTemporary() const;
+
+	// Called when the view explicitly wants to notify widgets of changes
+	// Not called during bulk rebuilds when the view will be re-populated even if items are being recycled
+	void BroadcastItemDataChanged();
+
+	/** Get the event fired when the data for this item changes */
+	FSimpleMulticastDelegate& OnItemDataChanged();
+
+	/** Get the event fired whenever a rename is requested */
+	FSimpleDelegate& OnRenameRequested();
+
+	/** Get the event fired whenever a rename is canceled */
+	FSimpleDelegate& OnRenameCanceled();
+
+	/** Helper function to turn an item into a string for debugging */
+	static FString ItemToString_Debug(TSharedPtr<FAssetViewItem> AssetItem);
+
+private:
+	/** Underlying Content Browser item data */
+	FContentBrowserItem Item;
+
+	/**
+	 * Index at which this is stored in the asset view's item collection.
+	 * Can be used to detect an item being added to the collection twice by mistake.
+	 */
+	std::atomic<int32> Index;
+
+	/** An event to fire when the data for this item changes */
+	FSimpleMulticastDelegate ItemDataChangedEvent;
+
+	/** Broadcasts whenever a rename is requested */
+	FSimpleDelegate RenameRequestedEvent;
+
+	/** Broadcasts whenever a rename is canceled */
+	FSimpleDelegate RenameCanceledEvent;
+
+	/** Map of values/types for custom columns */
+	TMap<FName, TTuple<FString, UObject::FAssetRegistryTag::ETagType>> CachedCustomColumnData;
+	
+	/** Map of display text for custom columns */
+	TMap<FName, FText> CachedCustomColumnDisplayText;
+};
diff --git a/Engine/Source/Editor/ContentBrowser/Public/ContentBrowserModule.h b/Engine/Source/Editor/ContentBrowser/Public/ContentBrowserModule.h
--- a/Engine/Source/Editor/ContentBrowser/Public/ContentBrowserModule.h
+++ b/Engine/Source/Editor/ContentBrowser/Public/ContentBrowserModule.h
@@ -8,6 +8,7 @@
 #include "ContentBrowserDelegates.h"
 #include "CoreMinimal.h"
 #include "Delegates/Delegate.h"
+#include "Experimental/ContentBrowserViewExtender.h"
 #include "HAL/PlatformCrt.h"
 #include "MRUFavoritesList.h"
 #include "Modules/ModuleInterface.h"
@@ -173,6 +174,12 @@ public:
 	/** Delegates to be called to extend list of content browser Plugin Filters*/
 	virtual TArray<FAddPathViewPluginFilters>& GetAddPathViewPluginFilters() { return PathViewPluginFilters; }
 
+	/** Register a custom view that can show up inside SAssetView in the Content Browser */
+	virtual void SetContentBrowserViewExtender(const TSharedPtr<IContentBrowserViewExtender>& InExtender);
+
+	/** Get the currently registered custom view */
+	virtual TSharedPtr<IContentBrowserViewExtender> GetContentBrowserViewExtender();
+
 	/** Delegate accessors */
 	FOnFilterChanged& GetOnFilterChanged() { return OnFilterChanged; } 
 	FOnSearchBoxChanged& GetOnSearchBoxChanged() { return OnSearchBoxChanged; } 
@@ -248,4 +255,7 @@ private:
 
 	FDefaultSelectedPathsDelegate DefaultSelectedPathsDelegate;
 	FDefaultPathsToExpandDelegate DefaultPathsToExpandDelegate;
+
+	/** Extension used to provide a custom view to the Content Browser */
+	TSharedPtr<IContentBrowserViewExtender> ContentBrowserViewExtender;
 };
diff --git a/Engine/Source/Editor/ContentBrowser/Public/Experimental/ContentBrowserViewExtender.h b/Engine/Source/Editor/ContentBrowser/Public/Experimental/ContentBrowserViewExtender.h
new file mode 100644
--- /dev/null
+++ b/Engine/Source/Editor/ContentBrowser/Public/Experimental/ContentBrowserViewExtender.h
@@ -0,0 +1,65 @@
+// Copyright Epic Games, Inc. All Rights Reserved.
+
+#pragma once
+
+#include "AssetViewTypes.h"
+#include "Delegates/DelegateCombinations.h"
+#include "Framework/SlateDelegates.h"
+
+class SWidget;
+
+/*
+ * Interface that can be used to add a custom view to the Content Browser
+ * NOTE: This API is likely to change as it is being actively iterated on
+ */
+class IContentBrowserViewExtender
+{
+public:
+	// Create the widget for the view with the given input items
+	virtual TSharedRef<SWidget> CreateView(TArray<TSharedPtr<FAssetViewItem>>* InItemsSource) = 0;
+
+	virtual ~IContentBrowserViewExtender() = default;
+	
+	// Function called when the items source has some changes and the view needs a refresh
+    virtual void OnItemListChanged(TArray<TSharedPtr<FAssetViewItem>>* InItemsSource) = 0;
+
+	using FOnSelectionChanged = TSlateDelegates<TSharedPtr<FAssetViewItem>>::FOnSelectionChanged;
+	// Delegate the extender must fire when the selection in the UI changes
+	virtual FOnSelectionChanged& OnSelectionChanged() = 0;
+
+	using FOnItemScrolledIntoView = TSlateDelegates<TSharedPtr<FAssetViewItem>>::FOnItemScrolledIntoView;
+	// Delegate the extender must fire when an item is scrolled into view in the UI
+	virtual FOnItemScrolledIntoView& OnItemScrolledIntoView() = 0;
+	
+	using FOnMouseButtonClick = TSlateDelegates<TSharedPtr<FAssetViewItem>>::FOnMouseButtonClick;
+	// Delegate the extender must fire when an item is double clicked
+	virtual FOnMouseButtonClick& OnItemDoubleClicked() = 0;
+	
+	// Delegate the extender must fire when the context menu is opened on the list
+	virtual FOnContextMenuOpening& OnContextMenuOpened() = 0;
+
+	// Get the items selected by the custom view
+	virtual TArray<TSharedPtr<FAssetViewItem>> GetSelectedItems() = 0;
+
+	// Select the given item in the view
+	virtual void SetSelection(const TSharedPtr<FAssetViewItem>& Item, bool bSelected, const ESelectInfo::Type SelectInfo) = 0;
+
+	// Clear the current selection in the view
+	virtual void ClearSelection() = 0;
+
+	// Transfer focus to the internal view widget
+	virtual void FocusList() = 0;
+
+	// Scroll the given item into view
+	virtual void RequestScrollIntoView(const TSharedPtr<FAssetViewItem>& Item) = 0;
+
+	// true if the view is currently right click scrolling
+	virtual bool IsRightClickScrolling() = 0;
+
+	// Get a display name for the view used in the Settings menu
+	virtual FText GetViewDisplayName() = 0;
+	
+	// Get a tooltip for the custom view
+	virtual FText GetViewTooltipText() = 0;
+
+};
\ No newline at end of file
diff --git a/Engine/Source/Editor/ContentBrowser/Public/IContentBrowserSingleton.h b/Engine/Source/Editor/ContentBrowser/Public/IContentBrowserSingleton.h
--- a/Engine/Source/Editor/ContentBrowser/Public/IContentBrowserSingleton.h
+++ b/Engine/Source/Editor/ContentBrowser/Public/IContentBrowserSingleton.h
@@ -38,6 +38,7 @@ namespace EAssetViewType
 		List,
 		Tile,
 		Column,
+		Custom,
 
 		MAX
 	};
diff --git a/Engine/Source/Editor/ContentBrowser/Public/SAssetView.h b/Engine/Source/Editor/ContentBrowser/Public/SAssetView.h
--- a/Engine/Source/Editor/ContentBrowser/Public/SAssetView.h
+++ b/Engine/Source/Editor/ContentBrowser/Public/SAssetView.h
@@ -19,6 +19,7 @@
 #include "ContentBrowserTelemetry.h"
 #include "CoreMinimal.h"
 #include "Delegates/Delegate.h"
+#include "Experimental/ContentBrowserViewExtender.h"
 #include "Framework/Views/ITypedTableView.h"
 #include "HAL/Platform.h"
 #include "HistoryManager.h"
@@ -138,6 +139,7 @@ public:
 		, _ForceShowPluginContent(false)
 		, _ForceHideScrollbar(false)
 		, _ShowDisallowedAssetClassAsUnsupportedItems(false)
+		, _AllowCustomView(false)
 		{}
 
 		/** Called to check if an asset should be filtered out by external code */
@@ -287,6 +289,9 @@ public:
 		/** Allow the asset view to display the hidden asset class as unsupported items */
 		SLATE_ARGUMENT(bool, ShowDisallowedAssetClassAsUnsupportedItems)
 
+		/** Allow the asset view to display a custom view registered with the Content Browser module */
+		SLATE_ARGUMENT(bool, AllowCustomView)
+
 		/** Called to check if an asset tag should be display in details view. */
 		SLATE_EVENT( FOnShouldDisplayAssetTag, OnAssetTagWantsToBeDisplayed )
 
@@ -508,6 +513,9 @@ private:
 
 	/** Creates a new column view */
 	TSharedRef<SAssetColumnView> CreateColumnView();
+	
+	/** Creates a custom view (if specified to the content browser module) */
+	TSharedRef<SWidget> CreateCustomView();
 
 	const FSlateBrush* GetRevisionControlColumnIconBadge() const;
 
@@ -942,6 +950,9 @@ private:
 	/** Handler for Paste */
 	void ExecutePaste();
 
+	/** Check if the custom view view is available */
+	bool IsCustomViewSet() const;
+
 private:
 	friend class FAssetViewFrontendFilterHelper;
 
@@ -1289,6 +1300,12 @@ private:
 
 	/** An Id for the cache of the data sources for the filters compilation */
 	FContentBrowserDataFilterCacheIDOwner FilterCacheID;
+
+	/** An extender to create the custom view */
+	TSharedPtr<IContentBrowserViewExtender> ViewExtender;
+
+	/** The actual widget for the custom view */
+	TSharedPtr<SWidget> CustomView;
 	
 	/*
 	 * Telemetry-related fields
