TedsAssetData: Add a widget to display disk size in bytes

diff --git a/Engine/Plugins/Experimental/EditorDataStorageFeatures/Source/TedsAssetData/Private/Widgets/AssetDataDiskSizeWidget.cpp b/Engine/Plugins/Experimental/EditorDataStorageFeatures/Source/TedsAssetData/Private/Widgets/AssetDataDiskSizeWidget.cpp
new file mode 100644
--- /dev/null
+++ b/Engine/Plugins/Experimental/EditorDataStorageFeatures/Source/TedsAssetData/Private/Widgets/AssetDataDiskSizeWidget.cpp
@@ -0,0 +1,34 @@
+// Copyright Epic Games, Inc. All Rights Reserved.
+
+#include "Widgets/AssetDataDiskSizeWidget.h"
+
+#include "Elements/Framework/TypedElementAttributeBinding.h"
+#include "Internationalization/Text.h"
+#include "TedsAssetDataColumns.h"
+#include "UObject/ObjectMacros.h"
+#include "Widgets/Text/STextBlock.h"
+
+void UDiskSizeWidgetFactory::RegisterWidgetConstructors(ITypedElementDataStorageInterface& DataStorage,
+                                                        ITypedElementDataStorageUiInterface& DataStorageUi) const
+{
+	DataStorageUi.RegisterWidgetFactory<FDiskSizeWidgetConstructor>(TEXT("General.Cell"), TypedElementDataStorage::FColumn<FDiskSizeColumn>());
+}
+
+FDiskSizeWidgetConstructor::FDiskSizeWidgetConstructor()
+	: FSimpleWidgetConstructor(StaticStruct())
+{
+}
+
+TSharedPtr<SWidget> FDiskSizeWidgetConstructor::CreateWidget(ITypedElementDataStorageInterface* DataStorage,
+	ITypedElementDataStorageUiInterface* DataStorageUi, RowHandle TargetRow, RowHandle WidgetRow,
+	const TypedElementDataStorage::FMetaDataView& Arguments)
+{
+	UE::EditorDataStorage::FAttributeBinder Binder(TargetRow, DataStorage);
+	
+	return SNew(STextBlock)
+			.Text(Binder.BindData(&FDiskSizeColumn::DiskSize, [](int64 DiskSize)
+				{
+					return FText::AsMemory(DiskSize);
+				}));
+			
+}
diff --git a/Engine/Plugins/Experimental/EditorDataStorageFeatures/Source/TedsAssetData/Private/Widgets/AssetDataDiskSizeWidget.h b/Engine/Plugins/Experimental/EditorDataStorageFeatures/Source/TedsAssetData/Private/Widgets/AssetDataDiskSizeWidget.h
new file mode 100644
--- /dev/null
+++ b/Engine/Plugins/Experimental/EditorDataStorageFeatures/Source/TedsAssetData/Private/Widgets/AssetDataDiskSizeWidget.h
@@ -0,0 +1,41 @@
+// Copyright Epic Games, Inc. All Rights Reserved.
+
+#pragma once
+
+#include "Elements/Interfaces/TypedElementDataStorageFactory.h"
+#include "Elements/Interfaces/TypedElementDataStorageUiInterface.h"
+
+#include "AssetDataDiskSizeWidget.generated.h"
+
+class ITypedElementDataStorageInterface;
+class UScriptStruct;
+
+UCLASS()
+class UDiskSizeWidgetFactory : public UTypedElementDataStorageFactory
+{
+	GENERATED_BODY()
+
+public:
+	~UDiskSizeWidgetFactory() override = default;
+
+	TEDSASSETDATA_API void RegisterWidgetConstructors(ITypedElementDataStorageInterface& DataStorage,
+		ITypedElementDataStorageUiInterface& DataStorageUi) const override;
+};
+
+// Widget to show disk size in bytes
+USTRUCT()
+struct FDiskSizeWidgetConstructor : public FSimpleWidgetConstructor
+{
+	GENERATED_BODY()
+
+public:
+	TEDSASSETDATA_API FDiskSizeWidgetConstructor();
+	~FDiskSizeWidgetConstructor() override = default;
+
+	TEDSASSETDATA_API virtual TSharedPtr<SWidget> CreateWidget(
+		ITypedElementDataStorageInterface* DataStorage,
+		ITypedElementDataStorageUiInterface* DataStorageUi,
+		RowHandle TargetRow,
+		RowHandle WidgetRow, 
+		const TypedElementDataStorage::FMetaDataView& Arguments) override;
+};
\ No newline at end of file
