TedsSettings : Store additional columns ( Name, DisplayName, Description ) for containers and categories in rows of their own respective tables.

FSettingsContainerReferenceColumn and FSettingsCategoryReferenceColumn now contain a RowHandle that points to the container or category's respective row.  These handles can be followed by widgets and other code that desires to access this additional information.

FSettingsTag has been replaced by FSettingsSectionTag, which was previously a column.

diff --git a/Engine/Plugins/Experimental/EditorDataStorageFeatures/Source/TedsSettings/Private/TedsSettingsManager.cpp b/Engine/Plugins/Experimental/EditorDataStorageFeatures/Source/TedsSettings/Private/TedsSettingsManager.cpp
--- a/Engine/Plugins/Experimental/EditorDataStorageFeatures/Source/TedsSettings/Private/TedsSettingsManager.cpp
+++ b/Engine/Plugins/Experimental/EditorDataStorageFeatures/Source/TedsSettings/Private/TedsSettingsManager.cpp
@@ -3,6 +3,9 @@
 #include "TedsSettingsManager.h"
 
 #include "Elements/Columns/TypedElementCompatibilityColumns.h"
+#include "Elements/Columns/TypedElementMiscColumns.h"
+#include "Elements/Columns/TypedElementSlateWidgetColumns.h"
+#include "Elements/Framework/TypedElementIndexHasher.h"
 #include "Elements/Framework/TypedElementQueryBuilder.h"
 #include "Elements/Framework/TypedElementRegistry.h"
 #include "ISettingsCategory.h"
@@ -19,6 +22,8 @@
 FTedsSettingsManager::FTedsSettingsManager()
 	: bIsInitialized{ false }
 	, SelectAllSettingsQuery{ TypedElementDataStorage::InvalidQueryHandle }
+	, SettingsContainerTable{ TypedElementDataStorage::InvalidTableHandle }
+	, SettingsCategoryTable{ TypedElementDataStorage::InvalidTableHandle }
 {
 }
 
@@ -36,6 +41,7 @@ void FTedsSettingsManager::Initialize()
 				ITypedElementDataStorageInterface* DataStorage = TypedElementRegistry->GetMutableDataStorage();
 				check(DataStorage);
 
+				RegisterTables(*DataStorage);
 				RegisterQueries(*DataStorage);
 				RegisterSettings();
 			};
@@ -75,6 +81,23 @@ void FTedsSettingsManager::Shutdown()
 	}
 }
 
+void FTedsSettingsManager::RegisterTables(ITypedElementDataStorageInterface& DataStorage)
+{
+	if (SettingsContainerTable == TypedElementDataStorage::InvalidTableHandle)
+	{
+		SettingsContainerTable = DataStorage.RegisterTable(
+			TTypedElementColumnTypeList<FNameColumn, FDisplayNameColumn, FDescriptionColumn>(),
+			FName(TEXT("Editor_SettingsContainerTable")));
+	}
+
+	if (SettingsCategoryTable == TypedElementDataStorage::InvalidTableHandle)
+	{
+		SettingsCategoryTable = DataStorage.RegisterTable(
+			TTypedElementColumnTypeList<FNameColumn, FDisplayNameColumn, FDescriptionColumn>(),
+			FName(TEXT("Editor_SettingsCategoryTable")));
+	}
+}
+
 void FTedsSettingsManager::RegisterQueries(ITypedElementDataStorageInterface& DataStorage)
 {
 	using namespace TypedElementQueryBuilder;
@@ -83,7 +106,9 @@ void FTedsSettingsManager::RegisterQueries(ITypedElementDataStorageInterface& Da
 	{
 		SelectAllSettingsQuery = DataStorage.RegisterQuery(
 			Select()
-			.ReadOnly<FTypedElementUObjectColumn, FSettingsContainerColumn, FSettingsCategoryColumn, FSettingsSectionColumn>()
+				.ReadOnly<FSettingsContainerReferenceColumn, FSettingsCategoryReferenceColumn, FNameColumn>()
+			.Where()
+				.All<FSettingsSectionTag>()
 			.Compile());
 	}
 }
@@ -114,32 +139,45 @@ void FTedsSettingsManager::RegisterSettings()
 
 void FTedsSettingsManager::RegisterSettingsContainer(const FName& ContainerName)
 {
+	using namespace UE::Editor::DataStorage;
+
 	TRACE_CPUPROFILER_EVENT_SCOPE(TedsSettingsManager.RegisterSettingsContainer);
 
 	ISettingsModule* SettingsModule = FModuleManager::GetModulePtr<ISettingsModule>("Settings");
 	check(SettingsModule);
 
+	UTypedElementRegistry* TypedElementRegistry = UTypedElementRegistry::GetInstance();
+	check(TypedElementRegistry);
+
+	ITypedElementDataStorageInterface* DataStorage = TypedElementRegistry->GetMutableDataStorage();
+	check(DataStorage);
+
 	UE_LOG(LogTedsSettings, Log, TEXT("Register Settings Container : '%s'"), *ContainerName.ToString());
 
 	ISettingsContainerPtr ContainerPtr = SettingsModule->GetContainer(ContainerName);
 
+	RowHandle ContainerRow = DataStorage->AddRow(SettingsContainerTable);
+	DataStorage->AddColumn<FNameColumn>(ContainerRow, { .Name = ContainerName });
+	DataStorage->AddColumn<FDisplayNameColumn>(ContainerRow, { .DisplayName = ContainerPtr->GetDisplayName() });
+	DataStorage->AddColumn<FDescriptionColumn>(ContainerRow, { .Description = ContainerPtr->GetDescription() });
+
 	TArray<ISettingsCategoryPtr> Categories;
 	ContainerPtr->GetCategories(Categories);
 
 	for (ISettingsCategoryPtr CategoryPtr : Categories)
 	{
 		const bool bQueryExistingRows = false;
-		UpdateSettingsCategory(CategoryPtr, ContainerName, bQueryExistingRows);
+		UpdateSettingsCategory(CategoryPtr, ContainerRow, bQueryExistingRows);
 	}
 
 	// OnCategoryModified is called at the same time as OnSectionRemoved so we only bind to OnCategoryModified for add / update / remove
-	ContainerPtr->OnCategoryModified().AddSPLambda(this, [this, ContainerPtr](const FName& ModifiedCategoryName)
+	ContainerPtr->OnCategoryModified().AddSPLambda(this, [this, ContainerPtr, ContainerRow](const FName& ModifiedCategoryName)
 		{
 			UE_LOG(LogTedsSettings, Log, TEXT("Settings Category modified : '%s->%s'"), *ContainerPtr->GetName().ToString(), *ModifiedCategoryName.ToString());
 
 			ISettingsCategoryPtr CategoryPtr = ContainerPtr->GetCategory(ModifiedCategoryName);
 
-			UpdateSettingsCategory(CategoryPtr, ContainerPtr->GetName());
+			UpdateSettingsCategory(CategoryPtr, ContainerRow);
 		});
 }
 
@@ -193,16 +231,13 @@ void FTedsSettingsManager::UnregisterSettings()
 	}
 }
 
-void FTedsSettingsManager::UpdateSettingsCategory(TSharedPtr<ISettingsCategory> SettingsCategory, const FName& ContainerName, const bool bQueryExistingRows)
+void FTedsSettingsManager::UpdateSettingsCategory(TSharedPtr<ISettingsCategory> SettingsCategory, UE::Editor::DataStorage::RowHandle ContainerRow, const bool bQueryExistingRows)
 {
+	using namespace TypedElementDataStorage;
 	using namespace UE::Editor::DataStorage;
 
 	TRACE_CPUPROFILER_EVENT_SCOPE(TedsSettingsManager.UpdateSettingsCategory);
 
-	const FName& CategoryName = SettingsCategory->GetName();
-
-	UE_LOG(LogTedsSettings, Log, TEXT("Update Settings Category: '%s->%s'"), *ContainerName.ToString(), *CategoryName.ToString());
-
 	UTypedElementRegistry* TypedElementRegistry = UTypedElementRegistry::GetInstance();
 	check(TypedElementRegistry);
 
@@ -212,6 +247,25 @@ void FTedsSettingsManager::UpdateSettingsCategory(TSharedPtr<ISettingsCategory>
 	ITypedElementDataStorageCompatibilityInterface* DataStorageCompatibility = TypedElementRegistry->GetMutableDataStorageCompatibility();
 	check(DataStorageCompatibility);
 
+	const FName& ContainerName = DataStorage->GetColumn<FNameColumn>(ContainerRow)->Name;
+	const FName& CategoryName = SettingsCategory->GetName();
+
+	UE_LOG(LogTedsSettings, Log, TEXT("Update Settings Category: '%s->%s'"), *ContainerName.ToString(), *CategoryName.ToString());
+
+	uint64 CategoryIndexHash = GenerateIndexHash(SettingsCategory.Get());
+
+	RowHandle CategoryRow = DataStorage->FindIndexedRow(CategoryIndexHash);
+	if (CategoryRow == InvalidRowHandle)
+	{
+		CategoryRow = DataStorage->AddRow(SettingsCategoryTable);
+
+		DataStorage->AddColumn<FNameColumn>(CategoryRow, { .Name = CategoryName });
+		DataStorage->AddColumn<FDisplayNameColumn>(CategoryRow, { .DisplayName = SettingsCategory->GetDisplayName() });
+		DataStorage->AddColumn<FDescriptionColumn>(CategoryRow, { .Description = SettingsCategory->GetDescription() });
+
+		DataStorage->IndexRow(CategoryIndexHash, CategoryRow);
+	}
+
 	TArray<RowHandle> OldRowHandles;
 	TArray<FName> OldSectionNames;
 
@@ -224,10 +278,9 @@ void FTedsSettingsManager::UpdateSettingsCategory(TSharedPtr<ISettingsCategory>
 		DataStorage->RunQuery(SelectAllSettingsQuery, CreateDirectQueryCallbackBinding(
 			[&OldRowHandles, &OldSectionNames, &ContainerName, &CategoryName](
 				DSI::IDirectQueryContext& Context,
-				const FTypedElementUObjectColumn* ObjectColumns,
-				const FSettingsContainerColumn* ContainerColumns,
-				const FSettingsCategoryColumn* CategoryColumns,
-				const FSettingsSectionColumn* SectionColumns)
+				const FSettingsContainerReferenceColumn* ContainerColumns,
+				const FSettingsCategoryReferenceColumn* CategoryColumns,
+				const FNameColumn* SectionNameColumns)
 			{
 				const uint32 RowCount = Context.GetRowCount();
 
@@ -239,7 +292,7 @@ void FTedsSettingsManager::UpdateSettingsCategory(TSharedPtr<ISettingsCategory>
 						TempCategoryName == CategoryName)
 					{
 						OldRowHandles.Emplace(Context.GetRowHandles()[RowIndex]);
-						OldSectionNames.Emplace(SectionColumns[RowIndex].SectionName);
+						OldSectionNames.Emplace(SectionNameColumns[RowIndex].Name);
 					}
 				}
 			}));
@@ -269,10 +322,12 @@ void FTedsSettingsManager::UpdateSettingsCategory(TSharedPtr<ISettingsCategory>
 
 			RowHandle NewRow = DataStorageCompatibility->AddCompatibleObject(SettingsObjectPtr);
 
-			DataStorage->AddColumn<FSettingsContainerColumn>(NewRow, { .ContainerName = ContainerName });
-			DataStorage->AddColumn<FSettingsCategoryColumn>(NewRow, { .CategoryName = CategoryName });
-			DataStorage->AddColumn<FSettingsSectionColumn>(NewRow, { .SectionName = SectionName });
-			DataStorage->AddColumn<FSettingsTag>(NewRow);
+			DataStorage->AddColumn<FSettingsContainerReferenceColumn>(NewRow, { .ContainerName = ContainerName, .ContainerRow = ContainerRow });
+			DataStorage->AddColumn<FSettingsCategoryReferenceColumn>(NewRow, { .CategoryName = CategoryName, .CategoryRow = CategoryRow });
+			DataStorage->AddColumn<FNameColumn>(NewRow, { .Name = SectionName });
+			DataStorage->AddColumn<FDisplayNameColumn>(NewRow, { .DisplayName = SectionPtr->GetDisplayName() });
+			DataStorage->AddColumn<FDescriptionColumn>(NewRow, { .Description = SectionPtr->GetDescription() });
+			DataStorage->AddColumn<FSettingsSectionTag>(NewRow);
 
 			UE_LOG(LogTedsSettings, Log, TEXT("Added Settings Section : '%s'"), *SectionName.ToString());
 		}
diff --git a/Engine/Plugins/Experimental/EditorDataStorageFeatures/Source/TedsSettings/Private/TedsSettingsManager.h b/Engine/Plugins/Experimental/EditorDataStorageFeatures/Source/TedsSettings/Private/TedsSettingsManager.h
--- a/Engine/Plugins/Experimental/EditorDataStorageFeatures/Source/TedsSettings/Private/TedsSettingsManager.h
+++ b/Engine/Plugins/Experimental/EditorDataStorageFeatures/Source/TedsSettings/Private/TedsSettingsManager.h
@@ -23,6 +23,8 @@ public:
 
 private:
 
+	void RegisterTables(ITypedElementDataStorageInterface& DataStorage);
+
 	void RegisterQueries(ITypedElementDataStorageInterface& DataStorage);
 	void UnregisterQueries(ITypedElementDataStorageInterface& DataStorage);
 
@@ -31,9 +33,11 @@ private:
 
 	void RegisterSettingsContainer(const FName& ContainerName);
 
-	void UpdateSettingsCategory(TSharedPtr<ISettingsCategory> SettingsCategory, const FName& ContainerName, const bool bQueryExistingRows = true);
+	void UpdateSettingsCategory(TSharedPtr<ISettingsCategory> SettingsCategory, UE::Editor::DataStorage::RowHandle ContainerRow, const bool bQueryExistingRows = true);
 
 	bool bIsInitialized;
 	TypedElementDataStorage::QueryHandle SelectAllSettingsQuery;
+	UE::Editor::DataStorage::TableHandle SettingsContainerTable;
+	UE::Editor::DataStorage::TableHandle SettingsCategoryTable;
 
 };
diff --git a/Engine/Plugins/Experimental/EditorDataStorageFeatures/Source/TedsSettings/Public/TedsSettingsColumns.h b/Engine/Plugins/Experimental/EditorDataStorageFeatures/Source/TedsSettings/Public/TedsSettingsColumns.h
--- a/Engine/Plugins/Experimental/EditorDataStorageFeatures/Source/TedsSettings/Public/TedsSettingsColumns.h
+++ b/Engine/Plugins/Experimental/EditorDataStorageFeatures/Source/TedsSettings/Public/TedsSettingsColumns.h
@@ -7,35 +7,32 @@
 
 #include "TedsSettingsColumns.generated.h"
 
-USTRUCT(meta = (DisplayName = "Settings Container"))
-struct FSettingsContainerColumn final : public FEditorDataStorageColumn
+USTRUCT(meta = (DisplayName = "Settings Container Reference"))
+struct FSettingsContainerReferenceColumn final : public FEditorDataStorageColumn
 {
 	GENERATED_BODY()
 
+	// Store the name as it is frequently accessed to avoid indirection to the row.
 	UPROPERTY(meta = (Searchable))
 	FName ContainerName;
+
+	UE::Editor::DataStorage::RowHandle ContainerRow;
 };
 
-USTRUCT(meta = (DisplayName = "Settings Category"))
-struct FSettingsCategoryColumn final : public FEditorDataStorageColumn
+USTRUCT(meta = (DisplayName = "Settings Category Reference"))
+struct FSettingsCategoryReferenceColumn final : public FEditorDataStorageColumn
 {
 	GENERATED_BODY()
 
+	// Store the name as it is frequently accessed to avoid indirection to the row.
 	UPROPERTY(meta = (Searchable))
 	FName CategoryName;
-};
-
-USTRUCT(meta = (DisplayName = "Settings Section"))
-struct FSettingsSectionColumn final : public FEditorDataStorageColumn
-{
-	GENERATED_BODY()
 
-	UPROPERTY(meta = (Searchable))
-	FName SectionName;
+	UE::Editor::DataStorage::RowHandle CategoryRow;
 };
 
-USTRUCT(meta = (DisplayName = "Settings"))
-struct FSettingsTag final : public FEditorDataStorageTag
+USTRUCT(meta = (DisplayName = "Settings Section"))
+struct FSettingsSectionTag final : public FEditorDataStorageTag
 {
 	GENERATED_BODY()
 };
diff --git a/Engine/Plugins/Experimental/EditorDataStorageFeatures/Source/TedsSettings/Tests/TedsSettingsTests.cpp b/Engine/Plugins/Experimental/EditorDataStorageFeatures/Source/TedsSettings/Tests/TedsSettingsTests.cpp
--- a/Engine/Plugins/Experimental/EditorDataStorageFeatures/Source/TedsSettings/Tests/TedsSettingsTests.cpp
+++ b/Engine/Plugins/Experimental/EditorDataStorageFeatures/Source/TedsSettings/Tests/TedsSettingsTests.cpp
@@ -7,6 +7,7 @@
 #include "Containers/Ticker.h"
 #include "Editor.h"
 #include "Elements/Columns/TypedElementCompatibilityColumns.h"
+#include "Elements/Columns/TypedElementMiscColumns.h"
 #include "Elements/Framework/TypedElementQueryBuilder.h"
 #include "Elements/Framework/TypedElementRegistry.h"
 #include "Elements/Interfaces/TypedElementDataStorageInterface.h"
@@ -92,7 +93,7 @@ namespace UE::Editor::Settings::Tests
 				CountAllSettingsQuery = DataStorage->RegisterQuery(
 					Count()
 					.Where()
-					.All<FTypedElementUObjectColumn, FSettingsContainerColumn, FSettingsCategoryColumn, FSettingsSectionColumn>()
+					.All<FSettingsSectionTag>()
 					.Compile());
 			}
 
@@ -142,9 +143,9 @@ namespace UE::Editor::Settings::Tests
 
 					TestEqual(TEXT("RowCount"), AfterRowCount, BeforeRowCount + 1);
 
-					TestEqual(TEXT("ContainerName"), DataStorage->GetColumn<FSettingsContainerColumn>(RowHandle)->ContainerName, ContainerName);
-					TestEqual(TEXT("CategoryName"), DataStorage->GetColumn<FSettingsCategoryColumn>(RowHandle)->CategoryName, CategoryName);
-					TestEqual(TEXT("SectionName"), DataStorage->GetColumn<FSettingsSectionColumn>(RowHandle)->SectionName, SectionName);
+					TestEqual(TEXT("ContainerName"), DataStorage->GetColumn<FSettingsContainerReferenceColumn>(RowHandle)->ContainerName, ContainerName);
+					TestEqual(TEXT("CategoryName"), DataStorage->GetColumn<FSettingsCategoryReferenceColumn>(RowHandle)->CategoryName, CategoryName);
+					TestEqual(TEXT("SectionName"), DataStorage->GetColumn<FNameColumn>(RowHandle)->Name, SectionName);
 				});
 			});
 		});
diff --git a/Engine/Source/Runtime/TypedElementFramework/Public/Elements/Columns/TypedElementMiscColumns.h b/Engine/Source/Runtime/TypedElementFramework/Public/Elements/Columns/TypedElementMiscColumns.h
--- a/Engine/Source/Runtime/TypedElementFramework/Public/Elements/Columns/TypedElementMiscColumns.h
+++ b/Engine/Source/Runtime/TypedElementFramework/Public/Elements/Columns/TypedElementMiscColumns.h
@@ -54,4 +54,18 @@ struct FTypedElementRowReferenceColumn final : public FEditorDataStorageColumn
 	GENERATED_BODY()
 
 	UE::Editor::DataStorage::RowHandle Row;
-};
\ No newline at end of file
+};
+
+/**
+ * A name for this row.
+ * 
+ * This can be used as a dynamic column to specify names for multiple items in a row.
+ */
+USTRUCT(meta = (DisplayName = "Name"))
+struct FNameColumn final : public FEditorDataStorageColumn
+{
+	GENERATED_BODY()
+
+	UPROPERTY(meta = (Searchable))
+	FName Name;
+};
diff --git a/Engine/Source/Runtime/TypedElementFramework/Public/Elements/Columns/TypedElementSlateWidgetColumns.h b/Engine/Source/Runtime/TypedElementFramework/Public/Elements/Columns/TypedElementSlateWidgetColumns.h
--- a/Engine/Source/Runtime/TypedElementFramework/Public/Elements/Columns/TypedElementSlateWidgetColumns.h
+++ b/Engine/Source/Runtime/TypedElementFramework/Public/Elements/Columns/TypedElementSlateWidgetColumns.h
@@ -41,4 +41,32 @@ USTRUCT(meta = (DisplayName = "Slate widget reference deletes row"))
 struct FTypedElementSlateWidgetReferenceDeletesRowTag final : public FEditorDataStorageTag
 {
 	GENERATED_BODY()
-};
\ No newline at end of file
+};
+
+/**
+ * A localized display name for this row.
+ * 
+ * This can be used as a dynamic column to specify display names for multiple items in a row.
+ */
+USTRUCT(meta = (DisplayName = "Display Name"))
+struct FDisplayNameColumn final : public FEditorDataStorageColumn
+{
+	GENERATED_BODY()
+
+	UPROPERTY(meta = (Searchable))
+	FText DisplayName;
+};
+
+/**
+ * A localized description for this row.
+ * 
+ * This can be used as a dynamic column to specify descriptions for multiple items in a row.
+ */
+USTRUCT(meta = (DisplayName = "Description"))
+struct FDescriptionColumn final : public FEditorDataStorageColumn
+{
+	GENERATED_BODY()
+
+	UPROPERTY(meta = (Searchable))
+	FText Description;
+};
