SceneOutlinerTedsBridge: Fix missing widget purpose

diff --git a/Engine/Plugins/Experimental/EditorDataStorageFeatures/Source/TedsOutliner/Private/Compatibility/SceneOutlinerTedsBridge.cpp b/Engine/Plugins/Experimental/EditorDataStorageFeatures/Source/TedsOutliner/Private/Compatibility/SceneOutlinerTedsBridge.cpp
--- a/Engine/Plugins/Experimental/EditorDataStorageFeatures/Source/TedsOutliner/Private/Compatibility/SceneOutlinerTedsBridge.cpp
+++ b/Engine/Plugins/Experimental/EditorDataStorageFeatures/Source/TedsOutliner/Private/Compatibility/SceneOutlinerTedsBridge.cpp
@@ -362,8 +362,8 @@ const FName FSceneOutlinerTedsQueryBinder::HeaderWidgetPurpose(TEXT("SceneOutlin
 const FName FSceneOutlinerTedsQueryBinder::DefaultHeaderWidgetPurpose(TEXT("SceneOutliner.Header.Default"));
 const FName FSceneOutlinerTedsQueryBinder::CellWidgetPurpose(TEXT("SceneOutliner.Cell"));
 const FName FSceneOutlinerTedsQueryBinder::DefaultCellWidgetPurpose(TEXT("SceneOutliner.Cell.Default"));
-const FName FSceneOutlinerTedsQueryBinder::ItemLabelCellWidgetPurpose(TEXT("SceneOutliner.ItemLabel.Cell"));
-const FName FSceneOutlinerTedsQueryBinder::DefaultItemLabelCellWidgetPurpose(TEXT("SceneOutliner.ItemLabel.Cell.Default"));
+const FName FSceneOutlinerTedsQueryBinder::ItemLabelCellWidgetPurpose(TEXT("SceneOutliner.RowLabel"));
+const FName FSceneOutlinerTedsQueryBinder::DefaultItemLabelCellWidgetPurpose(TEXT("SceneOutliner.RowLabel.Default"));
 
 FSceneOutlinerTedsQueryBinder::FSceneOutlinerTedsQueryBinder()
 {
