TedsSettings : Add FSettingsContainerTag and FSettingsCategoryTag

diff --git a/Engine/Plugins/Experimental/EditorDataStorageFeatures/Source/TedsSettings/Private/TedsSettingsManager.cpp b/Engine/Plugins/Experimental/EditorDataStorageFeatures/Source/TedsSettings/Private/TedsSettingsManager.cpp
--- a/Engine/Plugins/Experimental/EditorDataStorageFeatures/Source/TedsSettings/Private/TedsSettingsManager.cpp
+++ b/Engine/Plugins/Experimental/EditorDataStorageFeatures/Source/TedsSettings/Private/TedsSettingsManager.cpp
@@ -86,14 +86,14 @@ void FTedsSettingsManager::RegisterTables(ITypedElementDataStorageInterface& Dat
 	if (SettingsContainerTable == TypedElementDataStorage::InvalidTableHandle)
 	{
 		SettingsContainerTable = DataStorage.RegisterTable(
-			TTypedElementColumnTypeList<FNameColumn, FDisplayNameColumn, FDescriptionColumn>(),
+			TTypedElementColumnTypeList<FNameColumn, FDisplayNameColumn, FDescriptionColumn, FSettingsContainerTag>(),
 			FName(TEXT("Editor_SettingsContainerTable")));
 	}
 
 	if (SettingsCategoryTable == TypedElementDataStorage::InvalidTableHandle)
 	{
 		SettingsCategoryTable = DataStorage.RegisterTable(
-			TTypedElementColumnTypeList<FNameColumn, FDisplayNameColumn, FDescriptionColumn>(),
+			TTypedElementColumnTypeList<FSettingsContainerReferenceColumn, FNameColumn, FDisplayNameColumn, FDescriptionColumn, FSettingsCategoryTag>(),
 			FName(TEXT("Editor_SettingsCategoryTable")));
 	}
 }
@@ -160,6 +160,7 @@ void FTedsSettingsManager::RegisterSettingsContainer(const FName& ContainerName)
 	DataStorage->AddColumn<FNameColumn>(ContainerRow, { .Name = ContainerName });
 	DataStorage->AddColumn<FDisplayNameColumn>(ContainerRow, { .DisplayName = ContainerPtr->GetDisplayName() });
 	DataStorage->AddColumn<FDescriptionColumn>(ContainerRow, { .Description = ContainerPtr->GetDescription() });
+	DataStorage->AddColumn<FSettingsContainerTag>(ContainerRow);
 
 	TArray<ISettingsCategoryPtr> Categories;
 	ContainerPtr->GetCategories(Categories);
@@ -259,9 +260,11 @@ void FTedsSettingsManager::UpdateSettingsCategory(TSharedPtr<ISettingsCategory>
 	{
 		CategoryRow = DataStorage->AddRow(SettingsCategoryTable);
 
+		DataStorage->AddColumn<FSettingsContainerReferenceColumn>(CategoryRow, { .ContainerName = ContainerName, .ContainerRow = ContainerRow });
 		DataStorage->AddColumn<FNameColumn>(CategoryRow, { .Name = CategoryName });
 		DataStorage->AddColumn<FDisplayNameColumn>(CategoryRow, { .DisplayName = SettingsCategory->GetDisplayName() });
 		DataStorage->AddColumn<FDescriptionColumn>(CategoryRow, { .Description = SettingsCategory->GetDescription() });
+		DataStorage->AddColumn<FSettingsCategoryTag>(CategoryRow);
 
 		DataStorage->IndexRow(CategoryIndexHash, CategoryRow);
 	}
diff --git a/Engine/Plugins/Experimental/EditorDataStorageFeatures/Source/TedsSettings/Public/TedsSettingsColumns.h b/Engine/Plugins/Experimental/EditorDataStorageFeatures/Source/TedsSettings/Public/TedsSettingsColumns.h
--- a/Engine/Plugins/Experimental/EditorDataStorageFeatures/Source/TedsSettings/Public/TedsSettingsColumns.h
+++ b/Engine/Plugins/Experimental/EditorDataStorageFeatures/Source/TedsSettings/Public/TedsSettingsColumns.h
@@ -31,6 +31,18 @@ struct FSettingsCategoryReferenceColumn final : public FEditorDataStorageColumn
 	UE::Editor::DataStorage::RowHandle CategoryRow;
 };
 
+USTRUCT(meta = (DisplayName = "Settings Container"))
+struct FSettingsContainerTag final : public FEditorDataStorageTag
+{
+	GENERATED_BODY()
+};
+
+USTRUCT(meta = (DisplayName = "Settings Category"))
+struct FSettingsCategoryTag final : public FEditorDataStorageTag
+{
+	GENERATED_BODY()
+};
+
 USTRUCT(meta = (DisplayName = "Settings Section"))
 struct FSettingsSectionTag final : public FEditorDataStorageTag
 {
