[Teds] Update namespace from TypedElementDataStorage to UE::Editor::DataStorage for metadata and query store.

diff --git a/Engine/Plugins/Experimental/EditorDataStorage/Source/TedsCore/Tests/Commands/EditorDataStorageCommandBufferTests.cpp b/Engine/Plugins/Experimental/EditorDataStorage/Source/TedsCore/Tests/Commands/EditorDataStorageCommandBufferTests.cpp
--- a/Engine/Plugins/Experimental/EditorDataStorage/Source/TedsCore/Tests/Commands/EditorDataStorageCommandBufferTests.cpp
+++ b/Engine/Plugins/Experimental/EditorDataStorage/Source/TedsCore/Tests/Commands/EditorDataStorageCommandBufferTests.cpp
@@ -81,7 +81,7 @@ struct FSequenceTestingProcessor
 	}
 };
 
-TEST_CASE_NAMED(EditorDataStorage_FCommandBuffer, "EditorDataStorage::Command Buffer", "[ApplicationContextMask][EngineFilter]")
+TEST_CASE_NAMED(EditorDataStorage_FCommandBuffer, "Editor::DataStorage::Command Buffer", "[ApplicationContextMask][EngineFilter]")
 {
 	using namespace UE::Editor::DataStorage;
 	
@@ -129,7 +129,7 @@ TEST_CASE_NAMED(EditorDataStorage_FCommandBuffer, "EditorDataStorage::Command Bu
 	}
 }
 
-TEST_CASE_NAMED(EditorDataStorage_FCommandBuffer_FCollection, "EditorDataStorage::Command Buffer::FCollection", "[ApplicationContextMask][EngineFilter]")
+TEST_CASE_NAMED(EditorDataStorage_FCommandBuffer_FCollection, "Editor::DataStorage::Command Buffer::FCollection", "[ApplicationContextMask][EngineFilter]")
 {
 	using namespace UE::Editor::DataStorage;
 
@@ -401,7 +401,7 @@ TEST_CASE_NAMED(EditorDataStorage_FCommandBuffer_FCollection, "EditorDataStorage
 	}
 }
 
-TEST_CASE_NAMED(EditorDataStorage_FCommandBuffer_FOptimizer, "EditorDataStorage::Command Buffer::FOptimizer", "[ApplicationContextMask][EngineFilter]")
+TEST_CASE_NAMED(EditorDataStorage_FCommandBuffer_FOptimizer, "Editor::DataStorage::Command Buffer::FOptimizer", "[ApplicationContextMask][EngineFilter]")
 {
 	using namespace UE::Editor::DataStorage;
 
diff --git a/Engine/Plugins/Experimental/EditorDataStorage/Source/TedsCore/Tests/DynamicColumnTests.cpp b/Engine/Plugins/Experimental/EditorDataStorage/Source/TedsCore/Tests/DynamicColumnTests.cpp
--- a/Engine/Plugins/Experimental/EditorDataStorage/Source/TedsCore/Tests/DynamicColumnTests.cpp
+++ b/Engine/Plugins/Experimental/EditorDataStorage/Source/TedsCore/Tests/DynamicColumnTests.cpp
@@ -1,4 +1,4 @@
-﻿// Copyright Epic Games, Inc. All Rights Reserved.
+// Copyright Epic Games, Inc. All Rights Reserved.
 
 #include "Elements/Framework/TypedElementRegistry.h"
 #include "Elements/Framework/TypedElementTestColumns.h"
@@ -9,23 +9,25 @@
 
 #include "Misc/AutomationTest.h"
 
-namespace UE::Editor::DataStorage::Tests
+namespace UE::Editor::DataStorage
+{
+namespace Tests
 {
 	BEGIN_DEFINE_SPEC(DynamicColumnTestFixture, "Editor.DataStorage.DynamicColumns", EAutomationTestFlags::EditorContext | EAutomationTestFlags::EngineFilter)
 		ITypedElementDataStorageInterface* TedsInterface = nullptr;
 		const FName TestTableName = TEXT("TestTable_DynamicColumnsTest");
 
-		TypedElementDataStorage::TableHandle TestTable;
-		TArray<TypedElementDataStorage::RowHandle> Rows;
-		TArray<TypedElementDataStorage::QueryHandle> QueryHandles;
+		TableHandle TestTable;
+		TArray<RowHandle> Rows;
+		TArray<QueryHandle> QueryHandles;
 
 		TArray<FName> Identifiers;
 
-		TypedElementDataStorage::TableHandle RegisterTestTable() const
+		TableHandle RegisterTestTable() const
 		{
-			const TypedElementDataStorage::TableHandle Table = TedsInterface->FindTable(TestTableName);
+			const TableHandle Table = TedsInterface->FindTable(TestTableName);
 				
-			if (Table != TypedElementDataStorage::InvalidTableHandle)
+			if (Table != InvalidTableHandle)
 			{
 				return Table;
 			}
@@ -37,16 +39,16 @@ namespace UE::Editor::DataStorage::Tests
 			TestTableName);
 		}
 
-		TypedElementDataStorage::RowHandle CreateTestRow(TypedElementDataStorage::TableHandle InTableHandle)
+		RowHandle CreateTestRow(TableHandle InTableHandle)
 		{
-			TypedElementDataStorage::RowHandle RowHandle = TedsInterface->AddRow(InTableHandle);
-			Rows.Add(RowHandle);
-			return RowHandle;
+			RowHandle Row = TedsInterface->AddRow(InTableHandle);
+			Rows.Add(Row);
+			return Row;
 		}
 
-		TypedElementDataStorage::QueryHandle RegisterQuery(TypedElementDataStorage::FQueryDescription&& Query)
+		QueryHandle RegisterQuery(FQueryDescription&& Query)
 		{
-			TypedElementDataStorage::QueryHandle QueryHandle = TedsInterface->RegisterQuery(MoveTemp(Query));
+			QueryHandle QueryHandle = TedsInterface->RegisterQuery(MoveTemp(Query));
 			QueryHandles.Add(QueryHandle);
 			return QueryHandle;
 		}
@@ -152,11 +154,10 @@ namespace UE::Editor::DataStorage::Tests
 					Async(EAsyncExecution::TaskGraphMainThread,[this]()
 					{
 						using namespace TypedElementQueryBuilder;
-						using namespace UE::Editor::DataStorage;
 
-						TArray<TypedElementDataStorage::RowHandle> RowsToMatch;
+						TArray<RowHandle> RowsToMatch;
 						TBitArray WasMatched;
-						auto SetExpectedMatches = [&RowsToMatch, &WasMatched](TConstArrayView<TypedElementDataStorage::RowHandle> Expectation)
+						auto SetExpectedMatches = [&RowsToMatch, &WasMatched](TConstArrayView<RowHandle> Expectation)
 						{
 							RowsToMatch.Empty(RowsToMatch.Num());
 							RowsToMatch.Append(Expectation);
@@ -167,10 +168,10 @@ namespace UE::Editor::DataStorage::Tests
 						{
 							return WasMatched.CountSetBits();
 						};
-						auto Callback = CreateDirectQueryCallbackBinding([this, &RowsToMatch, &WasMatched] (TypedElementDataStorage::IDirectQueryContext& Context, const TypedElementDataStorage::RowHandle* CallbackRows)
+						auto Callback = CreateDirectQueryCallbackBinding([this, &RowsToMatch, &WasMatched] (IDirectQueryContext& Context, const RowHandle* CallbackRows)
 						{
-							TConstArrayView<TypedElementDataStorage::RowHandle> RowsView = MakeConstArrayView(CallbackRows, Context.GetRowCount());
-							for (TypedElementDataStorage::RowHandle Row : RowsView)
+							TConstArrayView<RowHandle> RowsView = MakeConstArrayView(CallbackRows, Context.GetRowCount());
+							for (RowHandle Row : RowsView)
 							{
 								int32 Index = RowsToMatch.Find(Row);
 								TestTrue(TEXT("Returned row in query is within expected match array"), Index != INDEX_NONE);
@@ -184,7 +185,7 @@ namespace UE::Editor::DataStorage::Tests
 
 						// Should match Rows[0]
 						{
-							TypedElementDataStorage::QueryHandle Query = RegisterQuery(
+							QueryHandle Query = RegisterQuery(
 								Select().
 								Where().
 									All<FTestDynamicTag>(Identifiers[0]).
@@ -197,7 +198,7 @@ namespace UE::Editor::DataStorage::Tests
 						}
 						{
 							// Should match Rows 0 and 1
-							TypedElementDataStorage::QueryHandle Query = RegisterQuery(
+							QueryHandle Query = RegisterQuery(
 								Select().
 									Where().
 										All<FTestDynamicTag>(Identifiers[0]).
@@ -208,7 +209,7 @@ namespace UE::Editor::DataStorage::Tests
 						}
 						{
 							// Should match Rows 1
-							TypedElementDataStorage::QueryHandle Query = RegisterQuery(
+							QueryHandle Query = RegisterQuery(
 								Select().
 									Where().
 										All<FTestDynamicTag>(Identifiers[0]).
@@ -220,7 +221,7 @@ namespace UE::Editor::DataStorage::Tests
 						}
 						{
 							// Should match Rows 0 and 2
-							TypedElementDataStorage::QueryHandle Query = RegisterQuery(
+							QueryHandle Query = RegisterQuery(
 								Select().
 									Where().
 										All<FTestDynamicTag>(Identifiers[1]).
@@ -231,7 +232,7 @@ namespace UE::Editor::DataStorage::Tests
 						}
 						{
 							// Should match Rows 0 and 2
-							TypedElementDataStorage::QueryHandle Query = RegisterQuery(
+							QueryHandle Query = RegisterQuery(
 								Select().
 									Where().
 										None<FTestDynamicTag>(Identifiers[0]).
@@ -243,7 +244,7 @@ namespace UE::Editor::DataStorage::Tests
 						}
 						{
 							// Should match Rows 0, 1 and 2
-							TypedElementDataStorage::QueryHandle Query = RegisterQuery(
+							QueryHandle Query = RegisterQuery(
 								Select().
 									Where().
 										Any<FTestDynamicTag>(Identifiers[0]).
@@ -258,10 +259,10 @@ namespace UE::Editor::DataStorage::Tests
 
 				// Processor Query
 				{
-					TArray<TypedElementDataStorage::RowHandle> RowsToMatch;
+					TArray<RowHandle> RowsToMatch;
 					TArray<int32> MatchCount;
 					int32 UnexpectedRowCount;
-					auto SetExpectedMatches = [&RowsToMatch, &MatchCount, &UnexpectedRowCount](TConstArrayView<TypedElementDataStorage::RowHandle> Expectation)
+					auto SetExpectedMatches = [&RowsToMatch, &MatchCount, &UnexpectedRowCount](TConstArrayView<RowHandle> Expectation)
 					{
 						RowsToMatch.Empty(RowsToMatch.Num());
 						RowsToMatch.Append(Expectation);
@@ -275,13 +276,13 @@ namespace UE::Editor::DataStorage::Tests
 					};
 
 					TArray<FName> ActivationKeys;
-					TArray<TArray<TypedElementDataStorage::RowHandle>> QueryExpectedMatchRows;
+					TArray<TArray<RowHandle>> QueryExpectedMatchRows;
 					
 					Async(EAsyncExecution::TaskGraphMainThread,[this, &RowsToMatch, &MatchCount, &UnexpectedRowCount, &ActivationKeys, &QueryExpectedMatchRows]()
 					{
-						auto Callback = [this, &RowsToMatch, &MatchCount, &UnexpectedRowCount] (TypedElementDataStorage::IQueryContext& Context, const TypedElementDataStorage::RowHandle* CallbackRows)
+						auto Callback = [this, &RowsToMatch, &MatchCount, &UnexpectedRowCount] (IQueryContext& Context, const RowHandle* CallbackRows)
 						{
-							TConstArrayView<TypedElementDataStorage::RowHandle> RowsView = MakeConstArrayView(CallbackRows, Context.GetRowCount());
+							TConstArrayView<RowHandle> RowsView = MakeConstArrayView(CallbackRows, Context.GetRowCount());
 							for (auto Row : RowsView)
 							{
 								int32 Index = RowsToMatch.Find(Row);
@@ -472,7 +473,7 @@ namespace UE::Editor::DataStorage::Tests
 					
 					struct FExpectation
 					{
-						TypedElementDataStorage::RowHandle Row;
+						RowHandle Row;
 						TArray<int32> ColumnValues;
 
 						int32 MatchCount = 0;
@@ -597,13 +598,13 @@ namespace UE::Editor::DataStorage::Tests
 		AfterEach([this]()
 		{
 			Identifiers.Empty(Identifiers.Num());
-			for (TypedElementDataStorage::RowHandle Row : Rows)
+			for (RowHandle Row : Rows)
 			{
 				TedsInterface->RemoveRow(Row);
 			}
 			Rows.Empty(Rows.Num());
 
-			for (TypedElementDataStorage::QueryHandle QueryHandle : QueryHandles)
+			for (QueryHandle QueryHandle : QueryHandles)
 			{
 				TedsInterface->UnregisterQuery(QueryHandle);
 			}
@@ -612,6 +613,7 @@ namespace UE::Editor::DataStorage::Tests
 			TedsInterface = nullptr;
 		});
 	}
-}
+} // namespace Tests
+} // namespace UE::Editor::DataStorage
 
 #endif // WITH_TESTS
\ No newline at end of file
diff --git a/Engine/Plugins/Experimental/EditorDataStorage/Source/TedsCore/Tests/TypedElementAttributeBindingTests.cpp b/Engine/Plugins/Experimental/EditorDataStorage/Source/TedsCore/Tests/TypedElementAttributeBindingTests.cpp
--- a/Engine/Plugins/Experimental/EditorDataStorage/Source/TedsCore/Tests/TypedElementAttributeBindingTests.cpp
+++ b/Engine/Plugins/Experimental/EditorDataStorage/Source/TedsCore/Tests/TypedElementAttributeBindingTests.cpp
@@ -7,20 +7,22 @@
 #include "Elements/Framework/TypedElementRegistry.h"
 #include "Misc/AutomationTest.h"
 
-namespace TypedElementDataStorageTests
+namespace UE::Editor::DataStorage
 {
-	BEGIN_DEFINE_SPEC(TypedElementAttributeBindingTestsFixture, "EditorDataStorage.AttributeBinding", EAutomationTestFlags::EditorContext | EAutomationTestFlags::EngineFilter)
+namespace Tests
+{
+	BEGIN_DEFINE_SPEC(TedsAttributeBindingTestsFixture, "Editor.DataStorage.AttributeBinding", EAutomationTestFlags::EditorContext | EAutomationTestFlags::EngineFilter)
 
 		ITypedElementDataStorageInterface* TedsInterface = nullptr;
 		const FName TestTableName = TEXT("TestTable_AttributeBinding");
-		TypedElementDataStorage::TableHandle TestTableHandle = TypedElementDataStorage::InvalidTableHandle;
-		TypedElementDataStorage::RowHandle TestRowHandle = TypedElementDataStorage::InvalidRowHandle;
+		TableHandle TestTableHandle = InvalidTableHandle;
+		RowHandle TestRowHandle = InvalidRowHandle;
 
-		TypedElementDataStorage::TableHandle RegisterTestTable() const
+		TableHandle RegisterTestTable() const
 		{
-			const TypedElementDataStorage::TableHandle Table = TedsInterface->FindTable(TestTableName);
+			const TableHandle Table = TedsInterface->FindTable(TestTableName);
 			
-			if (Table != TypedElementDataStorage::InvalidTableHandle)
+			if (Table != InvalidTableHandle)
 			{
 				return Table;
 			}
@@ -33,20 +35,20 @@ namespace TypedElementDataStorageTests
 			TestTableName);
 		}
 	
-		TypedElementDataStorage::RowHandle CreateTestRow(TypedElementDataStorage::TableHandle InTableHandle) const
+		RowHandle CreateTestRow(TableHandle InTableHandle) const
 		{
-			const TypedElementDataStorage::RowHandle RowHandle = TedsInterface->AddRow(InTableHandle);
+			const RowHandle RowHandle = TedsInterface->AddRow(InTableHandle);
 			return RowHandle;
 		}
 	
-		void CleanupTestRow(TypedElementDataStorage::RowHandle InRowHandle) const
+		void CleanupTestRow(RowHandle InRowHandle) const
 		{
 			TedsInterface->RemoveRow(InRowHandle);
 		}
 	
-	END_DEFINE_SPEC(TypedElementAttributeBindingTestsFixture)
+	END_DEFINE_SPEC(TedsAttributeBindingTestsFixture)
 
-	void TypedElementAttributeBindingTestsFixture::Define()
+	void TedsAttributeBindingTestsFixture::Define()
 	{
 		BeforeEach([this]()
 		{
@@ -55,10 +57,10 @@ namespace TypedElementDataStorageTests
 			TestTrue("", TedsInterface != nullptr);
 			
 			TestTableHandle = RegisterTestTable();
-			TestNotEqual("Expecting valid table handle", TestTableHandle, TypedElementDataStorage::InvalidTableHandle);
+			TestNotEqual("Expecting valid table handle", TestTableHandle, InvalidTableHandle);
 
 			TestRowHandle = CreateTestRow(TestTableHandle);
-			TestFalse("Expect valid row handle", TestRowHandle == TypedElementDataStorage::InvalidRowHandle);
+			TestFalse("Expect valid row handle", TestRowHandle == InvalidRowHandle);
 		});
 
 		Describe("", [this]()
@@ -79,7 +81,7 @@ namespace TypedElementDataStorageTests
 						TestNotNull("Expecting Valid Column", TestColumnInt);
 
 						// Create an int attribute and bind it
-						UE::EditorDataStorage::FAttributeBinder Binder(TestRowHandle);
+						FAttributeBinder Binder(TestRowHandle);
 						const TAttribute TestAttribute(Binder.BindData(&FTestColumnInt::TestInt));
 
 						TestEqual("Expecting attribute value to match column value before modification", TestAttribute.Get(), TestColumnInt->TestInt);
@@ -106,7 +108,7 @@ namespace TypedElementDataStorageTests
 						TestNotNull("Expecting valid column", TestColumnInt);
 
 						// Create a float attribute and bind it by providing a conversion function
-						UE::EditorDataStorage::FAttributeBinder Binder(TestRowHandle);
+						FAttributeBinder Binder(TestRowHandle);
 						const TAttribute<float> TestAttribute(Binder.BindData(&FTestColumnInt::TestInt,
 							[](const int& Data)
 							{
@@ -140,7 +142,7 @@ namespace TypedElementDataStorageTests
 						TestNotNull("Expecting valid column", TestColumnString);
 
 						// Create an int attribute and bind it
-						UE::EditorDataStorage::FAttributeBinder Binder(TestRowHandle);
+						FAttributeBinder Binder(TestRowHandle);
 						const TAttribute TestAttribute(Binder.BindData(&FTestColumnString::TestString));
 
 						TestEqual("Expecting attribute value to match column value before modification", TestAttribute.Get(), TestColumnString->TestString);
@@ -167,7 +169,7 @@ namespace TypedElementDataStorageTests
 						TestNotNull("Expecting valid column", TestColumnString);
 
 						// Create an int attribute and bind it
-						UE::EditorDataStorage::FAttributeBinder Binder(TestRowHandle);
+						FAttributeBinder Binder(TestRowHandle);
 						const TAttribute<FText> TestAttribute(Binder.BindData(&FTestColumnString::TestString,
 							[](const FString& Data)
 							{
@@ -191,7 +193,7 @@ namespace TypedElementDataStorageTests
 				{
 					constexpr int DefaultValue = 10;
 					
-					UE::EditorDataStorage::FAttributeBinder Binder(TestRowHandle);
+					FAttributeBinder Binder(TestRowHandle);
 
 					// Create an int attribute and directly bind it
 					const TAttribute TestIntAttribute(Binder.BindData(&FTestColumnInt::TestInt, DefaultValue));
@@ -217,11 +219,12 @@ namespace TypedElementDataStorageTests
 		AfterEach([this]()
 		{
 			CleanupTestRow(TestRowHandle);
-			TestRowHandle = TypedElementDataStorage::InvalidRowHandle;
-			TestTableHandle = TypedElementDataStorage::InvalidTableHandle;
+			TestRowHandle = InvalidRowHandle;
+			TestTableHandle = InvalidTableHandle;
 			TedsInterface = nullptr;
 		});
 	}
-}
+} // namespace Tests
+} // namespace UE::Editor::DataStorage
 
 #endif // WITH_TESTS
\ No newline at end of file
diff --git a/Engine/Plugins/Experimental/EditorDataStorage/Source/TedsCore/Tests/TypedElementDataStorageInterfaceTests.cpp b/Engine/Plugins/Experimental/EditorDataStorage/Source/TedsCore/Tests/TypedElementDataStorageInterfaceTests.cpp
--- a/Engine/Plugins/Experimental/EditorDataStorage/Source/TedsCore/Tests/TypedElementDataStorageInterfaceTests.cpp
+++ b/Engine/Plugins/Experimental/EditorDataStorage/Source/TedsCore/Tests/TypedElementDataStorageInterfaceTests.cpp
@@ -31,7 +31,7 @@ namespace UE::Editor::DataStorage
 			TEXT("TestTableA"));
 		}
 
-		bool IsEmptyQueryDescription(const TypedElementDataStorage::FQueryDescription& QueryDescription)
+		bool IsEmptyQueryDescription(const FQueryDescription& QueryDescription)
 		{
 			return
 				QueryDescription.Callback.MonitoredType == nullptr &&
@@ -59,8 +59,7 @@ namespace UE::Editor::DataStorage
 			ITypedElementDataStorageInterface* TedsInterface, TypedElementQueryBuilder::FObserver::EEvent EventType, int64* QueryCallCountPtr, QueryHandle* QueryHandleInOut)
 		{
 			using namespace TypedElementQueryBuilder;
-			using DSI = ITypedElementDataStorageInterface;
-			TypedElementDataStorage::FQueryDescription QueryDescription = TedsInterface->GetQueryDescription(*QueryHandleInOut);
+			FQueryDescription QueryDescription = TedsInterface->GetQueryDescription(*QueryHandleInOut);
 
 			// Check if it is empty... If not then we need to register
 			// This is for an idempotent API to the test harness				
@@ -70,14 +69,14 @@ namespace UE::Editor::DataStorage
 					Select(
 						ObserverName,
 						FObserver(EventType, ColumnType::StaticStruct()),
-						[QueryCallCountPtr](DSI::IQueryContext& Context, RowHandle Row)
+						[QueryCallCountPtr](IQueryContext& Context, RowHandle Row)
 						{
 							++(*QueryCallCountPtr);
 						}).Compile());
 			}
 		}
 	
-		BEGIN_DEFINE_SPEC(TypedElementDataStorageTestsFixture, "TypedElementsDataStorage", EAutomationTestFlags::EditorContext | EAutomationTestFlags::EngineFilter)
+		BEGIN_DEFINE_SPEC(EditorDataStorageTestsFixture, "Editor.DataStorage.Core", EAutomationTestFlags::EditorContext | EAutomationTestFlags::EngineFilter)
 			ITypedElementDataStorageInterface* TedsInterface = nullptr;
 			TableHandle TestTableHandleA = InvalidTableHandle;
 
@@ -122,9 +121,9 @@ namespace UE::Editor::DataStorage
 				}
 				CreatedRows.Empty();
 			}
-		END_DEFINE_SPEC(TypedElementDataStorageTestsFixture)
+		END_DEFINE_SPEC(EditorDataStorageTestsFixture)
 
-		void TypedElementDataStorageTestsFixture::Define()
+		void EditorDataStorageTestsFixture::Define()
 		{
 			BeforeEach([this]()
 			{
diff --git a/Engine/Plugins/Experimental/EditorDataStorage/Source/TedsCore/Tests/TypedElementDatabaseScratchBufferTests.cpp b/Engine/Plugins/Experimental/EditorDataStorage/Source/TedsCore/Tests/TypedElementDatabaseScratchBufferTests.cpp
--- a/Engine/Plugins/Experimental/EditorDataStorage/Source/TedsCore/Tests/TypedElementDatabaseScratchBufferTests.cpp
+++ b/Engine/Plugins/Experimental/EditorDataStorage/Source/TedsCore/Tests/TypedElementDatabaseScratchBufferTests.cpp
@@ -24,7 +24,7 @@ static void WaitForFinalization(std::condition_variable& WaitVariable, std::mute
 	REQUIRE(CompletedThreadCount == ThreadCount);
 }
 
-TEST_CASE("TEDS::Scratch Buffer - MT (FScratchBuffer)", "[ApplicationContextMask][EngineFilter]")
+TEST_CASE("Editor::DataStorage::Scratch Buffer - MT (FScratchBuffer)", "[ApplicationContextMask][EngineFilter]")
 {
 	using namespace UE::Editor::DataStorage;
 	SECTION("Stress test without recycling.")
@@ -138,7 +138,7 @@ TEST_CASE("TEDS::Scratch Buffer - MT (FScratchBuffer)", "[ApplicationContextMask
 	}
 }
 
-TEST_CASE("TEDS::Scratch Buffer (FScratchBuffer)", "[ApplicationContextMask][EngineFilter]")
+TEST_CASE("Editor::DataStorage::Scratch Buffer (FScratchBuffer)", "[ApplicationContextMask][EngineFilter]")
 {
 	using namespace UE::Editor::DataStorage;
 	SECTION("Create and destroy buffer")
diff --git a/Engine/Plugins/Experimental/EditorDataStorage/Source/TedsUI/Private/Widgets/CounterWidget.cpp b/Engine/Plugins/Experimental/EditorDataStorage/Source/TedsUI/Private/Widgets/CounterWidget.cpp
--- a/Engine/Plugins/Experimental/EditorDataStorage/Source/TedsUI/Private/Widgets/CounterWidget.cpp
+++ b/Engine/Plugins/Experimental/EditorDataStorage/Source/TedsUI/Private/Widgets/CounterWidget.cpp
@@ -165,7 +165,7 @@ TConstArrayView<const UScriptStruct*> FCounterWidgetConstructor::GetAdditionalCo
 	return Columns;
 }
 
-TSharedPtr<SWidget> FCounterWidgetConstructor::CreateWidget(const TypedElementDataStorage::FMetaDataView& Arguments)
+TSharedPtr<SWidget> FCounterWidgetConstructor::CreateWidget(const UE::Editor::DataStorage::FMetaDataView& Arguments)
 {
 	return SNew(STextBlock)
 		.Text(FText::Format(LabelText, 0))
diff --git a/Engine/Plugins/Experimental/EditorDataStorage/Source/TedsUI/Private/Widgets/ExportedTextWidget.cpp b/Engine/Plugins/Experimental/EditorDataStorage/Source/TedsUI/Private/Widgets/ExportedTextWidget.cpp
--- a/Engine/Plugins/Experimental/EditorDataStorage/Source/TedsUI/Private/Widgets/ExportedTextWidget.cpp
+++ b/Engine/Plugins/Experimental/EditorDataStorage/Source/TedsUI/Private/Widgets/ExportedTextWidget.cpp
@@ -56,6 +56,7 @@ namespace UE::Editor
 		{
 			using namespace TypedElementQueryBuilder;
 			using namespace TypedElementDataStorage;
+			using namespace UE::Editor::DataStorage;
 			
 			QueryHandle TypeDataQuery = DataStorage.RegisterQuery(
 				Select()
@@ -158,7 +159,7 @@ FString FExportedTextWidgetConstructor::CreateWidgetDisplayName(ITypedElementDat
 	}
 }
 
-TSharedPtr<SWidget> FExportedTextWidgetConstructor::CreateWidget(const TypedElementDataStorage::FMetaDataView& Arguments)
+TSharedPtr<SWidget> FExportedTextWidgetConstructor::CreateWidget(const UE::Editor::DataStorage::FMetaDataView& Arguments)
 {
 	return SNew(STextBlock);
 }
diff --git a/Engine/Plugins/Experimental/EditorDataStorage/Source/TedsUI/Private/Widgets/LabelWidget.cpp b/Engine/Plugins/Experimental/EditorDataStorage/Source/TedsUI/Private/Widgets/LabelWidget.cpp
--- a/Engine/Plugins/Experimental/EditorDataStorage/Source/TedsUI/Private/Widgets/LabelWidget.cpp
+++ b/Engine/Plugins/Experimental/EditorDataStorage/Source/TedsUI/Private/Widgets/LabelWidget.cpp
@@ -64,9 +64,9 @@ TConstArrayView<const UScriptStruct*> FLabelWidgetConstructor::GetAdditionalColu
 
 TSharedPtr<SWidget> FLabelWidgetConstructor::CreateWidget(ITypedElementDataStorageInterface* DataStorage,
 	ITypedElementDataStorageUiInterface* DataStorageUi, RowHandle TargetRow, RowHandle WidgetRow,
-	const TypedElementDataStorage::FMetaDataView& Arguments)
+	const UE::Editor::DataStorage::FMetaDataView& Arguments)
 {
-	UE::EditorDataStorage::FAttributeBinder Binder(TargetRow, DataStorage);
+	UE::Editor::DataStorage::FAttributeBinder Binder(TargetRow, DataStorage);
 
 	return SNew(STextBlock)
 		.Text(Binder.BindText(&FTypedElementLabelColumn::Label))
diff --git a/Engine/Plugins/Experimental/EditorDataStorage/Source/TedsUI/Private/Widgets/PackagePathWidget.cpp b/Engine/Plugins/Experimental/EditorDataStorage/Source/TedsUI/Private/Widgets/PackagePathWidget.cpp
--- a/Engine/Plugins/Experimental/EditorDataStorage/Source/TedsUI/Private/Widgets/PackagePathWidget.cpp
+++ b/Engine/Plugins/Experimental/EditorDataStorage/Source/TedsUI/Private/Widgets/PackagePathWidget.cpp
@@ -36,7 +36,7 @@ FPackagePathWidgetConstructor::FPackagePathWidgetConstructor(const UScriptStruct
 {
 }
 
-TSharedPtr<SWidget> FPackagePathWidgetConstructor::CreateWidget(const TypedElementDataStorage::FMetaDataView& Arguments)
+TSharedPtr<SWidget> FPackagePathWidgetConstructor::CreateWidget(const UE::Editor::DataStorage::FMetaDataView& Arguments)
 {
 	return SNew(STextBlock)
 		.OverflowPolicy(ETextOverflowPolicy::Ellipsis)
diff --git a/Engine/Plugins/Experimental/EditorDataStorage/Source/TedsUI/Public/Widgets/CounterWidget.h b/Engine/Plugins/Experimental/EditorDataStorage/Source/TedsUI/Public/Widgets/CounterWidget.h
--- a/Engine/Plugins/Experimental/EditorDataStorage/Source/TedsUI/Public/Widgets/CounterWidget.h
+++ b/Engine/Plugins/Experimental/EditorDataStorage/Source/TedsUI/Public/Widgets/CounterWidget.h
@@ -60,7 +60,7 @@ public:
 	UE::Editor::DataStorage::QueryHandle Query;
 
 protected:
-	TSharedPtr<SWidget> CreateWidget(const TypedElementDataStorage::FMetaDataView& Arguments) override;
+	TSharedPtr<SWidget> CreateWidget(const UE::Editor::DataStorage::FMetaDataView& Arguments) override;
 	bool SetColumns(ITypedElementDataStorageInterface* DataStorage, UE::Editor::DataStorage::RowHandle Row) override;
 };
 
diff --git a/Engine/Plugins/Experimental/EditorDataStorage/Source/TedsUI/Public/Widgets/ExportedTextWidget.h b/Engine/Plugins/Experimental/EditorDataStorage/Source/TedsUI/Public/Widgets/ExportedTextWidget.h
--- a/Engine/Plugins/Experimental/EditorDataStorage/Source/TedsUI/Public/Widgets/ExportedTextWidget.h
+++ b/Engine/Plugins/Experimental/EditorDataStorage/Source/TedsUI/Public/Widgets/ExportedTextWidget.h
@@ -35,12 +35,12 @@ public:
 	virtual ~FExportedTextWidgetConstructor() override = default;
 
 	virtual TConstArrayView<const UScriptStruct*> GetAdditionalColumnsList() const override;
-	virtual const TypedElementDataStorage::FQueryConditions* GetQueryConditions() const override;
+	virtual const UE::Editor::DataStorage::FQueryConditions* GetQueryConditions() const override;
 	virtual FString CreateWidgetDisplayName(
 		ITypedElementDataStorageInterface* DataStorage, UE::Editor::DataStorage::RowHandle Row) const override;
 
 protected:
-	virtual TSharedPtr<SWidget> CreateWidget(const TypedElementDataStorage::FMetaDataView& Arguments) override;
+	virtual TSharedPtr<SWidget> CreateWidget(const UE::Editor::DataStorage::FMetaDataView& Arguments) override;
 	virtual bool FinalizeWidget(
 		ITypedElementDataStorageInterface* DataStorage,
 		ITypedElementDataStorageUiInterface* DataStorageUi,
@@ -49,7 +49,7 @@ protected:
 
 protected:
 	// The column this exported text widget is operating on
-	TypedElementDataStorage::FQueryConditions MatchedColumn;
+	UE::Editor::DataStorage::FQueryConditions MatchedColumn;
 };
 
 USTRUCT(meta = (DisplayName = "Exported text widget"))
diff --git a/Engine/Plugins/Experimental/EditorDataStorage/Source/TedsUI/Public/Widgets/LabelWidget.h b/Engine/Plugins/Experimental/EditorDataStorage/Source/TedsUI/Public/Widgets/LabelWidget.h
--- a/Engine/Plugins/Experimental/EditorDataStorage/Source/TedsUI/Public/Widgets/LabelWidget.h
+++ b/Engine/Plugins/Experimental/EditorDataStorage/Source/TedsUI/Public/Widgets/LabelWidget.h
@@ -39,9 +39,9 @@ public:
 	TEDSUI_API virtual TSharedPtr<SWidget> CreateWidget(
 		ITypedElementDataStorageInterface* DataStorage,
 		ITypedElementDataStorageUiInterface* DataStorageUi,
-		RowHandle TargetRow,
-		RowHandle WidgetRow, 
-		const TypedElementDataStorage::FMetaDataView& Arguments) override;
+		UE::Editor::DataStorage::RowHandle TargetRow,
+		UE::Editor::DataStorage::RowHandle WidgetRow, 
+		const UE::Editor::DataStorage::FMetaDataView& Arguments) override;
 
 protected:
 };
\ No newline at end of file
diff --git a/Engine/Plugins/Experimental/EditorDataStorage/Source/TedsUI/Public/Widgets/PackagePathWidget.h b/Engine/Plugins/Experimental/EditorDataStorage/Source/TedsUI/Public/Widgets/PackagePathWidget.h
--- a/Engine/Plugins/Experimental/EditorDataStorage/Source/TedsUI/Public/Widgets/PackagePathWidget.h
+++ b/Engine/Plugins/Experimental/EditorDataStorage/Source/TedsUI/Public/Widgets/PackagePathWidget.h
@@ -35,7 +35,7 @@ public:
 protected:
 	explicit FPackagePathWidgetConstructor(const UScriptStruct* InTypeInfo);
 
-	TSharedPtr<SWidget> CreateWidget(const TypedElementDataStorage::FMetaDataView& Arguments) override;
+	TSharedPtr<SWidget> CreateWidget(const UE::Editor::DataStorage::FMetaDataView& Arguments) override;
 	bool FinalizeWidget(ITypedElementDataStorageInterface* DataStorage, ITypedElementDataStorageUiInterface* DataStorageUi,
 		UE::Editor::DataStorage::RowHandle Row, const TSharedPtr<SWidget>& Widget) override;
 };
diff --git a/Engine/Plugins/Experimental/EditorDataStorageFeatures/Source/TedsAssetData/Private/CB/TedsAssetDataCBDataSource.cpp b/Engine/Plugins/Experimental/EditorDataStorageFeatures/Source/TedsAssetData/Private/CB/TedsAssetDataCBDataSource.cpp
--- a/Engine/Plugins/Experimental/EditorDataStorageFeatures/Source/TedsAssetData/Private/CB/TedsAssetDataCBDataSource.cpp
+++ b/Engine/Plugins/Experimental/EditorDataStorageFeatures/Source/TedsAssetData/Private/CB/TedsAssetDataCBDataSource.cpp
@@ -42,7 +42,7 @@ FTedsAssetDataCBDataSource::FTedsAssetDataCBDataSource(ITypedElementDataStorageI
 		Select(
 			TEXT("FTedsAssetDataCBDataSource: Process Path updates"),
 			FProcessor(TypedElementDataStorage::EQueryTickPhase::DuringPhysics, Database.GetQueryTickGroupName(TypedElementDataStorage::EQueryTickGroups::Update)),
-			[GenerateVirtualPaths](TypedElementDataStorage::IQueryContext& Context, const RowHandle* Rows, const FAssetPathColumn_Experimental* PathColumn)
+			[GenerateVirtualPaths](IQueryContext& Context, const RowHandle* Rows, const FAssetPathColumn_Experimental* PathColumn)
 			{
 				int32 NumOfRowToProcess = Context.GetRowCount();
 
@@ -72,7 +72,7 @@ FTedsAssetDataCBDataSource::FTedsAssetDataCBDataSource(ITypedElementDataStorageI
 		Select(
 			TEXT("FTedsAssetDataCBDataSource:: Process Asset Data Path Update"),
 			FProcessor(TypedElementDataStorage::EQueryTickPhase::DuringPhysics, Database.GetQueryTickGroupName(TypedElementDataStorage::EQueryTickGroups::Update)),
-			[GenerateVirtualPaths](TypedElementDataStorage::IQueryContext& Context, const RowHandle Row, const FAssetDataColumn_Experimental& AssetDataColumn)
+			[GenerateVirtualPaths](IQueryContext& Context, const RowHandle Row, const FAssetDataColumn_Experimental& AssetDataColumn)
 			{
 				int32 NumOfRowToProcess = Context.GetRowCount();
 
@@ -98,7 +98,7 @@ FTedsAssetDataCBDataSource::FTedsAssetDataCBDataSource(ITypedElementDataStorageI
 
 
 	// For now just add the columns one by one but this should be rework to work in batch
-	auto AddAssetDataColumns = [](TypedElementDataStorage::IQueryContext& InContext, RowHandle Row, const FAssetData& InAssetData, const FAssetPackageData* PackageData)
+	auto AddAssetDataColumns = [](IQueryContext& InContext, RowHandle Row, const FAssetData& InAssetData, const FAssetPackageData* PackageData)
 	{
 		// Not optimized at all but we would like to have the data in sooner for testing purposes.
 		
@@ -133,7 +133,7 @@ FTedsAssetDataCBDataSource::FTedsAssetDataCBDataSource(ITypedElementDataStorageI
 		Select(
 			TEXT("FTedsAssetDataCBDataSource: Process Asset Data and Path Updates"),
 			FProcessor(TypedElementDataStorage::EQueryTickPhase::DuringPhysics, Database.GetQueryTickGroupName(TypedElementDataStorage::EQueryTickGroups::Update)),
-			[GenerateVirtualPaths, AddAssetDataColumns, AssetRegistry = static_cast<const IAssetRegistry*>(&IAssetRegistry::GetChecked())](TypedElementDataStorage::IQueryContext& Context, const RowHandle* Rows, const FAssetDataColumn_Experimental* AssetDataColumn)
+			[GenerateVirtualPaths, AddAssetDataColumns, AssetRegistry = static_cast<const IAssetRegistry*>(&IAssetRegistry::GetChecked())](UE::Editor::DataStorage::IQueryContext& Context, const RowHandle* Rows, const FAssetDataColumn_Experimental* AssetDataColumn)
 			{
 				const int32 RowCount = Context.GetRowCount();
 
@@ -183,7 +183,7 @@ FTedsAssetDataCBDataSource::FTedsAssetDataCBDataSource(ITypedElementDataStorageI
 		Select(
 			TEXT("FTedsAssetDataCBDataSource: Process Asset Data updates"),
 			FProcessor(TypedElementDataStorage::EQueryTickPhase::DuringPhysics, Database.GetQueryTickGroupName(TypedElementDataStorage::EQueryTickGroups::Update)),
-			[AddAssetDataColumns,  AssetRegistry = static_cast<const IAssetRegistry*>(&IAssetRegistry::GetChecked())](TypedElementDataStorage::IQueryContext& Context, const RowHandle* Rows, const FAssetDataColumn_Experimental* AssetDataColumn)
+			[AddAssetDataColumns,  AssetRegistry = static_cast<const IAssetRegistry*>(&IAssetRegistry::GetChecked())](UE::Editor::DataStorage::IQueryContext& Context, const RowHandle* Rows, const FAssetDataColumn_Experimental* AssetDataColumn)
 			{
 				const int32 RowCount = Context.GetRowCount();
 				TArray<FName, TInlineAllocator<32>> PackageNames;
diff --git a/Engine/Plugins/Experimental/EditorDataStorageFeatures/Source/TedsAssetData/Private/Debug/TedsAssetDataDebugCommands.cpp b/Engine/Plugins/Experimental/EditorDataStorageFeatures/Source/TedsAssetData/Private/Debug/TedsAssetDataDebugCommands.cpp
--- a/Engine/Plugins/Experimental/EditorDataStorageFeatures/Source/TedsAssetData/Private/Debug/TedsAssetDataDebugCommands.cpp
+++ b/Engine/Plugins/Experimental/EditorDataStorageFeatures/Source/TedsAssetData/Private/Debug/TedsAssetDataDebugCommands.cpp
@@ -23,7 +23,9 @@
 
 DEFINE_LOG_CATEGORY_STATIC(LogTEDSAssetRegistry, Log, All)
 
-namespace UE::EditorDataStorage::Debug::Private
+namespace UE::Editor::DataStorage
+{
+namespace Debug::Private
 {
 	static FAutoConsoleCommand CCMDTestFolderRowData(
 		TEXT("TEDS.Debug.ShowDataOfAssetFolder"),
@@ -36,9 +38,9 @@ namespace UE::EditorDataStorage::Debug::Private
 				for (const FString& Path : InArgs)
 				{
 					FName PathAsName(*Path);
-					TypedElementDataStorage::RowHandle RowHandle = Database->FindIndexedRow(TypedElementDataStorage::GenerateIndexHash(PathAsName));
+					RowHandle Row = Database->FindIndexedRow(TypedElementDataStorage::GenerateIndexHash(PathAsName));
 
-					if (Database->IsRowAssigned(RowHandle))
+					if (Database->IsRowAssigned(Row))
 					{
 						UE_LOG(LogTEDSAssetRegistry, Display, TEXT("The path isn't indexed."));
 						return;
@@ -47,12 +49,12 @@ namespace UE::EditorDataStorage::Debug::Private
 
 					UE_LOG(LogTEDSAssetRegistry, Display, TEXT("Found some information for the path (%s) in the database."), *Path);
 
-					if (const FAssetPathColumn_Experimental* AssetPath = Database->GetColumn<FAssetPathColumn_Experimental>(RowHandle))
+					if (const FAssetPathColumn_Experimental* AssetPath = Database->GetColumn<FAssetPathColumn_Experimental>(Row))
 					{
 						UE_LOG(LogTEDSAssetRegistry, Display, TEXT("Path stored in the database as (%s)."), *AssetPath->Path.ToString());
 					}
 
-					if (const FParentAssetPathColumn_Experimental* ParentAssetPath = Database->GetColumn<FParentAssetPathColumn_Experimental>(RowHandle))
+					if (const FParentAssetPathColumn_Experimental* ParentAssetPath = Database->GetColumn<FParentAssetPathColumn_Experimental>(Row))
 					{
 						if (const FAssetPathColumn_Experimental* AssetPath = Database->GetColumn<FAssetPathColumn_Experimental>(ParentAssetPath->ParentRow))
 						{
@@ -61,11 +63,11 @@ namespace UE::EditorDataStorage::Debug::Private
 					}
 
 
-					if (const FChildrenAssetPathColumn_Experimental* ChildrenPath = Database->GetColumn<FChildrenAssetPathColumn_Experimental>(RowHandle))
+					if (const FChildrenAssetPathColumn_Experimental* ChildrenPath = Database->GetColumn<FChildrenAssetPathColumn_Experimental>(Row))
 					{
 						UE_LOG(LogTEDSAssetRegistry, Display, TEXT("	Path as %i children"), ChildrenPath->ChildrenRows.Num());
 
-						for (TypedElementDataStorage::RowHandle Row : ChildrenPath->ChildrenRows)
+						for (RowHandle ChildRow : ChildrenPath->ChildrenRows)
 						{
 							if (const FAssetPathColumn_Experimental* AssetPath = Database->GetColumn<FAssetPathColumn_Experimental>(Row))
 							{
@@ -75,11 +77,11 @@ namespace UE::EditorDataStorage::Debug::Private
 					}
 
 
-					if (const FAssetsInPathColumn_Experimental* AssetInPath = Database->GetColumn<FAssetsInPathColumn_Experimental>(RowHandle))
+					if (const FAssetsInPathColumn_Experimental* AssetInPath = Database->GetColumn<FAssetsInPathColumn_Experimental>(Row))
 					{
 						UE_LOG(LogTEDSAssetRegistry, Display, TEXT("	Asset in Paths"));
 
-						for (TypedElementDataStorage::RowHandle AssetRow : AssetInPath->AssetsRow)
+						for (RowHandle AssetRow : AssetInPath->AssetsRow)
 						{
 							if (const FAssetDataColumn_Experimental* AssetData = Database->GetColumn<FAssetDataColumn_Experimental>(AssetRow))
 							{
@@ -98,7 +100,7 @@ namespace UE::EditorDataStorage::Debug::Private
 					AssetRegistry.GetAssetsByPath(PathAsName, Assets);
 					for (const FAssetData& Asset : Assets)
 					{
-						TypedElementDataStorage::RowHandle AssetRow = Database->FindIndexedRow(TypedElementDataStorage::GenerateIndexHash(Asset.GetSoftObjectPath()));
+						RowHandle AssetRow = Database->FindIndexedRow(TypedElementDataStorage::GenerateIndexHash(Asset.GetSoftObjectPath()));
 
 						if (const FUnresolvedAssetsInPathColumn_Experimental* UnresolvedAssetsInPathColumn = Database->GetColumn<FUnresolvedAssetsInPathColumn_Experimental>(AssetRow))
 						{
@@ -160,7 +162,7 @@ namespace UE::EditorDataStorage::Debug::Private
 		ULevel::StaticClass(), UBlueprint::StaticClass(), UTexture::StaticClass() };
 
 	// Populate an asset row with random information
-	void PopulateRowWithRandomInfo(TypedElementDataStorage::RowHandle Row, ITypedElementDataStorageInterface* DataStorage)
+	void PopulateRowWithRandomInfo(RowHandle Row, ITypedElementDataStorageInterface* DataStorage)
 	{
 		// Don't modify any rows that aren't our placeholder assets
 		if (!DataStorage->HasColumns<FAssetTag>(Row))
@@ -279,4 +281,5 @@ namespace UE::EditorDataStorage::Debug::Private
 				}
 			}
 		));
-}
+} // namespace Debug::Private
+} // namespace UE::Editor::DataStorage
diff --git a/Engine/Plugins/Experimental/EditorDataStorageFeatures/Source/TedsAssetData/Private/TedsAssetData.cpp b/Engine/Plugins/Experimental/EditorDataStorageFeatures/Source/TedsAssetData/Private/TedsAssetData.cpp
--- a/Engine/Plugins/Experimental/EditorDataStorageFeatures/Source/TedsAssetData/Private/TedsAssetData.cpp
+++ b/Engine/Plugins/Experimental/EditorDataStorageFeatures/Source/TedsAssetData/Private/TedsAssetData.cpp
@@ -179,13 +179,13 @@ FTedsAssetData::FTedsAssetData(ITypedElementDataStorageInterface& InDatabase)
 
 	// Register data types to TEDS
 	PathsTable = Database.FindTable(FName(TEXT("Editor_AssetRegistryPathsTable")));
-	if (PathsTable == TypedElementDataStorage::InvalidTableHandle)
+	if (PathsTable == DataStorage::InvalidTableHandle)
 	{
 		PathsTable = Database.RegisterTable<FAssetPathColumn_Experimental, FChildrenAssetPathColumn_Experimental, FParentAssetPathColumn_Experimental, FAssetsInPathColumn_Experimental>(FName(TEXT("Editor_AssetRegistryPathsTable")));
 	}
 
 	AssetsDataTable = Database.FindTable(FName(TEXT("Editor_AssetRegistryAssetDataTable")));
-	if (AssetsDataTable == TypedElementDataStorage::InvalidTableHandle)
+	if (AssetsDataTable == DataStorage::InvalidTableHandle)
 	{
 		AssetsDataTable = Database.RegisterTable<FAssetDataColumn_Experimental, FUpdatedPathTag, FUpdatedAssetDataTag>(FName(TEXT("Editor_AssetRegistryAssetDataTable")));
 	}
@@ -195,7 +195,7 @@ FTedsAssetData::FTedsAssetData(ITypedElementDataStorageInterface& InDatabase)
 			TypedElementQueryBuilder::Select(
 				TEXT("FTedsAssetData: Remove Updated Path Tag"),
 				TypedElementQueryBuilder::FPhaseAmble(TypedElementQueryBuilder::FPhaseAmble::ELocation::Postamble, TypedElementDataStorage::EQueryTickPhase::FrameEnd),
-				[](TypedElementDataStorage::IQueryContext& Context, const DataStorage::RowHandle* Rows)
+				[](DataStorage::IQueryContext& Context, const DataStorage::RowHandle* Rows)
 				{
 					Context.RemoveColumns<FUpdatedPathTag>(TConstArrayView<DataStorage::RowHandle>(Rows, Context.GetRowCount()));
 				}
@@ -209,7 +209,7 @@ FTedsAssetData::FTedsAssetData(ITypedElementDataStorageInterface& InDatabase)
 			TypedElementQueryBuilder::Select(
 				TEXT("FTedsAssetData: Remove Updated Asset Data Tag"),
 				TypedElementQueryBuilder::FPhaseAmble(TypedElementQueryBuilder::FPhaseAmble::ELocation::Postamble, TypedElementDataStorage::EQueryTickPhase::FrameEnd),
-				[](TypedElementDataStorage::IQueryContext& Context, const DataStorage::RowHandle* Rows)
+				[](DataStorage::IQueryContext& Context, const DataStorage::RowHandle* Rows)
 				{
 					Context.RemoveColumns<FUpdatedAssetDataTag>(TConstArrayView<DataStorage::RowHandle>(Rows, Context.GetRowCount()));
 				}
@@ -229,7 +229,7 @@ FTedsAssetData::FTedsAssetData(ITypedElementDataStorageInterface& InDatabase)
 		TypedElementQueryBuilder::Select(
 			TEXT("FTedsAssetData: Resolve Missing Asset In Path"),
 			TypedElementQueryBuilder::FProcessor(TypedElementDataStorage::EQueryTickPhase::FrameEnd, Database.GetQueryTickGroupName(TypedElementDataStorage::EQueryTickGroups::Default)),
-			[this](TypedElementDataStorage::IQueryContext& Context, DataStorage::RowHandle Row, const FUnresolvedAssetsInPathColumn_Experimental& UnresolvedAssetPath)
+			[this](DataStorage::IQueryContext& Context, DataStorage::RowHandle Row, const FUnresolvedAssetsInPathColumn_Experimental& UnresolvedAssetPath)
 			{
 #if TRACK_TEDSASSETDATA_MEMORY
 				LLM_SCOPE_BYNAME(TEXT("FTedsAssetData"))
@@ -267,7 +267,7 @@ FTedsAssetData::FTedsAssetData(ITypedElementDataStorageInterface& InDatabase)
 		TypedElementQueryBuilder::Select(
 			TEXT("FTedsAssetData: Resolve Missing Parent Path Row"),
 			TypedElementQueryBuilder::FProcessor(TypedElementDataStorage::EQueryTickPhase::FrameEnd, Database.GetQueryTickGroupName(TypedElementDataStorage::EQueryTickGroups::Default)),
-			[this](TypedElementDataStorage::IQueryContext& Context, DataStorage::RowHandle Row, const FUnresolvedParentAssetPathColumn_Experimental& UnresolvedParentAssetPath, FParentAssetPathColumn_Experimental& ParentAssetPathColumn)
+			[this](DataStorage::IQueryContext& Context, DataStorage::RowHandle Row, const FUnresolvedParentAssetPathColumn_Experimental& UnresolvedParentAssetPath, FParentAssetPathColumn_Experimental& ParentAssetPathColumn)
 			{
 #if TRACK_TEDSASSETDATA_MEMORY
 				LLM_SCOPE_BYNAME(TEXT("FTedsAssetData"))
diff --git a/Engine/Plugins/Experimental/EditorDataStorageFeatures/Source/TedsAssetData/Private/Tests/TedsAssetDataTest.cpp b/Engine/Plugins/Experimental/EditorDataStorageFeatures/Source/TedsAssetData/Private/Tests/TedsAssetDataTest.cpp
--- a/Engine/Plugins/Experimental/EditorDataStorageFeatures/Source/TedsAssetData/Private/Tests/TedsAssetDataTest.cpp
+++ b/Engine/Plugins/Experimental/EditorDataStorageFeatures/Source/TedsAssetData/Private/Tests/TedsAssetDataTest.cpp
@@ -16,11 +16,16 @@
 #include "TedsAssetDataModule.h"
 #include "UObject/NameTypes.h"
 
-IMPLEMENT_SIMPLE_AUTOMATION_TEST(FTedsAssetDataTest, "TEDS.AssetRegistry.ValidateState", EAutomationTestFlags::EditorContext | EAutomationTestFlags::EngineFilter | EAutomationTestFlags::HighPriority)
- 
+namespace UE::Editor
+{
+namespace AssetData::Tests
+{
+
+IMPLEMENT_SIMPLE_AUTOMATION_TEST(FTedsAssetDataTest, "Editor.DataStorage.AssetRegistry.ValidateState", EAutomationTestFlags::EditorContext | EAutomationTestFlags::EngineFilter | EAutomationTestFlags::HighPriority)
+
 bool FTedsAssetDataTest::RunTest(const FString& Parameters)
 {
-	UE::Editor::AssetData::FTedsAssetDataModule& TedsAssetRegistry = UE::Editor::AssetData::FTedsAssetDataModule::GetChecked();
+	FTedsAssetDataModule& TedsAssetRegistry = FTedsAssetDataModule::GetChecked();
 
 	const bool bIsEnabled = TedsAssetRegistry.IsTedsAssetRegistryStorageEnabled();
 	ON_SCOPE_EXIT
@@ -66,9 +71,9 @@ bool FTedsAssetDataTest::RunTest(const FString& Parameters)
 	
 	for (const TPair<FName, int32>& PathAndAssetCount : AssetRegistryPathsAndAssetCount)
 	{
-		TypedElementDataStorage::RowHandle Row = Database->FindIndexedRow(TypedElementDataStorage::GenerateIndexHash(PathAndAssetCount.Key));
+		DataStorage::RowHandle Row = Database->FindIndexedRow(TypedElementDataStorage::GenerateIndexHash(PathAndAssetCount.Key));
 		bool bHadError = false;
-		bHadError |= !TestTrue(TEXT("Asset registry folder/path is indexed in TEDS"), Row != TypedElementDataStorage::InvalidRowHandle);
+		bHadError |= !TestTrue(TEXT("Asset registry folder/path is indexed in TEDS"), Row != DataStorage::InvalidRowHandle);
 
 		FName NameStoredInTeds;
 		if (const FAssetPathColumn_Experimental* AssetPathColumn = Database->GetColumn<FAssetPathColumn_Experimental>(Row))
@@ -104,17 +109,17 @@ bool FTedsAssetDataTest::RunTest(const FString& Parameters)
 		// Verse package don't emit asset added event so we will ignore those for the test to avoid false errors
 		if (!FPackageName::IsVersePackage(Builder))
 		{ 
-			TypedElementDataStorage::RowHandle Row = Database->FindIndexedRow(TypedElementDataStorage::GenerateIndexHash(AssetData.GetSoftObjectPath()));
+			DataStorage::RowHandle Row = Database->FindIndexedRow(TypedElementDataStorage::GenerateIndexHash(AssetData.GetSoftObjectPath()));
 
 			bool bHadError = false;
 
-			bHadError |= !TestTrue(TEXT("Asset registry asset path is indexed in TEDS"), Row != TypedElementDataStorage::InvalidRowHandle);
+			bHadError |= !TestTrue(TEXT("Asset registry asset path is indexed in TEDS"), Row != DataStorage::InvalidRowHandle);
 
 			bHadError |= !TestNotNull(TEXT("Teds doesn't have a asset column for an asset of the asset registry"), Database->GetColumn<FAssetDataColumn_Experimental>(Row));
 
 			if (!Database->GetColumn<FUnresolvedAssetsInPathColumn_Experimental>(Row))
 			{
-				TypedElementDataStorage::RowHandle PathRow = Database->FindIndexedRow(TypedElementDataStorage::GenerateIndexHash(AssetData.PackagePath));
+				DataStorage::RowHandle PathRow = Database->FindIndexedRow(TypedElementDataStorage::GenerateIndexHash(AssetData.PackagePath));
 			
 				const FAssetsInPathColumn_Experimental* AssetsInPathColumn = Database->GetColumn<FAssetsInPathColumn_Experimental>(PathRow);
 
@@ -147,5 +152,5 @@ bool FTedsAssetDataTest::RunTest(const FString& Parameters)
 	return true;
 }
 
-
-
+} // namespace AssetData::Tests
+} // namespace UE::Editor
diff --git a/Engine/Plugins/Experimental/EditorDataStorageFeatures/Source/TedsAssetData/Private/Widgets/AssetDataDiskSizeWidget.cpp b/Engine/Plugins/Experimental/EditorDataStorageFeatures/Source/TedsAssetData/Private/Widgets/AssetDataDiskSizeWidget.cpp
--- a/Engine/Plugins/Experimental/EditorDataStorageFeatures/Source/TedsAssetData/Private/Widgets/AssetDataDiskSizeWidget.cpp
+++ b/Engine/Plugins/Experimental/EditorDataStorageFeatures/Source/TedsAssetData/Private/Widgets/AssetDataDiskSizeWidget.cpp
@@ -20,10 +20,10 @@ FDiskSizeWidgetConstructor::FDiskSizeWidgetConstructor()
 }
 
 TSharedPtr<SWidget> FDiskSizeWidgetConstructor::CreateWidget(ITypedElementDataStorageInterface* DataStorage,
-	ITypedElementDataStorageUiInterface* DataStorageUi, RowHandle TargetRow, RowHandle WidgetRow,
-	const TypedElementDataStorage::FMetaDataView& Arguments)
+	ITypedElementDataStorageUiInterface* DataStorageUi, UE::Editor::DataStorage::RowHandle TargetRow, UE::Editor::DataStorage::RowHandle WidgetRow,
+	const UE::Editor::DataStorage::FMetaDataView& Arguments)
 {
-	UE::EditorDataStorage::FAttributeBinder Binder(TargetRow, DataStorage);
+	UE::Editor::DataStorage::FAttributeBinder Binder(TargetRow, DataStorage);
 	
 	return SNew(STextBlock)
 			.Text(Binder.BindData(&FDiskSizeColumn::DiskSize, [](int64 DiskSize)
diff --git a/Engine/Plugins/Experimental/EditorDataStorageFeatures/Source/TedsAssetData/Private/Widgets/AssetDataDiskSizeWidget.h b/Engine/Plugins/Experimental/EditorDataStorageFeatures/Source/TedsAssetData/Private/Widgets/AssetDataDiskSizeWidget.h
--- a/Engine/Plugins/Experimental/EditorDataStorageFeatures/Source/TedsAssetData/Private/Widgets/AssetDataDiskSizeWidget.h
+++ b/Engine/Plugins/Experimental/EditorDataStorageFeatures/Source/TedsAssetData/Private/Widgets/AssetDataDiskSizeWidget.h
@@ -35,7 +35,7 @@ public:
 	TEDSASSETDATA_API virtual TSharedPtr<SWidget> CreateWidget(
 		ITypedElementDataStorageInterface* DataStorage,
 		ITypedElementDataStorageUiInterface* DataStorageUi,
-		RowHandle TargetRow,
-		RowHandle WidgetRow, 
-		const TypedElementDataStorage::FMetaDataView& Arguments) override;
+		UE::Editor::DataStorage::RowHandle TargetRow,
+		UE::Editor::DataStorage::RowHandle WidgetRow,
+		const UE::Editor::DataStorage::FMetaDataView& Arguments) override;
 };
\ No newline at end of file
diff --git a/Engine/Plugins/Experimental/EditorDataStorageFeatures/Source/TedsAssetData/Private/Widgets/AssetDataLabelWidget.cpp b/Engine/Plugins/Experimental/EditorDataStorageFeatures/Source/TedsAssetData/Private/Widgets/AssetDataLabelWidget.cpp
--- a/Engine/Plugins/Experimental/EditorDataStorageFeatures/Source/TedsAssetData/Private/Widgets/AssetDataLabelWidget.cpp
+++ b/Engine/Plugins/Experimental/EditorDataStorageFeatures/Source/TedsAssetData/Private/Widgets/AssetDataLabelWidget.cpp
@@ -26,12 +26,12 @@ FAssetDataLabelWidgetConstructor::FAssetDataLabelWidgetConstructor()
 
 
 TSharedPtr<SWidget> FAssetDataLabelWidgetConstructor::CreateWidget(ITypedElementDataStorageInterface* DataStorage,
-	ITypedElementDataStorageUiInterface* DataStorageUi, RowHandle TargetRow, RowHandle WidgetRow,
-	const TypedElementDataStorage::FMetaDataView& Arguments)
+	ITypedElementDataStorageUiInterface* DataStorageUi, UE::Editor::DataStorage::RowHandle TargetRow, UE::Editor::DataStorage::RowHandle WidgetRow,
+	const UE::Editor::DataStorage::FMetaDataView& Arguments)
 {
 	if (DataStorage->IsRowAvailable(TargetRow))
 	{
-		UE::EditorDataStorage::FAttributeBinder Binder(TargetRow, DataStorage);
+		UE::Editor::DataStorage::FAttributeBinder Binder(TargetRow, DataStorage);
 		
 		return SNew(SHorizontalBox)
 			+SHorizontalBox::Slot()
diff --git a/Engine/Plugins/Experimental/EditorDataStorageFeatures/Source/TedsAssetData/Public/TedsAssetDataColumns.h b/Engine/Plugins/Experimental/EditorDataStorageFeatures/Source/TedsAssetData/Public/TedsAssetDataColumns.h
--- a/Engine/Plugins/Experimental/EditorDataStorageFeatures/Source/TedsAssetData/Public/TedsAssetDataColumns.h
+++ b/Engine/Plugins/Experimental/EditorDataStorageFeatures/Source/TedsAssetData/Public/TedsAssetDataColumns.h
@@ -23,7 +23,7 @@ struct FParentAssetPathColumn_Experimental : public FEditorDataStorageColumn
 {
 	GENERATED_BODY()
 	
-	TypedElementDataStorage::RowHandle ParentRow;
+	UE::Editor::DataStorage::RowHandle ParentRow;
 };
 
 USTRUCT()
@@ -31,7 +31,7 @@ struct FChildrenAssetPathColumn_Experimental : public FEditorDataStorageColumn
 {
 	GENERATED_BODY()
 	
-	TSet<TypedElementDataStorage::RowHandle> ChildrenRows;
+	TSet<UE::Editor::DataStorage::RowHandle> ChildrenRows;
 };
 
 USTRUCT()
@@ -39,7 +39,7 @@ struct FAssetsInPathColumn_Experimental : public FEditorDataStorageColumn
 {
 	GENERATED_BODY()
 	
-	TSet<TypedElementDataStorage::RowHandle> AssetsRow;
+	TSet<UE::Editor::DataStorage::RowHandle> AssetsRow;
 };
 
 USTRUCT()
diff --git a/Engine/Plugins/Experimental/EditorDataStorageFeatures/Source/TedsDebugger/Private/QueryEditor/TedsQueryEditorModel.cpp b/Engine/Plugins/Experimental/EditorDataStorageFeatures/Source/TedsDebugger/Private/QueryEditor/TedsQueryEditorModel.cpp
--- a/Engine/Plugins/Experimental/EditorDataStorageFeatures/Source/TedsDebugger/Private/QueryEditor/TedsQueryEditorModel.cpp
+++ b/Engine/Plugins/Experimental/EditorDataStorageFeatures/Source/TedsDebugger/Private/QueryEditor/TedsQueryEditorModel.cpp
@@ -101,9 +101,9 @@ namespace UE::Editor::DataStorage::Debug::QueryEditor
 		return TypedElementDataStorage;
 	}
 
-	TypedElementDataStorage::FQueryDescription FTedsQueryEditorModel::GenerateQueryDescription()
+	FQueryDescription FTedsQueryEditorModel::GenerateQueryDescription()
 	{
-		TypedElementDataStorage::FQueryDescription Description;
+		FQueryDescription Description;
 
 		for (const FConditionEntryInternal& Entry : Conditions)
 		{
@@ -135,7 +135,7 @@ namespace UE::Editor::DataStorage::Debug::QueryEditor
 			{
 				if (Entry.OperatorType == EOperatorType::Select)
 				{
-					Description.SelectionMetaData.Emplace(TypedElementDataStorage::FColumnMetaData());
+					Description.SelectionMetaData.Emplace(FColumnMetaData());
 					Description.SelectionAccessTypes.Add(TypedElementDataStorage::EQueryAccessType::ReadOnly);
 					Description.SelectionTypes.Add(Target);
 				}
@@ -170,9 +170,9 @@ namespace UE::Editor::DataStorage::Debug::QueryEditor
 		return Description;
 	}
 
-	TypedElementDataStorage::FQueryDescription FTedsQueryEditorModel::GenerateNoSelectQueryDescription()
+	FQueryDescription FTedsQueryEditorModel::GenerateNoSelectQueryDescription()
 	{
-		TypedElementDataStorage::FQueryDescription Description = GenerateQueryDescription();
+		FQueryDescription Description = GenerateQueryDescription();
 
 		// Move all the selection types to Condition types
 		const int32 SelectionTypeCount = Description.SelectionTypes.Num();
diff --git a/Engine/Plugins/Experimental/EditorDataStorageFeatures/Source/TedsDebugger/Private/QueryEditor/TedsQueryEditorModel.h b/Engine/Plugins/Experimental/EditorDataStorageFeatures/Source/TedsDebugger/Private/QueryEditor/TedsQueryEditorModel.h
--- a/Engine/Plugins/Experimental/EditorDataStorageFeatures/Source/TedsDebugger/Private/QueryEditor/TedsQueryEditorModel.h
+++ b/Engine/Plugins/Experimental/EditorDataStorageFeatures/Source/TedsDebugger/Private/QueryEditor/TedsQueryEditorModel.h
@@ -56,11 +56,11 @@ namespace Debug::QueryEditor
 		ITypedElementDataStorageInterface& GetTedsInterface();
 		const ITypedElementDataStorageInterface& GetTedsInterface() const;
 		
-		TypedElementDataStorage::FQueryDescription GenerateQueryDescription();
+		FQueryDescription GenerateQueryDescription();
 		// Special function to generate a description that puts the Select elements as All Conditions
 		// This is helpful for using a Count query type or for the TableViewer which requires
 		// the row query to have no select items
-		TypedElementDataStorage::FQueryDescription GenerateNoSelectQueryDescription();
+		FQueryDescription GenerateNoSelectQueryDescription();
 
 		int32 CountConditionsOfOperator(EOperatorType OperatorType) const;
 		
diff --git a/Engine/Plugins/Experimental/EditorDataStorageFeatures/Source/TedsDebugger/Private/Widgets/QueryEditor/TedsQueryEditorResultsView.cpp b/Engine/Plugins/Experimental/EditorDataStorageFeatures/Source/TedsDebugger/Private/Widgets/QueryEditor/TedsQueryEditorResultsView.cpp
--- a/Engine/Plugins/Experimental/EditorDataStorageFeatures/Source/TedsDebugger/Private/Widgets/QueryEditor/TedsQueryEditorResultsView.cpp
+++ b/Engine/Plugins/Experimental/EditorDataStorageFeatures/Source/TedsDebugger/Private/Widgets/QueryEditor/TedsQueryEditorResultsView.cpp
@@ -109,7 +109,7 @@ namespace UE::Editor::DataStorage
 			if(bModelDirty)
 			{
 				{
-					TypedElementDataStorage::FQueryDescription CountQueryDescription = Model->GenerateNoSelectQueryDescription();
+					FQueryDescription CountQueryDescription = Model->GenerateNoSelectQueryDescription();
 		
 					if (CountQueryHandle != InvalidQueryHandle)
 					{
@@ -120,7 +120,7 @@ namespace UE::Editor::DataStorage
 				}
 
 				{
-					TypedElementDataStorage::FQueryDescription TableViewerQueryDescription = Model->GenerateQueryDescription();
+					FQueryDescription TableViewerQueryDescription = Model->GenerateQueryDescription();
 
 					// Update the columns in the table viewer using the selection types from the query description
 					TableViewer->SetColumns(TArray<TWeakObjectPtr<const UScriptStruct>>(TableViewerQueryDescription.SelectionTypes));
@@ -192,7 +192,7 @@ namespace UE::Editor::DataStorage
 			ITypedElementDataStorageUiInterface* StorageUi = Registry->GetMutableDataStorageUi();
 			checkf(StorageUi, TEXT("SResultsView created before data storage interfaces were initialized."))
 
-			StorageUi->CreateWidgetConstructors(TEXT("General.Cell.RowHandle"), TypedElementDataStorage::FMetaDataView(), AssignWidgetToColumn);
+			StorageUi->CreateWidgetConstructors(TEXT("General.Cell.RowHandle"), FMetaDataView(), AssignWidgetToColumn);
 		}
 	} // namespace Debug::QueryEditor
 } // namespace UE::Editor::DataStorage
diff --git a/Engine/Plugins/Experimental/EditorDataStorageFeatures/Source/TedsDebugger/Private/Widgets/RowReferenceWidget.cpp b/Engine/Plugins/Experimental/EditorDataStorageFeatures/Source/TedsDebugger/Private/Widgets/RowReferenceWidget.cpp
--- a/Engine/Plugins/Experimental/EditorDataStorageFeatures/Source/TedsDebugger/Private/Widgets/RowReferenceWidget.cpp
+++ b/Engine/Plugins/Experimental/EditorDataStorageFeatures/Source/TedsDebugger/Private/Widgets/RowReferenceWidget.cpp
@@ -124,7 +124,7 @@ void URowReferenceWidgetFactory::RegisterQueries(ITypedElementDataStorageInterfa
 				.SetExecutionMode(EExecutionMode::GameThread),
 			[](
 				IQueryContext& Context,
-				TypedElementDataStorage::RowHandle UiRowHandle,
+				RowHandle UiRowHandle,
 				FTypedElementSlateWidgetReferenceColumn& Widget,
 				const FTypedElementRowReferenceColumn& Target)
 			{
@@ -145,7 +145,7 @@ FRowReferenceWidgetConstructor::FRowReferenceWidgetConstructor()
 {
 }
 
-TSharedPtr<SWidget> FRowReferenceWidgetConstructor::CreateWidget(const TypedElementDataStorage::FMetaDataView& Arguments)
+TSharedPtr<SWidget> FRowReferenceWidgetConstructor::CreateWidget(const UE::Editor::DataStorage::FMetaDataView& Arguments)
 {
 	return SNew(SBox)
 			.HAlign(HAlign_Left)
diff --git a/Engine/Plugins/Experimental/EditorDataStorageFeatures/Source/TedsDebugger/Private/Widgets/RowReferenceWidget.h b/Engine/Plugins/Experimental/EditorDataStorageFeatures/Source/TedsDebugger/Private/Widgets/RowReferenceWidget.h
--- a/Engine/Plugins/Experimental/EditorDataStorageFeatures/Source/TedsDebugger/Private/Widgets/RowReferenceWidget.h
+++ b/Engine/Plugins/Experimental/EditorDataStorageFeatures/Source/TedsDebugger/Private/Widgets/RowReferenceWidget.h
@@ -37,7 +37,7 @@ public:
 	~FRowReferenceWidgetConstructor() override = default;
 
 protected:
-	TSharedPtr<SWidget> CreateWidget(const TypedElementDataStorage::FMetaDataView& Arguments) override;
+	TSharedPtr<SWidget> CreateWidget(const UE::Editor::DataStorage::FMetaDataView& Arguments) override;
 	bool FinalizeWidget(ITypedElementDataStorageInterface* DataStorage, ITypedElementDataStorageUiInterface* DataStorageUi,
-		TypedElementDataStorage::RowHandle Row, const TSharedPtr<SWidget>& Widget) override;
+		UE::Editor::DataStorage::RowHandle Row, const TSharedPtr<SWidget>& Widget) override;
 };
\ No newline at end of file
diff --git a/Engine/Plugins/Experimental/EditorDataStorageFeatures/Source/TedsDebugger/Private/Widgets/SlateVisualizationWidget.cpp b/Engine/Plugins/Experimental/EditorDataStorageFeatures/Source/TedsDebugger/Private/Widgets/SlateVisualizationWidget.cpp
--- a/Engine/Plugins/Experimental/EditorDataStorageFeatures/Source/TedsDebugger/Private/Widgets/SlateVisualizationWidget.cpp
+++ b/Engine/Plugins/Experimental/EditorDataStorageFeatures/Source/TedsDebugger/Private/Widgets/SlateVisualizationWidget.cpp
@@ -21,7 +21,7 @@ FSlateVisualizationWidgetConstructor::FSlateVisualizationWidgetConstructor()
 
 }
 
-TSharedPtr<SWidget> FSlateVisualizationWidgetConstructor::CreateWidget(const TypedElementDataStorage::FMetaDataView& Arguments)
+TSharedPtr<SWidget> FSlateVisualizationWidgetConstructor::CreateWidget(const UE::Editor::DataStorage::FMetaDataView& Arguments)
 {
 	return SNew(SHorizontalBox);
 }
diff --git a/Engine/Plugins/Experimental/EditorDataStorageFeatures/Source/TedsOutliner/Private/Compatibility/SceneOutlinerRowHandleColumn.cpp b/Engine/Plugins/Experimental/EditorDataStorageFeatures/Source/TedsOutliner/Private/Compatibility/SceneOutlinerRowHandleColumn.cpp
--- a/Engine/Plugins/Experimental/EditorDataStorageFeatures/Source/TedsOutliner/Private/Compatibility/SceneOutlinerRowHandleColumn.cpp
+++ b/Engine/Plugins/Experimental/EditorDataStorageFeatures/Source/TedsOutliner/Private/Compatibility/SceneOutlinerRowHandleColumn.cpp
@@ -24,7 +24,7 @@ FSceneOutlinerRowHandleColumn::FSceneOutlinerRowHandleColumn(ISceneOutliner& Sce
 	ITypedElementDataStorageUiInterface* StorageUi = Registry->GetMutableDataStorageUi();
 	checkf(StorageUi, TEXT("FSceneOutlinerRowHandleColumn created before data storage interfaces were initialized."))
 
-	StorageUi->CreateWidgetConstructors(TEXT("General.Cell.RowHandle"), TypedElementDataStorage::FMetaDataView(), AssignWidgetToColumn);
+	StorageUi->CreateWidgetConstructors(TEXT("General.Cell.RowHandle"), UE::Editor::DataStorage::FMetaDataView(), AssignWidgetToColumn);
 }
 
 
diff --git a/Engine/Plugins/Experimental/EditorDataStorageFeatures/Source/TedsOutliner/Private/Compatibility/SceneOutlinerTedsBridge.cpp b/Engine/Plugins/Experimental/EditorDataStorageFeatures/Source/TedsOutliner/Private/Compatibility/SceneOutlinerTedsBridge.cpp
--- a/Engine/Plugins/Experimental/EditorDataStorageFeatures/Source/TedsOutliner/Private/Compatibility/SceneOutlinerTedsBridge.cpp
+++ b/Engine/Plugins/Experimental/EditorDataStorageFeatures/Source/TedsOutliner/Private/Compatibility/SceneOutlinerTedsBridge.cpp
@@ -143,7 +143,7 @@ class FOutlinerColumn : public ISceneOutlinerColumn
 {
 public:
 	FOutlinerColumn(
-		TypedElementDataStorage::QueryHandle InQuery,
+		UE::Editor::DataStorage::QueryHandle InQuery,
 		ITypedElementDataStorageInterface& InStorage, 
 		ITypedElementDataStorageUiInterface& InStorageUi, 
 		ITypedElementDataStorageCompatibilityInterface& InStorageCompatibility,
@@ -164,13 +164,13 @@ public:
 	{
 		MetaData.AddOrSetMutableData(TEXT("Name"), NameId.ToString());
 
-		using namespace TypedElementDataStorage;
+		using namespace UE::Editor::DataStorage;
 
-		TableViewerColumnImpl = MakeUnique<UE::Editor::DataStorage::FTedsTableViewerColumn>(NameId, InCellWidgetConstructor, InColumnTypes,
+		TableViewerColumnImpl = MakeUnique<FTedsTableViewerColumn>(NameId, InCellWidgetConstructor, InColumnTypes,
 			InHeaderWidgetConstructor, FComboMetaDataView(FGenericMetaDataView(MetaData)).Next(FQueryMetaDataView(Storage.GetQueryDescription(QueryHandle))));
 
 		TableViewerColumnImpl->SetIsRowVisibleDelegate(
-		UE::Editor::DataStorage::FTedsTableViewerColumn::FIsRowVisible::CreateRaw(this, &FOutlinerColumn::IsRowVisible)
+		FTedsTableViewerColumn::FIsRowVisible::CreateRaw(this, &FOutlinerColumn::IsRowVisible)
 		);
 		
 		// Try to find a fallback column from the regular item, for handling cases like folders which are not in TEDS but want to use TEDS columns
@@ -188,7 +188,7 @@ public:
 		TableViewerColumnImpl->Tick();
 	}
 
-	bool IsRowVisible(const TypedElementDataStorage::RowHandle InRowHandle) const
+	bool IsRowVisible(const UE::Editor::DataStorage::RowHandle InRowHandle) const
 	{
 		TSharedPtr<ISceneOutliner> OutlinerPinned = OwningOutliner.Pin();
 
@@ -262,7 +262,7 @@ public:
 	
 	const TSharedRef<SWidget> ConstructRowWidget(FSceneOutlinerTreeItemRef TreeItem, const STableRow<FSceneOutlinerTreeItemPtr>& Row) override
 	{
-		using namespace TypedElementDataStorage;
+		using namespace UE::Editor::DataStorage;
 		
 		RowHandle RowHandle = InvalidRowHandle;
 
@@ -314,8 +314,8 @@ public:
 	ITypedElementDataStorageInterface& Storage;
 	ITypedElementDataStorageUiInterface& StorageUi;
 	ITypedElementDataStorageCompatibilityInterface& StorageCompatibility;
-	TypedElementDataStorage::QueryHandle QueryHandle;
-	TypedElementDataStorage::FMetaData MetaData;
+	UE::Editor::DataStorage::QueryHandle QueryHandle;
+	UE::Editor::DataStorage::FMetaData MetaData;
 	FName NameId;
 	TSharedPtr<ISceneOutlinerColumn> FallbackColumn;
 	TWeakPtr<ISceneOutliner> OwningOutliner;
@@ -507,7 +507,7 @@ void FSceneOutlinerTedsBridge::AssignQuery(UE::Editor::DataStorage::QueryHandle
 	if (TSharedPtr<ISceneOutliner> OutlinerPinned = Outliner.Pin())
 	{
 		const ITypedElementDataStorageInterface::FQueryDescription& Description = Storage->GetQueryDescription(Query);
-		TypedElementDataStorage::FQueryMetaDataView MetaDataView(Description);
+		UE::Editor::DataStorage::FQueryMetaDataView MetaDataView(Description);
 
 		ClearColumns(*OutlinerPinned);
 
diff --git a/Engine/Plugins/Experimental/EditorDataStorageFeatures/Source/TedsOutliner/Private/TedsOutlinerFilter.cpp b/Engine/Plugins/Experimental/EditorDataStorageFeatures/Source/TedsOutliner/Private/TedsOutlinerFilter.cpp
--- a/Engine/Plugins/Experimental/EditorDataStorageFeatures/Source/TedsOutliner/Private/TedsOutlinerFilter.cpp
+++ b/Engine/Plugins/Experimental/EditorDataStorageFeatures/Source/TedsOutliner/Private/TedsOutlinerFilter.cpp
@@ -6,7 +6,7 @@
 
 FTedsOutlinerFilter::FTedsOutlinerFilter(const FName& InFilterName, const FText& InFilterDisplayName,
 	TSharedPtr<FFilterCategory> InCategory, TSharedRef<FTedsOutlinerImpl> InTedsOutlinerImpl,
-	const TypedElementDataStorage::FQueryDescription& InFilterQuery)
+	const UE::Editor::DataStorage::FQueryDescription& InFilterQuery)
 	: FFilterBase(InCategory)
 	, FilterName(InFilterName)
 	, FilterDisplayName(InFilterDisplayName)
diff --git a/Engine/Plugins/Experimental/EditorDataStorageFeatures/Source/TedsOutliner/Private/TedsOutlinerImpl.cpp b/Engine/Plugins/Experimental/EditorDataStorageFeatures/Source/TedsOutliner/Private/TedsOutlinerImpl.cpp
--- a/Engine/Plugins/Experimental/EditorDataStorageFeatures/Source/TedsOutliner/Private/TedsOutlinerImpl.cpp
+++ b/Engine/Plugins/Experimental/EditorDataStorageFeatures/Source/TedsOutliner/Private/TedsOutlinerImpl.cpp
@@ -52,7 +52,7 @@ void FTedsOutlinerImpl::CreateLabelWidgetConstructors()
 
 	// Create and assign the widget constructors for the label column
 	
-	auto CreateWidgetConstructorForQuery = [this](TypedElementDataStorage::FQueryDescription InQueryDescription) -> TSharedPtr<FTypedElementWidgetConstructor>
+	auto CreateWidgetConstructorForQuery = [this](UE::Editor::DataStorage::FQueryDescription InQueryDescription) -> TSharedPtr<FTypedElementWidgetConstructor>
 	{
 		// Create the Widget Constructor for the item label column
 		using MatchApproach = ITypedElementDataStorageUiInterface::EMatchApproach;
@@ -94,7 +94,7 @@ void FTedsOutlinerImpl::CreateLabelWidgetConstructors()
 		return OutWidgetConstructorPtr;
 	};
 	
-	TypedElementDataStorage::QueryHandle LabelColumnQueryHandle = Storage->RegisterQuery(
+	UE::Editor::DataStorage::QueryHandle LabelColumnQueryHandle = Storage->RegisterQuery(
 																	Select()
 																		.ReadWrite<FTypedElementLabelColumn, FTypedElementClassTypeInfoColumn>()
 																	.Compile()
@@ -127,7 +127,7 @@ void FTedsOutlinerImpl::CreateFilterQueries()
 				if (Struct->IsChildOf(TedsColumn) || Struct->IsChildOf(TedsTag))
 				{
 					// Create a query description to filter for this tag/column
-					TypedElementDataStorage::FQueryDescription FilterQueryDesc =
+					UE::Editor::DataStorage::FQueryDescription FilterQueryDesc =
 					Select()
 					.Where()
 						.All(Struct)
@@ -145,7 +145,7 @@ void FTedsOutlinerImpl::CreateFilterQueries()
 	// Custom filters input by the user
 	TSharedRef<FFilterCategory> CustomFiltersCategory = MakeShared<FFilterCategory>(LOCTEXT("TedsFilters", "TEDS Custom Filters"), LOCTEXT("TedsFiltersTooltip", "Filter by custom TEDS queries"));
 
-	for(const TPair<FName, const TypedElementDataStorage::FQueryDescription>& FilterQuery : CreationParams.FilterQueries)
+	for(const TPair<FName, const UE::Editor::DataStorage::FQueryDescription>& FilterQuery : CreationParams.FilterQueries)
 	{
 		// TEDS-Outliner TODO: Custom filters need a localizable display name instead of using the FName, but we need to change how they are added first
 		// to see if it can be consolidated with the SFilterBar API
@@ -178,7 +178,7 @@ FTedsOutlinerImpl::~FTedsOutlinerImpl()
 
 	// This is done outside of unregister queries because that is called when the internal query changes (i.e filter) but we don't want to unregister the
 	// label widget queries until shutdown
-	for(const TPair<TypedElementDataStorage::QueryHandle, TSharedPtr<FTypedElementWidgetConstructor>>& QueryConstructorPair : QueryToWidgetConstructorMap)
+	for(const TPair<UE::Editor::DataStorage::QueryHandle, TSharedPtr<FTypedElementWidgetConstructor>>& QueryConstructorPair : QueryToWidgetConstructorMap)
 	{
 		Storage->UnregisterQuery(QueryConstructorPair.Key);
 	}
@@ -191,7 +191,7 @@ FTedsOutlinerImpl::FIsItemCompatible& FTedsOutlinerImpl::IsItemCompatible()
 	return IsItemCompatibleWithTeds;
 }
 
-void FTedsOutlinerImpl::SetSelection(const TArray<TypedElementDataStorage::RowHandle>& InSelectedRows)
+void FTedsOutlinerImpl::SetSelection(const TArray<UE::Editor::DataStorage::RowHandle>& InSelectedRows)
 {
 	if (!SelectionSetName.IsSet())
 	{
@@ -200,7 +200,7 @@ void FTedsOutlinerImpl::SetSelection(const TArray<TypedElementDataStorage::RowHa
 	
 	ClearSelection();
 
-	for(TypedElementDataStorage::RowHandle Row : InSelectedRows)
+	for(UE::Editor::DataStorage::RowHandle Row : InSelectedRows)
 	{
 		Storage->AddColumn(Row, FTypedElementSelectionColumn{ .SelectionSet = SelectionSetName.GetValue() });
 	}
@@ -210,18 +210,18 @@ TSharedRef<SWidget> FTedsOutlinerImpl::CreateLabelWidgetForItem(UE::Editor::Data
 {
 	auto CreateWidgetForQuery = [InRowHandle, this, &InRow](const TPair<UE::Editor::DataStorage::QueryHandle, TSharedPtr<FTypedElementWidgetConstructor>>& QueryConstructorPair) -> TSharedPtr<SWidget>
 	{
-		TypedElementDataStorage::FQueryDescription QueryDescription = Storage->GetQueryDescription(QueryConstructorPair.Key);
+		UE::Editor::DataStorage::FQueryDescription QueryDescription = Storage->GetQueryDescription(QueryConstructorPair.Key);
 		
 		// Create a generic metadata view for the Type Widget
-		TypedElementDataStorage::FMetaData QueryWideMetaData;
+		UE::Editor::DataStorage::FMetaData QueryWideMetaData;
 		QueryWideMetaData.AddImmutableData("TypeInfoWidget_bUseIcon", true);
-		TypedElementDataStorage::FGenericMetaDataView GenericMetaDataView(QueryWideMetaData);
+		UE::Editor::DataStorage::FGenericMetaDataView GenericMetaDataView(QueryWideMetaData);
 
 		// Create metadata for the query itself
-		TypedElementDataStorage::FQueryMetaDataView QueryMetaDataView(QueryDescription);
+		UE::Editor::DataStorage::FQueryMetaDataView QueryMetaDataView(QueryDescription);
 
 		// Combine the two metadata
-		TypedElementDataStorage::FComboMetaDataView MetaDataArgs(GenericMetaDataView, QueryMetaDataView);
+		UE::Editor::DataStorage::FComboMetaDataView MetaDataArgs(GenericMetaDataView, QueryMetaDataView);
 
 		TArray<TWeakObjectPtr<const UScriptStruct>> ColumnTypes(QueryDescription.SelectionTypes);
 		TSharedPtr<FTypedElementWidgetConstructor> CellWidgetConstructor = QueryConstructorPair.Value;
@@ -247,7 +247,7 @@ TSharedRef<SWidget> FTedsOutlinerImpl::CreateLabelWidgetForItem(UE::Editor::Data
 	
 	TSharedRef<SHorizontalBox> CombinedWidget = SNew(SHorizontalBox);
 
-	for(const TPair<TypedElementDataStorage::QueryHandle, TSharedPtr<FTypedElementWidgetConstructor>>& QueryConstructorPair : QueryToWidgetConstructorMap)
+	for(const TPair<UE::Editor::DataStorage::QueryHandle, TSharedPtr<FTypedElementWidgetConstructor>>& QueryConstructorPair : QueryToWidgetConstructorMap)
 	{
 		if (TSharedPtr<SWidget> WidgetForQuery = CreateWidgetForQuery(QueryConstructorPair))
 		{
@@ -265,13 +265,13 @@ TSharedRef<SWidget> FTedsOutlinerImpl::CreateLabelWidgetForItem(UE::Editor::Data
 	return CombinedWidget;
 }
 
-void FTedsOutlinerImpl::AppendQuery(TypedElementDataStorage::FQueryDescription& Query1, const TypedElementDataStorage::FQueryDescription& Query2)
+void FTedsOutlinerImpl::AppendQuery(UE::Editor::DataStorage::FQueryDescription& Query1, const UE::Editor::DataStorage::FQueryDescription& Query2)
 {
 	// TEDS-Outliner TODO: We simply discard duplicate types for now but we probably want a more robust system to detect duplicates and conflicting conditions
 	for(int32 i = 0; i < Query2.ConditionOperators.Num(); ++i)
 	{
 		// Make sure we don't add duplicate conditions
-		TypedElementDataStorage::FQueryDescription::FOperator* FoundCondition = Query1.ConditionOperators.FindByPredicate([&Query2, i](const TypedElementDataStorage::FQueryDescription::FOperator& Op)
+		UE::Editor::DataStorage::FQueryDescription::FOperator* FoundCondition = Query1.ConditionOperators.FindByPredicate([&Query2, i](const UE::Editor::DataStorage::FQueryDescription::FOperator& Op)
 		{
 			return Op.Type == Query2.ConditionOperators[i].Type;
 		});
@@ -290,7 +290,7 @@ void FTedsOutlinerImpl::AppendQuery(TypedElementDataStorage::FQueryDescription&
 	}
 }
 
-void FTedsOutlinerImpl::AddExternalQuery(FName QueryName, const TypedElementDataStorage::FQueryDescription& InQueryDescription)
+void FTedsOutlinerImpl::AddExternalQuery(FName QueryName, const UE::Editor::DataStorage::FQueryDescription& InQueryDescription)
 {
 	ExternalQueries.Emplace(QueryName, InQueryDescription);
 
@@ -302,15 +302,15 @@ void FTedsOutlinerImpl::RemoveExternalQuery(FName QueryName)
 	ExternalQueries.Remove(QueryName);
 }
 
-void FTedsOutlinerImpl::AppendExternalQueries(TypedElementDataStorage::FQueryDescription& OutQuery)
+void FTedsOutlinerImpl::AppendExternalQueries(UE::Editor::DataStorage::FQueryDescription& OutQuery)
 {
-	for(const TPair<FName, TypedElementDataStorage::FQueryDescription>& ExternalQuery : ExternalQueries)
+	for(const TPair<FName, UE::Editor::DataStorage::FQueryDescription>& ExternalQuery : ExternalQueries)
 	{
 		AppendQuery(OutQuery, ExternalQuery.Value);
 	}
 }
 
-bool FTedsOutlinerImpl::HasItemParentChanged(TypedElementDataStorage::RowHandle InRowHandle, TypedElementDataStorage::RowHandle ParentRowHandle) const
+bool FTedsOutlinerImpl::HasItemParentChanged(UE::Editor::DataStorage::RowHandle InRowHandle, UE::Editor::DataStorage::RowHandle ParentRowHandle) const
 {
 	const FSceneOutlinerTreeItemPtr Item = SceneOutliner->GetTreeItem(InRowHandle, true);
 
@@ -339,7 +339,7 @@ bool FTedsOutlinerImpl::HasItemParentChanged(TypedElementDataStorage::RowHandle
 	return false;
 }
 
-bool FTedsOutlinerImpl::CanDisplayRow(TypedElementDataStorage::RowHandle ItemRowHandle) const
+bool FTedsOutlinerImpl::CanDisplayRow(UE::Editor::DataStorage::RowHandle ItemRowHandle) const
 {
 	/*
 	 * Don't display widgets that are created for rows in this table viewer. Widgets are only created for rows that are currently visible, so if we
@@ -366,9 +366,9 @@ void FTedsOutlinerImpl::CreateItemsFromQuery(TArray<FSceneOutlinerTreeItemPtr>&
 	using namespace TypedElementQueryBuilder;
 	using DSI = ITypedElementDataStorageInterface;
 	
-	TArray<TypedElementDataStorage::RowHandle> Rows;
+	TArray<UE::Editor::DataStorage::RowHandle> Rows;
 	
-	TypedElementDataStorage::DirectQueryCallback RowCollector = CreateDirectQueryCallbackBinding(
+	UE::Editor::DataStorage::DirectQueryCallback RowCollector = CreateDirectQueryCallbackBinding(
 		[&Rows](DSI::IDirectQueryContext& Context)
 		{
 			TConstArrayView<UE::Editor::DataStorage::RowHandle> ContextRows = Context.GetRowHandles();
@@ -430,7 +430,7 @@ void FTedsOutlinerImpl::CreateChildren(const FSceneOutlinerTreeItemPtr& Item, TA
 	TSet<UE::Editor::DataStorage::RowHandle> MatchedRowsWithParentColumn;
 	
 	// Collect all entities that are owned by our entity
-	TypedElementDataStorage::DirectQueryCallback ChildRowCollector = CreateDirectQueryCallbackBinding(
+	UE::Editor::DataStorage::DirectQueryCallback ChildRowCollector = CreateDirectQueryCallbackBinding(
 	[&MatchedRowsWithParentColumn] (const DSI::IDirectQueryContext& Context)
 	{
 		MatchedRowsWithParentColumn.Append(Context.GetRowHandles());
@@ -616,9 +616,9 @@ void FTedsOutlinerImpl::RecompileQueries()
 			TEXT("Update item parent"),
 			FProcessor(EQueryTickPhase::DuringPhysics, Storage->GetQueryTickGroupName(EQueryTickGroups::Update))
 				.SetExecutionMode(EExecutionMode::GameThread),
-			[this](IQueryContext& Context, TypedElementDataStorage::RowHandle Row)
+			[this](IQueryContext& Context, UE::Editor::DataStorage::RowHandle Row)
 			{
-				TypedElementDataStorage::RowHandle ParentRowHandle = InvalidRowHandle;
+				UE::Editor::DataStorage::RowHandle ParentRowHandle = InvalidRowHandle;
 
 				if (const FTableRowParentColumn* ParentColumn = Context.GetColumn<FTableRowParentColumn>())
 				{
@@ -715,15 +715,15 @@ void FTedsOutlinerImpl::ClearSelection() const
 	using namespace TypedElementQueryBuilder;
 	using DSI = ITypedElementDataStorageInterface;
 
-	TArray<TypedElementDataStorage::RowHandle> RowsToRemoveSelectionColumn;
+	TArray<UE::Editor::DataStorage::RowHandle> RowsToRemoveSelectionColumn;
 
 	// Query to remove the selection column from all rows that belong to this selection set
-	TypedElementDataStorage::DirectQueryCallback RowCollector = CreateDirectQueryCallbackBinding(
-	[this, &RowsToRemoveSelectionColumn](const DSI::IDirectQueryContext& Context, const TypedElementDataStorage::RowHandle* RowHandles)
+	UE::Editor::DataStorage::DirectQueryCallback RowCollector = CreateDirectQueryCallbackBinding(
+	[this, &RowsToRemoveSelectionColumn](const DSI::IDirectQueryContext& Context, const UE::Editor::DataStorage::RowHandle* RowHandles)
 	{
-		const TConstArrayView<TypedElementDataStorage::RowHandle> Rows(RowHandles, Context.GetRowCount());
+		const TConstArrayView<UE::Editor::DataStorage::RowHandle> Rows(RowHandles, Context.GetRowCount());
 
-		for(const TypedElementDataStorage::RowHandle RowHandle : Rows)
+		for(const UE::Editor::DataStorage::RowHandle RowHandle : Rows)
 		{
 			if (const FTypedElementSelectionColumn* SelectionColumn = Storage->GetColumn<FTypedElementSelectionColumn>(RowHandle))
 			{
@@ -737,7 +737,7 @@ void FTedsOutlinerImpl::ClearSelection() const
 
 	Storage->RunQuery(SelectedRowsQuery, RowCollector);
 
-	for(const TypedElementDataStorage::RowHandle RowHandle : RowsToRemoveSelectionColumn)
+	for(const UE::Editor::DataStorage::RowHandle RowHandle : RowsToRemoveSelectionColumn)
 	{
 		Storage->RemoveColumn<FTypedElementSelectionColumn>(RowHandle);
 	}
diff --git a/Engine/Plugins/Experimental/EditorDataStorageFeatures/Source/TedsOutliner/Private/TedsOutlinerMode.cpp b/Engine/Plugins/Experimental/EditorDataStorageFeatures/Source/TedsOutliner/Private/TedsOutlinerMode.cpp
--- a/Engine/Plugins/Experimental/EditorDataStorageFeatures/Source/TedsOutliner/Private/TedsOutlinerMode.cpp
+++ b/Engine/Plugins/Experimental/EditorDataStorageFeatures/Source/TedsOutliner/Private/TedsOutlinerMode.cpp
@@ -60,7 +60,7 @@ void FTedsOutlinerMode::OnSelectionChanged()
 	{
 		if(const FTedsOutlinerTreeItem* TedsItem = InItem.CastTo<FTedsOutlinerTreeItem>())
 		{
-			const TypedElementDataStorage::RowHandle RowHandle = TedsItem->GetRowHandle();
+			const UE::Editor::DataStorage::RowHandle RowHandle = TedsItem->GetRowHandle();
 
 			if(const FTypedElementSelectionColumn* SelectionColumn = Storage->GetColumn<FTypedElementSelectionColumn>(RowHandle))
 			{
@@ -83,7 +83,7 @@ void FTedsOutlinerMode::OnItemSelectionChanged(FSceneOutlinerTreeItemPtr Item, E
 		return; // Direct selection means we selected from outside the Outliner i.e through TEDS, so we don't need to redo the column addition
 	}
 
-	TArray<TypedElementDataStorage::RowHandle> RowHandles;
+	TArray<UE::Editor::DataStorage::RowHandle> RowHandles;
 	
 	// The selection in the Outliner changed, update TEDS
 	Selection.ForEachItem([&RowHandles](FSceneOutlinerTreeItemPtr& Item)
@@ -107,7 +107,7 @@ TSharedPtr<FDragDropOperation> FTedsOutlinerMode::CreateDragDropOperation(const
 		return nullptr;
 	}
 	
-	TArray<TypedElementDataStorage::RowHandle> DraggedRowHandles;
+	TArray<UE::Editor::DataStorage::RowHandle> DraggedRowHandles;
 
 	for(const FSceneOutlinerTreeItemPtr& Item :InTreeItems)
 	{
@@ -127,7 +127,7 @@ bool FTedsOutlinerMode::ParseDragDrop(FSceneOutlinerDragDropPayload& OutPayload,
 	{
 		const FTedsRowDragDropOp& TedsOp = static_cast<const FTedsRowDragDropOp&>(Operation);
 
-		for(TypedElementDataStorage::RowHandle RowHandle : TedsOp.DraggedRows)
+		for(UE::Editor::DataStorage::RowHandle RowHandle : TedsOp.DraggedRows)
 		{
 			OutPayload.DraggedItems.Add(SceneOutliner->GetTreeItem(RowHandle));
 		}
@@ -148,7 +148,7 @@ FSceneOutlinerDragValidationInfo FTedsOutlinerMode::ValidateDrop(const ISceneOut
 			LOCTEXT("DropDisabled", "Drag/Drop is disabled due to missing hierarchy data!"));
 	}
 	
-	TArray<TypedElementDataStorage::RowHandle> DraggedRowHandles;
+	TArray<UE::Editor::DataStorage::RowHandle> DraggedRowHandles;
 
 	Payload.ForEachItem<FTedsOutlinerTreeItem>([&DraggedRowHandles](FTedsOutlinerTreeItem& TedsItem)
 		{
@@ -159,7 +159,7 @@ FSceneOutlinerDragValidationInfo FTedsOutlinerMode::ValidateDrop(const ISceneOut
 	// TEDS-Outliner TODO: Need better drag/drop validation and better place for this, TEDS-Outliner does not know about what types these rows are and all types that exist and what attachment is valid
 	if(const FTedsOutlinerTreeItem* TedsItem = DropTarget.CastTo<FTedsOutlinerTreeItem>())
 	{
-		TypedElementDataStorage::RowHandle DropTargetRowHandle = TedsItem->GetRowHandle();
+		UE::Editor::DataStorage::RowHandle DropTargetRowHandle = TedsItem->GetRowHandle();
 
 		FTypedElementClassTypeInfoColumn* DropTargetTypeInfoColumn = Storage->GetColumn<FTypedElementClassTypeInfoColumn>(DropTargetRowHandle);
 
@@ -171,7 +171,7 @@ FSceneOutlinerDragValidationInfo FTedsOutlinerMode::ValidateDrop(const ISceneOut
 
 		
 		// TEDS-Outliner TODO: Currently we detect parent changes by removing the column and then adding the column back with the new parent 
-		for(TypedElementDataStorage::RowHandle RowHandle : DraggedRowHandles)
+		for(UE::Editor::DataStorage::RowHandle RowHandle : DraggedRowHandles)
 		{
 			FTypedElementClassTypeInfoColumn* TypeInfoColumn = Storage->GetColumn<FTypedElementClassTypeInfoColumn>(RowHandle);
 
@@ -198,7 +198,7 @@ FSceneOutlinerDragValidationInfo FTedsOutlinerMode::ValidateDrop(const ISceneOut
 		{
 			bool bValidDetach = false;
 			
-			for(TypedElementDataStorage::RowHandle RowHandle : DraggedRowHandles)
+			for(UE::Editor::DataStorage::RowHandle RowHandle : DraggedRowHandles)
 			{
 				if(Storage->HasColumns(RowHandle, MakeArrayView({HierarchyData.GetValue().HierarchyColumn})))
 				{
@@ -227,7 +227,7 @@ void FTedsOutlinerMode::OnDrop(ISceneOutlinerTreeItem& DropTarget, const FSceneO
 		return;
 	}
 	
-	TArray<TypedElementDataStorage::RowHandle> DraggedRowHandles;
+	TArray<UE::Editor::DataStorage::RowHandle> DraggedRowHandles;
 
 	Payload.ForEachItem<FTedsOutlinerTreeItem>([&DraggedRowHandles](FTedsOutlinerTreeItem& TedsItem)
 	{
@@ -236,7 +236,7 @@ void FTedsOutlinerMode::OnDrop(ISceneOutlinerTreeItem& DropTarget, const FSceneO
 
 	if(ValidationInfo.CompatibilityType == ESceneOutlinerDropCompatibility::CompatibleDetach)
 	{
-		for(TypedElementDataStorage::RowHandle RowHandle : DraggedRowHandles)
+		for(UE::Editor::DataStorage::RowHandle RowHandle : DraggedRowHandles)
 		{
 			Storage->RemoveColumn(RowHandle, HierarchyData.GetValue().HierarchyColumn);
 			Storage->AddColumn<FTypedElementSyncBackToWorldTag>(RowHandle);
@@ -245,9 +245,9 @@ void FTedsOutlinerMode::OnDrop(ISceneOutlinerTreeItem& DropTarget, const FSceneO
 	
 	if(const FTedsOutlinerTreeItem* TedsItem = DropTarget.CastTo<FTedsOutlinerTreeItem>())
 	{
-		TypedElementDataStorage::RowHandle DropTargetRowHandle = TedsItem->GetRowHandle();
+		UE::Editor::DataStorage::RowHandle DropTargetRowHandle = TedsItem->GetRowHandle();
 		
-		for(TypedElementDataStorage::RowHandle RowHandle : DraggedRowHandles)
+		for(UE::Editor::DataStorage::RowHandle RowHandle : DraggedRowHandles)
 		{
 			// Add the column
 			Storage->AddColumn(RowHandle, HierarchyData.GetValue().HierarchyColumn);
diff --git a/Engine/Plugins/Experimental/EditorDataStorageFeatures/Source/TedsOutliner/Private/TedsOutlinerModule.cpp b/Engine/Plugins/Experimental/EditorDataStorageFeatures/Source/TedsOutliner/Private/TedsOutlinerModule.cpp
--- a/Engine/Plugins/Experimental/EditorDataStorageFeatures/Source/TedsOutliner/Private/TedsOutlinerModule.cpp
+++ b/Engine/Plugins/Experimental/EditorDataStorageFeatures/Source/TedsOutliner/Private/TedsOutlinerModule.cpp
@@ -61,7 +61,7 @@ FTedsOutlinerModule::FTedsOutlinerModule()
 {
 }
 
-TSharedRef<ISceneOutliner> FTedsOutlinerModule::CreateTedsOutliner(const FSceneOutlinerInitializationOptions& InInitOptions, const FTedsOutlinerParams& InInitTedsOptions, TypedElementDataStorage::QueryHandle ColumnQuery) const
+TSharedRef<ISceneOutliner> FTedsOutlinerModule::CreateTedsOutliner(const FSceneOutlinerInitializationOptions& InInitOptions, const FTedsOutlinerParams& InInitTedsOptions, UE::Editor::DataStorage::QueryHandle ColumnQuery) const
 {
 	UTypedElementRegistry* Registry = UTypedElementRegistry::GetInstance();
 
@@ -110,20 +110,20 @@ void FTedsOutlinerModule::ShutdownModule()
 	IModuleInterface::ShutdownModule();
 }
 
-TypedElementDataStorage::QueryHandle FTedsOutlinerModule::GetLevelEditorTedsOutlinerColumnQuery()
+UE::Editor::DataStorage::QueryHandle FTedsOutlinerModule::GetLevelEditorTedsOutlinerColumnQuery()
 {
 	UTypedElementRegistry* Registry = UTypedElementRegistry::GetInstance();
 	ITypedElementDataStorageInterface* Storage = Registry->GetMutableDataStorage();
 		
 	using namespace TypedElementQueryBuilder;
 
-	static TypedElementDataStorage::QueryHandle ColumnQuery = Storage->RegisterQuery(
+	static UE::Editor::DataStorage::QueryHandle ColumnQuery = Storage->RegisterQuery(
 		Select()
 			.ReadOnly<FTypedElementClassTypeInfoColumn, FTypedElementAlertColumn, FTypedElementChildAlertColumn>()
 		.Compile());
 
 	// Query to also include revision control info
-	static TypedElementDataStorage::QueryHandle RevisionControlQuery = Storage->RegisterQuery(
+	static UE::Editor::DataStorage::QueryHandle RevisionControlQuery = Storage->RegisterQuery(
 		Select()
 			.ReadOnly<FTypedElementClassTypeInfoColumn, FTypedElementPackageReference, FTypedElementAlertColumn>()
 		.Compile());
@@ -146,7 +146,7 @@ TSharedRef<SWidget> FTedsOutlinerModule::CreateLevelEditorTedsOutliner()
 	using namespace TypedElementQueryBuilder;
 
 	// The Outliner is populated with Actors and Entities
-	TypedElementDataStorage::FQueryDescription OutlinerQueryDescription =
+	UE::Editor::DataStorage::FQueryDescription OutlinerQueryDescription =
 						Select()
 						.Where()
 							.All<FTypedElementClassTypeInfoColumn>() // TEDS-Outliner TODO: Currently looking at all entries with type info in TEDS
@@ -163,7 +163,7 @@ TSharedRef<SWidget> FTedsOutlinerModule::CreateLevelEditorTedsOutliner()
 	Params.bUseDefaultTedsFilters = true;
 
 	// Example Query to filter for actors
-	TypedElementDataStorage::FQueryDescription ActorFilterQuery =
+	UE::Editor::DataStorage::FQueryDescription ActorFilterQuery =
 					Select()
 					.Where()
 						.All<FTypedElementActorTag>()
diff --git a/Engine/Plugins/Experimental/EditorDataStorageFeatures/Source/TedsOutliner/Private/Widgets/AlertWidget.cpp b/Engine/Plugins/Experimental/EditorDataStorageFeatures/Source/TedsOutliner/Private/Widgets/AlertWidget.cpp
--- a/Engine/Plugins/Experimental/EditorDataStorageFeatures/Source/TedsOutliner/Private/Widgets/AlertWidget.cpp
+++ b/Engine/Plugins/Experimental/EditorDataStorageFeatures/Source/TedsOutliner/Private/Widgets/AlertWidget.cpp
@@ -17,7 +17,7 @@
 namespace UE::Editor::DataStorage::Widgets::Private
 {
 	void UpdateWidget(const TSharedPtr<SWidget>& Widget, const FText& Alert, bool bIsWarning, uint16 ErrorCount, uint16 WarningCount,
-		TypedElementDataStorage::RowHandle RowWithAlertAction)
+		UE::Editor::DataStorage::RowHandle RowWithAlertAction)
 	{
 		if (Widget)
 		{
@@ -319,7 +319,7 @@ FAlertWidgetConstructor::FAlertWidgetConstructor()
 }
 
 TSharedPtr<SWidget> FAlertWidgetConstructor::CreateWidget(
-	const TypedElementDataStorage::FMetaDataView& Arguments)
+	const UE::Editor::DataStorage::FMetaDataView& Arguments)
 {
 	return SNew(SOverlay)
 		+SOverlay::Slot()
@@ -366,9 +366,9 @@ TConstArrayView<const UScriptStruct*> FAlertWidgetConstructor::GetAdditionalColu
 }
 
 bool FAlertWidgetConstructor::FinalizeWidget(ITypedElementDataStorageInterface* DataStorage,
-	ITypedElementDataStorageUiInterface* DataStorageUi, TypedElementDataStorage::RowHandle Row, const TSharedPtr<SWidget>& Widget)
+	ITypedElementDataStorageUiInterface* DataStorageUi, UE::Editor::DataStorage::RowHandle Row, const TSharedPtr<SWidget>& Widget)
 {
-	using namespace TypedElementDataStorage;
+	using namespace UE::Editor::DataStorage;
 
 	RowHandle TargetRow = DataStorage->GetColumn<FTypedElementRowReferenceColumn>(Row)->Row;
 	const FTypedElementAlertColumn* Alert = DataStorage->GetColumn<FTypedElementAlertColumn>(TargetRow);
@@ -377,7 +377,7 @@ bool FAlertWidgetConstructor::FinalizeWidget(ITypedElementDataStorageInterface*
 	uint16 ErrorCount = ChildAlert ? ChildAlert->Counts[static_cast<size_t>(FTypedElementAlertColumnType::Error)] : 0;
 	uint16 WarningCount = ChildAlert ? ChildAlert->Counts[static_cast<size_t>(FTypedElementAlertColumnType::Warning)] : 0;
 
-	UE::Editor::DataStorage::Widgets::Private::UpdateWidget(
+	Widgets::Private::UpdateWidget(
 		Widget, 
 		Alert ? Alert->Message : FText::GetEmpty(), 
 		Alert ? (Alert->AlertType == FTypedElementAlertColumnType::Warning) : false, 
@@ -399,7 +399,7 @@ FAlertHeaderWidgetConstructor::FAlertHeaderWidgetConstructor()
 }
 
 TSharedPtr<SWidget> FAlertHeaderWidgetConstructor::CreateWidget(
-	const TypedElementDataStorage::FMetaDataView& Arguments)
+	const UE::Editor::DataStorage::FMetaDataView& Arguments)
 {
 	return SNew(SImage)
 		.DesiredSizeOverride(FVector2D(16.f, 16.f))
diff --git a/Engine/Plugins/Experimental/EditorDataStorageFeatures/Source/TedsOutliner/Private/Widgets/OutlinerLabelWidget.cpp b/Engine/Plugins/Experimental/EditorDataStorageFeatures/Source/TedsOutliner/Private/Widgets/OutlinerLabelWidget.cpp
--- a/Engine/Plugins/Experimental/EditorDataStorageFeatures/Source/TedsOutliner/Private/Widgets/OutlinerLabelWidget.cpp
+++ b/Engine/Plugins/Experimental/EditorDataStorageFeatures/Source/TedsOutliner/Private/Widgets/OutlinerLabelWidget.cpp
@@ -118,14 +118,15 @@ FOutlinerLabelWidgetConstructor::FOutlinerLabelWidgetConstructor()
 
 TSharedPtr<SWidget> FOutlinerLabelWidgetConstructor::CreateWidget(ITypedElementDataStorageInterface* DataStorage,
 	ITypedElementDataStorageUiInterface* DataStorageUi, RowHandle TargetRow, RowHandle WidgetRow,
-	const TypedElementDataStorage::FMetaDataView& Arguments)
+	const UE::Editor::DataStorage::FMetaDataView& Arguments)
 {
+	using namespace UE::Editor::DataStorage;
 	if(DataStorage->IsRowAvailable(TargetRow))
 	{
-		UE::EditorDataStorage::FAttributeBinder Binder(TargetRow, DataStorage);
+		FAttributeBinder Binder(TargetRow, DataStorage);
 
 		TSharedRef<SLayeredImage> LayeredImageWidget = SNew(SLayeredImage)
-					.Image(UE::Editor::DataStorage::TableViewerUtils::GetIconForRow(DataStorage, TargetRow))
+					.Image(TableViewerUtils::GetIconForRow(DataStorage, TargetRow))
 					.ToolTipText(Binder.BindData(&FObjectOverrideColumn::OverriddenState, [](const EOverriddenState& OverriddenState)
 					{
 						return UE::OutlinerLabelWidget::Local::GetOverrideTooltip(OverriddenState);
@@ -171,16 +172,17 @@ TSharedPtr<SWidget> FOutlinerLabelWidgetConstructor::CreateWidget(ITypedElementD
 
 TSharedRef<SWidget> FOutlinerLabelWidgetConstructor::CreateLabel(ITypedElementDataStorageInterface* DataStorage,
 	ITypedElementDataStorageUiInterface* DataStorageUi, RowHandle TargetRow, RowHandle WidgetRow,
-	const TypedElementDataStorage::FMetaDataView& Arguments)
+	const UE::Editor::DataStorage::FMetaDataView& Arguments)
 {
+	using namespace UE::Editor::DataStorage;
 	TSharedRef<SWidget> Result = SNullWidget::NullWidget;
 	
-	const bool* IsEditable = Arguments.FindForColumn<FTypedElementLabelColumn>(TypedElementDataStorage::IsEditableName).TryGetExact<bool>();
+	const bool* IsEditable = Arguments.FindForColumn<FTypedElementLabelColumn>(IsEditableName).TryGetExact<bool>();
 
 	if (IsEditable && *IsEditable)
 	{
-		UE::EditorDataStorage::FAttributeBinder TargetRowBinder(TargetRow, DataStorage);
-		UE::EditorDataStorage::FAttributeBinder WidgetRowBinder(WidgetRow, DataStorage);
+		FAttributeBinder TargetRowBinder(TargetRow, DataStorage);
+		FAttributeBinder WidgetRowBinder(WidgetRow, DataStorage);
 			
 		TSharedPtr<SInlineEditableTextBlock> TextBlock = SNew(SInlineEditableTextBlock)
 			.OnTextCommitted_Lambda(
@@ -214,7 +216,7 @@ TSharedRef<SWidget> FOutlinerLabelWidgetConstructor::CreateLabel(ITypedElementDa
 	}
 	else
 	{
-		UE::EditorDataStorage::FAttributeBinder TargetRowBinder(TargetRow, DataStorage);
+		FAttributeBinder TargetRowBinder(TargetRow, DataStorage);
 		
 		TSharedPtr<STextBlock> TextBlock = SNew(STextBlock)
 			.IsEnabled(false)
diff --git a/Engine/Plugins/Experimental/EditorDataStorageFeatures/Source/TedsOutliner/Private/Widgets/OutlinerLabelWidget.h b/Engine/Plugins/Experimental/EditorDataStorageFeatures/Source/TedsOutliner/Private/Widgets/OutlinerLabelWidget.h
--- a/Engine/Plugins/Experimental/EditorDataStorageFeatures/Source/TedsOutliner/Private/Widgets/OutlinerLabelWidget.h
+++ b/Engine/Plugins/Experimental/EditorDataStorageFeatures/Source/TedsOutliner/Private/Widgets/OutlinerLabelWidget.h
@@ -37,13 +37,13 @@ public:
 	TEDSOUTLINER_API virtual TSharedPtr<SWidget> CreateWidget(
 		ITypedElementDataStorageInterface* DataStorage,
 		ITypedElementDataStorageUiInterface* DataStorageUi,
-		RowHandle TargetRow,
-		RowHandle WidgetRow, 
-		const TypedElementDataStorage::FMetaDataView& Arguments) override;
+		UE::Editor::DataStorage::RowHandle TargetRow,
+		UE::Editor::DataStorage::RowHandle WidgetRow, 
+		const UE::Editor::DataStorage::FMetaDataView& Arguments) override;
 
 protected:
 
 	TSharedRef<SWidget> CreateLabel(ITypedElementDataStorageInterface* DataStorage,
 	ITypedElementDataStorageUiInterface* DataStorageUi, RowHandle TargetRow, RowHandle WidgetRow,
-	const TypedElementDataStorage::FMetaDataView& Arguments);
+	const UE::Editor::DataStorage::FMetaDataView& Arguments);
 };
\ No newline at end of file
diff --git a/Engine/Plugins/Experimental/EditorDataStorageFeatures/Source/TedsOutliner/Private/Widgets/TransformHeadsUpWidget.cpp b/Engine/Plugins/Experimental/EditorDataStorageFeatures/Source/TedsOutliner/Private/Widgets/TransformHeadsUpWidget.cpp
--- a/Engine/Plugins/Experimental/EditorDataStorageFeatures/Source/TedsOutliner/Private/Widgets/TransformHeadsUpWidget.cpp
+++ b/Engine/Plugins/Experimental/EditorDataStorageFeatures/Source/TedsOutliner/Private/Widgets/TransformHeadsUpWidget.cpp
@@ -212,7 +212,7 @@ TConstArrayView<const UScriptStruct*> FTransformHeadsUpWidgetConstructor::GetAdd
 	return Columns;
 }
 
-TSharedPtr<SWidget> FTransformHeadsUpWidgetConstructor::CreateWidget(const TypedElementDataStorage::FMetaDataView& Arguments)
+TSharedPtr<SWidget> FTransformHeadsUpWidgetConstructor::CreateWidget(const UE::Editor::DataStorage::FMetaDataView& Arguments)
 {
 	return SNew(STransformQuickDisplay);
 }
diff --git a/Engine/Plugins/Experimental/EditorDataStorageFeatures/Source/TedsOutliner/Private/Widgets/TypeInfoWidget.cpp b/Engine/Plugins/Experimental/EditorDataStorageFeatures/Source/TedsOutliner/Private/Widgets/TypeInfoWidget.cpp
--- a/Engine/Plugins/Experimental/EditorDataStorageFeatures/Source/TedsOutliner/Private/Widgets/TypeInfoWidget.cpp
+++ b/Engine/Plugins/Experimental/EditorDataStorageFeatures/Source/TedsOutliner/Private/Widgets/TypeInfoWidget.cpp
@@ -35,12 +35,12 @@ FTypeInfoWidgetConstructor::FTypeInfoWidgetConstructor(const UScriptStruct* InTy
 }
 
 TSharedPtr<SWidget> FTypeInfoWidgetConstructor::CreateWidget(
-	const TypedElementDataStorage::FMetaDataView& Arguments	)
+	const UE::Editor::DataStorage::FMetaDataView& Arguments	)
 {
 	bUseIcon = false;
 	
 	// Check if the caller provided metadata to use an icon widget
-	TypedElementDataStorage::FMetaDataEntryView MetaDataEntryView = Arguments.FindGeneric("TypeInfoWidget_bUseIcon");
+	UE::Editor::DataStorage::FMetaDataEntryView MetaDataEntryView = Arguments.FindGeneric("TypeInfoWidget_bUseIcon");
 	if(MetaDataEntryView.IsSet())
 	{
 		check(MetaDataEntryView.IsType<bool>());
diff --git a/Engine/Plugins/Experimental/EditorDataStorageFeatures/Source/TedsOutliner/Public/Compatibility/SceneOutlinerTedsBridge.h b/Engine/Plugins/Experimental/EditorDataStorageFeatures/Source/TedsOutliner/Public/Compatibility/SceneOutlinerTedsBridge.h
--- a/Engine/Plugins/Experimental/EditorDataStorageFeatures/Source/TedsOutliner/Public/Compatibility/SceneOutlinerTedsBridge.h
+++ b/Engine/Plugins/Experimental/EditorDataStorageFeatures/Source/TedsOutliner/Public/Compatibility/SceneOutlinerTedsBridge.h
@@ -17,7 +17,7 @@ class ITypedElementDataStorageInterface;
 class ITypedElementDataStorageUiInterface;
 class ITypedElementDataStorageCompatibilityInterface;
 
-DECLARE_DELEGATE_RetVal_OneParam(FSceneOutlinerTreeItemID, FTreeItemIDDealiaser, TypedElementDataStorage::RowHandle);
+DECLARE_DELEGATE_RetVal_OneParam(FSceneOutlinerTreeItemID, FTreeItemIDDealiaser, UE::Editor::DataStorage::RowHandle);
 
 class FSceneOutlinerTedsBridge;
 
diff --git a/Engine/Plugins/Experimental/EditorDataStorageFeatures/Source/TedsOutliner/Public/TedsOutlinerFilter.h b/Engine/Plugins/Experimental/EditorDataStorageFeatures/Source/TedsOutliner/Public/TedsOutlinerFilter.h
--- a/Engine/Plugins/Experimental/EditorDataStorageFeatures/Source/TedsOutliner/Public/TedsOutlinerFilter.h
+++ b/Engine/Plugins/Experimental/EditorDataStorageFeatures/Source/TedsOutliner/Public/TedsOutlinerFilter.h
@@ -14,7 +14,7 @@ class FTedsOutlinerFilter : public FFilterBase<SceneOutliner::FilterBarType>
 {
 public:
 	FTedsOutlinerFilter(const FName& InFilterName, const FText& InFilterDisplayName, TSharedPtr<FFilterCategory> InCategory,
-		TSharedRef<FTedsOutlinerImpl> InTedsOutlinerImpl, const TypedElementDataStorage::FQueryDescription& InFilterQuery);
+		TSharedRef<FTedsOutlinerImpl> InTedsOutlinerImpl, const UE::Editor::DataStorage::FQueryDescription& InFilterQuery);
 
 	/** Returns the system name for this filter */
 	virtual FString GetName() const override;
@@ -57,5 +57,5 @@ protected:
 	FText FilterDisplayName;
 
 	TSharedRef<FTedsOutlinerImpl> TedsOutlinerImpl;
-	const TypedElementDataStorage::FQueryDescription FilterQuery;
+	const UE::Editor::DataStorage::FQueryDescription FilterQuery;
 };
\ No newline at end of file
