[Teds] Fix up remaining uses of UE::EditorDataStorage and move outliner to UE::Editor::Outliner from UE::EditorDataStorage

#rnx

diff --git a/Engine/Plugins/Experimental/Animation/HierarchyTable/Source/Editor/Private/HierarchyTableEditorToolkit.cpp b/Engine/Plugins/Experimental/Animation/HierarchyTable/Source/Editor/Private/HierarchyTableEditorToolkit.cpp
--- a/Engine/Plugins/Experimental/Animation/HierarchyTable/Source/Editor/Private/HierarchyTableEditorToolkit.cpp
+++ b/Engine/Plugins/Experimental/Animation/HierarchyTable/Source/Editor/Private/HierarchyTableEditorToolkit.cpp
@@ -113,6 +113,7 @@ TSharedRef<SWidget> FHierarchyTableEditorToolkit::CreateTedsOutliner()
 	}
 
 	using namespace UE::Editor::DataStorage::Queries;
+	using namespace UE::Editor::Outliner;
 
 	if (!ensure(HierarchyTable->TableType))
 	{
diff --git a/Engine/Plugins/Experimental/EditorDataStorageFeatures/Source/TedsDebugger/Private/QueryEditor/TedsQueryEditorModel.h b/Engine/Plugins/Experimental/EditorDataStorageFeatures/Source/TedsDebugger/Private/QueryEditor/TedsQueryEditorModel.h
--- a/Engine/Plugins/Experimental/EditorDataStorageFeatures/Source/TedsDebugger/Private/QueryEditor/TedsQueryEditorModel.h
+++ b/Engine/Plugins/Experimental/EditorDataStorageFeatures/Source/TedsDebugger/Private/QueryEditor/TedsQueryEditorModel.h
@@ -6,9 +6,7 @@
 
 class ITypedElementDataStorageInterface;
 
-namespace UE::Editor::DataStorage
-{
-namespace Debug::QueryEditor
+namespace UE::Editor::DataStorage::Debug::QueryEditor
 {
 	class FTedsQueryEditorModel;
 
@@ -118,5 +116,4 @@ namespace Debug::QueryEditor
 		uint64 GeneratedVersion = 0;
 		QueryHandle QueryHandle;
 	};
-} // namespace Debug::QueryEditor
-} // namespace UE::Editor::DataStorage
+} // namespace UE::Editor::DataStorage::Debug::QueryEditor
diff --git a/Engine/Plugins/Experimental/EditorDataStorageFeatures/Source/TedsDebugger/Private/Widgets/QueryEditor/TedsConditionSelectionCollectionWidget.cpp b/Engine/Plugins/Experimental/EditorDataStorageFeatures/Source/TedsDebugger/Private/Widgets/QueryEditor/TedsConditionSelectionCollectionWidget.cpp
--- a/Engine/Plugins/Experimental/EditorDataStorageFeatures/Source/TedsDebugger/Private/Widgets/QueryEditor/TedsConditionSelectionCollectionWidget.cpp
+++ b/Engine/Plugins/Experimental/EditorDataStorageFeatures/Source/TedsDebugger/Private/Widgets/QueryEditor/TedsConditionSelectionCollectionWidget.cpp
@@ -1,10 +1,10 @@
 // Copyright Epic Games, Inc. All Rights Reserved.
 #include "Widgets/QueryEditor/TedsConditionSelectionCollectionWidget.h"
 
-namespace UE::EditorDataStorage::Debug::QueryEditor
+namespace UE::Editor::DataStorage::Debug::QueryEditor
 {
 	void SConditionSelectionCollectionWidget::Construct(const FArguments& InArgs, FTedsQueryEditorModel& InModel)
 	{
 		Model = &InModel;
 	}
-}
\ No newline at end of file
+} // namespace UE::Editor::DataStorage::Debug::QueryEditor
diff --git a/Engine/Plugins/Experimental/EditorDataStorageFeatures/Source/TedsDebugger/Private/Widgets/QueryEditor/TedsConditionSelectionCollectionWidget.h b/Engine/Plugins/Experimental/EditorDataStorageFeatures/Source/TedsDebugger/Private/Widgets/QueryEditor/TedsConditionSelectionCollectionWidget.h
--- a/Engine/Plugins/Experimental/EditorDataStorageFeatures/Source/TedsDebugger/Private/Widgets/QueryEditor/TedsConditionSelectionCollectionWidget.h
+++ b/Engine/Plugins/Experimental/EditorDataStorageFeatures/Source/TedsDebugger/Private/Widgets/QueryEditor/TedsConditionSelectionCollectionWidget.h
@@ -3,7 +3,7 @@
 #include "Widgets/SCompoundWidget.h"
 #include "Templates/SharedPointerFwd.h"
 
-namespace UE::EditorDataStorage::Debug::QueryEditor
+namespace UE::Editor::DataStorage::Debug::QueryEditor
 {
 	class FTedsQueryEditorModel;
 	
@@ -21,4 +21,4 @@ namespace UE::EditorDataStorage::Debug::QueryEditor
 	private:
 		FTedsQueryEditorModel* Model = nullptr;
 	};
-}
+} // namespace UE::Editor::DataStorage::Debug::QueryEditor
diff --git a/Engine/Plugins/Experimental/EditorDataStorageFeatures/Source/TedsOutliner/Private/Compatibility/SceneOutlinerRowHandleColumn.cpp b/Engine/Plugins/Experimental/EditorDataStorageFeatures/Source/TedsOutliner/Private/Compatibility/SceneOutlinerRowHandleColumn.cpp
--- a/Engine/Plugins/Experimental/EditorDataStorageFeatures/Source/TedsOutliner/Private/Compatibility/SceneOutlinerRowHandleColumn.cpp
+++ b/Engine/Plugins/Experimental/EditorDataStorageFeatures/Source/TedsOutliner/Private/Compatibility/SceneOutlinerRowHandleColumn.cpp
@@ -51,7 +51,7 @@ const TSharedRef<SWidget> FSceneOutlinerRowHandleColumn::ConstructRowWidget(FSce
 	auto SceneOutliner = WeakSceneOutliner.Pin();
 	check(SceneOutliner.IsValid());
 
-	if (const FTedsOutlinerTreeItem* OutlinerTreeItem = TreeItem->CastTo<FTedsOutlinerTreeItem>())
+	if (const UE::Editor::Outliner::FTedsOutlinerTreeItem* OutlinerTreeItem = TreeItem->CastTo<UE::Editor::Outliner::FTedsOutlinerTreeItem>())
 	{
 		const UE::Editor::DataStorage::RowHandle RowHandle = OutlinerTreeItem->GetRowHandle();
 
@@ -65,7 +65,7 @@ const TSharedRef<SWidget> FSceneOutlinerRowHandleColumn::ConstructRowWidget(FSce
 
 void FSceneOutlinerRowHandleColumn::PopulateSearchStrings(const ISceneOutlinerTreeItem& Item, TArray<FString>& OutSearchStrings) const
 {
-	if (const FTedsOutlinerTreeItem* OutlinerTreeItem = Item.CastTo<FTedsOutlinerTreeItem>())
+	if (const UE::Editor::Outliner::FTedsOutlinerTreeItem* OutlinerTreeItem = Item.CastTo<UE::Editor::Outliner::FTedsOutlinerTreeItem>())
 	{
 		OutSearchStrings.Add(LexToString<FString>(OutlinerTreeItem->GetRowHandle()));
 	}
@@ -78,7 +78,7 @@ void FSceneOutlinerRowHandleColumn::SortItems(TArray<FSceneOutlinerTreeItemPtr>&
 		/** Sort by type first */
 		.Primary([this](const ISceneOutlinerTreeItem& Item)
 		{
-			if (const FTedsOutlinerTreeItem* OutlinerTreeItem = Item.CastTo<FTedsOutlinerTreeItem>())
+			if (const UE::Editor::Outliner::FTedsOutlinerTreeItem* OutlinerTreeItem = Item.CastTo<UE::Editor::Outliner::FTedsOutlinerTreeItem>())
 			{
 				return OutlinerTreeItem->GetRowHandle();
 			}
diff --git a/Engine/Plugins/Experimental/EditorDataStorageFeatures/Source/TedsOutliner/Private/Compatibility/SceneOutlinerTedsBridge.cpp b/Engine/Plugins/Experimental/EditorDataStorageFeatures/Source/TedsOutliner/Private/Compatibility/SceneOutlinerTedsBridge.cpp
--- a/Engine/Plugins/Experimental/EditorDataStorageFeatures/Source/TedsOutliner/Private/Compatibility/SceneOutlinerTedsBridge.cpp
+++ b/Engine/Plugins/Experimental/EditorDataStorageFeatures/Source/TedsOutliner/Private/Compatibility/SceneOutlinerTedsBridge.cpp
@@ -262,6 +262,7 @@ public:
 	const TSharedRef<SWidget> ConstructRowWidget(FSceneOutlinerTreeItemRef TreeItem, const STableRow<FSceneOutlinerTreeItemPtr>& Row) override
 	{
 		using namespace UE::Editor::DataStorage;
+		using namespace UE::Editor::Outliner;
 		
 		RowHandle RowHandle = InvalidRowHandle;
 
diff --git a/Engine/Plugins/Experimental/EditorDataStorageFeatures/Source/TedsOutliner/Private/TedsOutlinerFilter.cpp b/Engine/Plugins/Experimental/EditorDataStorageFeatures/Source/TedsOutliner/Private/TedsOutlinerFilter.cpp
--- a/Engine/Plugins/Experimental/EditorDataStorageFeatures/Source/TedsOutliner/Private/TedsOutlinerFilter.cpp
+++ b/Engine/Plugins/Experimental/EditorDataStorageFeatures/Source/TedsOutliner/Private/TedsOutlinerFilter.cpp
@@ -4,6 +4,8 @@
 
 #include "TedsOutlinerImpl.h"
 
+namespace UE::Editor::Outliner
+{
 FTedsOutlinerFilter::FTedsOutlinerFilter(const FName& InFilterName, const FText& InFilterDisplayName,
 	TSharedPtr<FFilterCategory> InCategory, TSharedRef<FTedsOutlinerImpl> InTedsOutlinerImpl,
 	const UE::Editor::DataStorage::FQueryDescription& InFilterQuery)
@@ -85,3 +87,4 @@ bool FTedsOutlinerFilter::PassesFilter(SceneOutliner::FilterBarType InItem) cons
 	// The filter is applied through a TEDS query and this is just a dummy to activate it, so we can simply return true otherwise
 	return false;
 }
+} // namespace UE::Editor::Outliner
diff --git a/Engine/Plugins/Experimental/EditorDataStorageFeatures/Source/TedsOutliner/Private/TedsOutlinerHierarchy.cpp b/Engine/Plugins/Experimental/EditorDataStorageFeatures/Source/TedsOutliner/Private/TedsOutlinerHierarchy.cpp
--- a/Engine/Plugins/Experimental/EditorDataStorageFeatures/Source/TedsOutliner/Private/TedsOutlinerHierarchy.cpp
+++ b/Engine/Plugins/Experimental/EditorDataStorageFeatures/Source/TedsOutliner/Private/TedsOutlinerHierarchy.cpp
@@ -14,6 +14,8 @@
 #include "Elements/Columns/TypedElementMiscColumns.h"
 #include "Elements/Columns/TypedElementTypeInfoColumns.h"
 
+namespace UE::Editor::Outliner
+{
 FTedsOutlinerHierarchy::FTedsOutlinerHierarchy(FTedsOutlinerMode* InMode,
 	const TSharedRef<FTedsOutlinerImpl>& InTedsOutlinerImpl)
 	: ISceneOutlinerHierarchy(InMode)
@@ -72,4 +74,5 @@ FSceneOutlinerTreeItemPtr FTedsOutlinerHierarchy::FindOrCreateParentItem(const I
 	}
 	
 	return nullptr;
-}
\ No newline at end of file
+}
+} // namespace UE::Editor::Outliner
\ No newline at end of file
diff --git a/Engine/Plugins/Experimental/EditorDataStorageFeatures/Source/TedsOutliner/Private/TedsOutlinerHierarchy.h b/Engine/Plugins/Experimental/EditorDataStorageFeatures/Source/TedsOutliner/Private/TedsOutlinerHierarchy.h
--- a/Engine/Plugins/Experimental/EditorDataStorageFeatures/Source/TedsOutliner/Private/TedsOutlinerHierarchy.h
+++ b/Engine/Plugins/Experimental/EditorDataStorageFeatures/Source/TedsOutliner/Private/TedsOutlinerHierarchy.h
@@ -5,6 +5,8 @@
 #include "ISceneOutlinerHierarchy.h"
 #include "Elements/Interfaces/TypedElementDataStorageInterface.h"
 
+namespace UE::Editor::Outliner
+{
 class FTedsOutlinerMode;
 class FTedsOutlinerImpl;
 
@@ -36,3 +38,4 @@ protected:
 	// Delegate called by TedsOutlinerImpl when the hierarchy changes
 	FDelegateHandle HierarchyChangedHandle;
 };
+} // namespace UE::Editor::Outliner
diff --git a/Engine/Plugins/Experimental/EditorDataStorageFeatures/Source/TedsOutliner/Private/TedsOutlinerImpl.cpp b/Engine/Plugins/Experimental/EditorDataStorageFeatures/Source/TedsOutliner/Private/TedsOutlinerImpl.cpp
--- a/Engine/Plugins/Experimental/EditorDataStorageFeatures/Source/TedsOutliner/Private/TedsOutlinerImpl.cpp
+++ b/Engine/Plugins/Experimental/EditorDataStorageFeatures/Source/TedsOutliner/Private/TedsOutlinerImpl.cpp
@@ -25,6 +25,8 @@
 
 #define LOCTEXT_NAMESPACE "TedsOutliner"
 
+namespace UE::Editor::Outliner
+{
 FTedsOutlinerImpl::FTedsOutlinerImpl(const FTedsOutlinerParams& InParams, ISceneOutlinerMode* InMode)
 	: CreationParams(InParams)
 	, CellWidgetPurposes(InParams.CellWidgetPurposes)
@@ -52,7 +54,7 @@ void FTedsOutlinerImpl::CreateLabelWidgetConstructors()
 
 	// Create and assign the widget constructors for the label column
 	
-	auto CreateWidgetConstructorForQuery = [this](UE::Editor::DataStorage::FQueryDescription InQueryDescription) -> TSharedPtr<FTypedElementWidgetConstructor>
+	auto CreateWidgetConstructorForQuery = [this](DataStorage::FQueryDescription InQueryDescription) -> TSharedPtr<FTypedElementWidgetConstructor>
 	{
 		// Create the Widget Constructor for the item label column
 		using MatchApproach = ITypedElementDataStorageUiInterface::EMatchApproach;
@@ -94,7 +96,7 @@ void FTedsOutlinerImpl::CreateLabelWidgetConstructors()
 		return OutWidgetConstructorPtr;
 	};
 	
-	UE::Editor::DataStorage::QueryHandle LabelColumnQueryHandle = Storage->RegisterQuery(
+	DataStorage::QueryHandle LabelColumnQueryHandle = Storage->RegisterQuery(
 		Select()
 			.ReadWrite<FTypedElementLabelColumn, FTypedElementClassTypeInfoColumn>()
 		.Compile());
@@ -115,8 +117,8 @@ void FTedsOutlinerImpl::CreateFilterQueries()
 		TSharedRef<FFilterCategory> TedsColumnFilterCategory = MakeShared<FFilterCategory>(LOCTEXT("TedsColumnFilters", "TEDS Columns"), LOCTEXT("TedsColumnFiltersTooltip", "Filter by TEDS columns"));
 		TSharedRef<FFilterCategory> TedsTagFilterCategory = MakeShared<FFilterCategory>(LOCTEXT("TedsTagFilters", "TEDS Tags"), LOCTEXT("TedsTagFiltersTooltip", "Filter by TEDS Tags"));
 
-		const UStruct* TedsColumn = UE::Editor::DataStorage::FColumn::StaticStruct();
-		const UStruct* TedsTag = UE::Editor::DataStorage::FTag::StaticStruct();
+		const UStruct* TedsColumn = DataStorage::FColumn::StaticStruct();
+		const UStruct* TedsTag = DataStorage::FTag::StaticStruct();
 
 		// Grab all UStruct types to see if they derive from FColumn or FTag
 		ForEachObjectOfClass(UScriptStruct::StaticClass(), [&](UObject* Obj)
@@ -126,7 +128,7 @@ void FTedsOutlinerImpl::CreateFilterQueries()
 				if (Struct->IsChildOf(TedsColumn) || Struct->IsChildOf(TedsTag))
 				{
 					// Create a query description to filter for this tag/column
-					UE::Editor::DataStorage::FQueryDescription FilterQueryDesc =
+					DataStorage::FQueryDescription FilterQueryDesc =
 					Select()
 					.Where()
 						.All(Struct)
@@ -144,7 +146,7 @@ void FTedsOutlinerImpl::CreateFilterQueries()
 	// Custom filters input by the user
 	TSharedRef<FFilterCategory> CustomFiltersCategory = MakeShared<FFilterCategory>(LOCTEXT("TedsFilters", "TEDS Custom Filters"), LOCTEXT("TedsFiltersTooltip", "Filter by custom TEDS queries"));
 
-	for(const TPair<FName, const UE::Editor::DataStorage::FQueryDescription>& FilterQuery : CreationParams.FilterQueries)
+	for(const TPair<FName, const DataStorage::FQueryDescription>& FilterQuery : CreationParams.FilterQueries)
 	{
 		// TEDS-Outliner TODO: Custom filters need a localizable display name instead of using the FName, but we need to change how they are added first
 		// to see if it can be consolidated with the SFilterBar API
@@ -177,7 +179,7 @@ FTedsOutlinerImpl::~FTedsOutlinerImpl()
 
 	// This is done outside of unregister queries because that is called when the internal query changes (i.e filter) but we don't want to unregister the
 	// label widget queries until shutdown
-	for(const TPair<UE::Editor::DataStorage::QueryHandle, TSharedPtr<FTypedElementWidgetConstructor>>& QueryConstructorPair : QueryToWidgetConstructorMap)
+	for(const TPair<DataStorage::QueryHandle, TSharedPtr<FTypedElementWidgetConstructor>>& QueryConstructorPair : QueryToWidgetConstructorMap)
 	{
 		Storage->UnregisterQuery(QueryConstructorPair.Key);
 	}
@@ -190,7 +192,7 @@ FTedsOutlinerImpl::FIsItemCompatible& FTedsOutlinerImpl::IsItemCompatible()
 	return IsItemCompatibleWithTeds;
 }
 
-void FTedsOutlinerImpl::SetSelection(const TArray<UE::Editor::DataStorage::RowHandle>& InSelectedRows)
+void FTedsOutlinerImpl::SetSelection(const TArray<DataStorage::RowHandle>& InSelectedRows)
 {
 	if (!SelectionSetName.IsSet())
 	{
@@ -199,33 +201,33 @@ void FTedsOutlinerImpl::SetSelection(const TArray<UE::Editor::DataStorage::RowHa
 	
 	ClearSelection();
 
-	for(UE::Editor::DataStorage::RowHandle Row : InSelectedRows)
+	for(DataStorage::RowHandle Row : InSelectedRows)
 	{
 		Storage->AddColumn(Row, FTypedElementSelectionColumn{ .SelectionSet = SelectionSetName.GetValue() });
 	}
 }
 
-TSharedRef<SWidget> FTedsOutlinerImpl::CreateLabelWidgetForItem(UE::Editor::DataStorage::RowHandle InRowHandle, const STableRow<FSceneOutlinerTreeItemPtr>& InRow) const
+TSharedRef<SWidget> FTedsOutlinerImpl::CreateLabelWidgetForItem(DataStorage::RowHandle InRowHandle, const STableRow<FSceneOutlinerTreeItemPtr>& InRow) const
 {
-	auto CreateWidgetForQuery = [InRowHandle, this, &InRow](const TPair<UE::Editor::DataStorage::QueryHandle, TSharedPtr<FTypedElementWidgetConstructor>>& QueryConstructorPair) -> TSharedPtr<SWidget>
+	auto CreateWidgetForQuery = [InRowHandle, this, &InRow](const TPair<DataStorage::QueryHandle, TSharedPtr<FTypedElementWidgetConstructor>>& QueryConstructorPair) -> TSharedPtr<SWidget>
 	{
-		UE::Editor::DataStorage::FQueryDescription QueryDescription = Storage->GetQueryDescription(QueryConstructorPair.Key);
+		DataStorage::FQueryDescription QueryDescription = Storage->GetQueryDescription(QueryConstructorPair.Key);
 		
 		// Create a generic metadata view for the Type Widget
-		UE::Editor::DataStorage::FMetaData QueryWideMetaData;
+		DataStorage::FMetaData QueryWideMetaData;
 		QueryWideMetaData.AddImmutableData("TypeInfoWidget_bUseIcon", true);
-		UE::Editor::DataStorage::FGenericMetaDataView GenericMetaDataView(QueryWideMetaData);
+		DataStorage::FGenericMetaDataView GenericMetaDataView(QueryWideMetaData);
 
 		// Create metadata for the query itself
-		UE::Editor::DataStorage::FQueryMetaDataView QueryMetaDataView(QueryDescription);
+		DataStorage::FQueryMetaDataView QueryMetaDataView(QueryDescription);
 
 		// Combine the two metadata
-		UE::Editor::DataStorage::FComboMetaDataView MetaDataArgs(GenericMetaDataView, QueryMetaDataView);
+		DataStorage::FComboMetaDataView MetaDataArgs(GenericMetaDataView, QueryMetaDataView);
 
 		TArray<TWeakObjectPtr<const UScriptStruct>> ColumnTypes(QueryDescription.SelectionTypes);
 		TSharedPtr<FTypedElementWidgetConstructor> CellWidgetConstructor = QueryConstructorPair.Value;
 
-		UE::Editor::DataStorage::RowHandle UiRowHandle = Storage->AddRow(Storage->FindTable(UE::Editor::DataStorage::TableViewerUtils::GetWidgetTableName()));
+		DataStorage::RowHandle UiRowHandle = Storage->AddRow(Storage->FindTable(DataStorage::TableViewerUtils::GetWidgetTableName()));
 
 		if (FTypedElementRowReferenceColumn* RowReference = Storage->GetColumn<FTypedElementRowReferenceColumn>(UiRowHandle))
 		{
@@ -246,7 +248,7 @@ TSharedRef<SWidget> FTedsOutlinerImpl::CreateLabelWidgetForItem(UE::Editor::Data
 	
 	TSharedRef<SHorizontalBox> CombinedWidget = SNew(SHorizontalBox);
 
-	for(const TPair<UE::Editor::DataStorage::QueryHandle, TSharedPtr<FTypedElementWidgetConstructor>>& QueryConstructorPair : QueryToWidgetConstructorMap)
+	for(const TPair<DataStorage::QueryHandle, TSharedPtr<FTypedElementWidgetConstructor>>& QueryConstructorPair : QueryToWidgetConstructorMap)
 	{
 		if (TSharedPtr<SWidget> WidgetForQuery = CreateWidgetForQuery(QueryConstructorPair))
 		{
@@ -264,13 +266,13 @@ TSharedRef<SWidget> FTedsOutlinerImpl::CreateLabelWidgetForItem(UE::Editor::Data
 	return CombinedWidget;
 }
 
-void FTedsOutlinerImpl::AppendQuery(UE::Editor::DataStorage::FQueryDescription& Query1, const UE::Editor::DataStorage::FQueryDescription& Query2)
+void FTedsOutlinerImpl::AppendQuery(DataStorage::FQueryDescription& Query1, const DataStorage::FQueryDescription& Query2)
 {
 	// TEDS-Outliner TODO: We simply discard duplicate types for now but we probably want a more robust system to detect duplicates and conflicting conditions
 	for(int32 i = 0; i < Query2.ConditionOperators.Num(); ++i)
 	{
 		// Make sure we don't add duplicate conditions
-		UE::Editor::DataStorage::FQueryDescription::FOperator* FoundCondition = Query1.ConditionOperators.FindByPredicate([&Query2, i](const UE::Editor::DataStorage::FQueryDescription::FOperator& Op)
+		DataStorage::FQueryDescription::FOperator* FoundCondition = Query1.ConditionOperators.FindByPredicate([&Query2, i](const DataStorage::FQueryDescription::FOperator& Op)
 		{
 			return Op.Type == Query2.ConditionOperators[i].Type;
 		});
@@ -289,7 +291,7 @@ void FTedsOutlinerImpl::AppendQuery(UE::Editor::DataStorage::FQueryDescription&
 	}
 }
 
-void FTedsOutlinerImpl::AddExternalQuery(FName QueryName, const UE::Editor::DataStorage::FQueryDescription& InQueryDescription)
+void FTedsOutlinerImpl::AddExternalQuery(FName QueryName, const DataStorage::FQueryDescription& InQueryDescription)
 {
 	ExternalQueries.Emplace(QueryName, InQueryDescription);
 
@@ -301,15 +303,15 @@ void FTedsOutlinerImpl::RemoveExternalQuery(FName QueryName)
 	ExternalQueries.Remove(QueryName);
 }
 
-void FTedsOutlinerImpl::AppendExternalQueries(UE::Editor::DataStorage::FQueryDescription& OutQuery)
+void FTedsOutlinerImpl::AppendExternalQueries(DataStorage::FQueryDescription& OutQuery)
 {
-	for(const TPair<FName, UE::Editor::DataStorage::FQueryDescription>& ExternalQuery : ExternalQueries)
+	for(const TPair<FName, DataStorage::FQueryDescription>& ExternalQuery : ExternalQueries)
 	{
 		AppendQuery(OutQuery, ExternalQuery.Value);
 	}
 }
 
-bool FTedsOutlinerImpl::HasItemParentChanged(UE::Editor::DataStorage::RowHandle InRowHandle, UE::Editor::DataStorage::RowHandle ParentRowHandle) const
+bool FTedsOutlinerImpl::HasItemParentChanged(DataStorage::RowHandle InRowHandle, DataStorage::RowHandle ParentRowHandle) const
 {
 	const FSceneOutlinerTreeItemPtr Item = SceneOutliner->GetTreeItem(InRowHandle, true);
 
@@ -338,7 +340,7 @@ bool FTedsOutlinerImpl::HasItemParentChanged(UE::Editor::DataStorage::RowHandle
 	return false;
 }
 
-bool FTedsOutlinerImpl::CanDisplayRow(UE::Editor::DataStorage::RowHandle ItemRowHandle) const
+bool FTedsOutlinerImpl::CanDisplayRow(DataStorage::RowHandle ItemRowHandle) const
 {
 	/*
 	 * Don't display widgets that are created for rows in this table viewer. Widgets are only created for rows that are currently visible, so if we
@@ -481,12 +483,12 @@ void FTedsOutlinerImpl::CreateChildren(const FSceneOutlinerTreeItemPtr& Item, TA
 	}
 }
 
-UE::Editor::DataStorage::RowHandle FTedsOutlinerImpl::GetParentRow(UE::Editor::DataStorage::RowHandle InRowHandle)
+DataStorage::RowHandle FTedsOutlinerImpl::GetParentRow(DataStorage::RowHandle InRowHandle)
 {
 	// No parent if there is no hierarchy data specified
 	if (!HierarchyData.IsSet())
 	{
-		return UE::Editor::DataStorage::InvalidRowHandle;
+		return DataStorage::InvalidRowHandle;
 	}
 	
 	// If this entity does not have a parent entity, return InvalidRowHandle
@@ -494,26 +496,26 @@ UE::Editor::DataStorage::RowHandle FTedsOutlinerImpl::GetParentRow(UE::Editor::D
 	
 	if (!ParentColumnData)
 	{
-		return UE::Editor::DataStorage::InvalidRowHandle;
+		return DataStorage::InvalidRowHandle;
 	}
 
 	// If the parent is invalid for some reason, return InvalidRowHandle
-	const UE::Editor::DataStorage::RowHandle ParentRowHandle = HierarchyData.GetValue().GetParent.Execute(ParentColumnData);
+	const DataStorage::RowHandle ParentRowHandle = HierarchyData.GetValue().GetParent.Execute(ParentColumnData);
 	
 	if (!Storage->IsRowAvailable(ParentRowHandle))
 	{
-		return UE::Editor::DataStorage::InvalidRowHandle;
+		return DataStorage::InvalidRowHandle;
 	}
 	
 	if (!CanDisplayRow(ParentRowHandle))
 	{
-		return UE::Editor::DataStorage::InvalidRowHandle;
+		return DataStorage::InvalidRowHandle;
 	}
 
 	return ParentRowHandle;
 }
 
-void FTedsOutlinerImpl::OnItemAdded(UE::Editor::DataStorage::RowHandle ItemRowHandle)
+void FTedsOutlinerImpl::OnItemAdded(DataStorage::RowHandle ItemRowHandle)
 {
 	if (!CanDisplayRow(ItemRowHandle))
 	{
@@ -526,7 +528,7 @@ void FTedsOutlinerImpl::OnItemAdded(UE::Editor::DataStorage::RowHandle ItemRowHa
 	HierarchyChangedEvent.Broadcast(EventData);
 }
 
-void FTedsOutlinerImpl::OnItemRemoved(UE::Editor::DataStorage::RowHandle ItemRowHandle)
+void FTedsOutlinerImpl::OnItemRemoved(DataStorage::RowHandle ItemRowHandle)
 {
 	FSceneOutlinerHierarchyChangedData EventData;
 	EventData.Type = FSceneOutlinerHierarchyChangedData::Removed;
@@ -534,7 +536,7 @@ void FTedsOutlinerImpl::OnItemRemoved(UE::Editor::DataStorage::RowHandle ItemRow
 	HierarchyChangedEvent.Broadcast(EventData);
 }
 
-void FTedsOutlinerImpl::OnItemMoved(UE::Editor::DataStorage::RowHandle ItemRowHandle)
+void FTedsOutlinerImpl::OnItemMoved(DataStorage::RowHandle ItemRowHandle)
 {
 	if (!CanDisplayRow(ItemRowHandle))
 	{
@@ -569,7 +571,7 @@ void FTedsOutlinerImpl::RecompileQueries()
 		Select(
 			TEXT("Add Row to Outliner"),
 			FObserver::OnAdd<FTypedElementLabelColumn>().SetExecutionMode(EExecutionMode::GameThread),
-			[this](IQueryContext& Context, UE::Editor::DataStorage::RowHandle Row)
+			[this](IQueryContext& Context, DataStorage::RowHandle Row)
 			{
 				OnItemAdded(Row);
 			})
@@ -583,7 +585,7 @@ void FTedsOutlinerImpl::RecompileQueries()
 		Select(
 				TEXT("Remove Row from Outliner"),
 				FObserver::OnRemove<FTypedElementLabelColumn>().SetExecutionMode(EExecutionMode::GameThread),
-				[this](IQueryContext& Context, UE::Editor::DataStorage::RowHandle Row)
+				[this](IQueryContext& Context, DataStorage::RowHandle Row)
 				{
 					OnItemRemoved(Row);
 				})
@@ -612,9 +614,9 @@ void FTedsOutlinerImpl::RecompileQueries()
 			TEXT("Update item parent"),
 			FProcessor(EQueryTickPhase::DuringPhysics, Storage->GetQueryTickGroupName(EQueryTickGroups::Update))
 				.SetExecutionMode(EExecutionMode::GameThread),
-			[this](IQueryContext& Context, UE::Editor::DataStorage::RowHandle Row)
+			[this](IQueryContext& Context, DataStorage::RowHandle Row)
 			{
-				UE::Editor::DataStorage::RowHandle ParentRowHandle = InvalidRowHandle;
+				DataStorage::RowHandle ParentRowHandle = InvalidRowHandle;
 
 				if (const FTableRowParentColumn* ParentColumn = Context.GetColumn<FTableRowParentColumn>())
 				{
@@ -652,7 +654,7 @@ void FTedsOutlinerImpl::RecompileQueries()
 							Select(
 							TEXT("Row selected"),
 							FObserver::OnAdd<FTypedElementSelectionColumn>().SetExecutionMode(EExecutionMode::GameThread),
-							[this](IQueryContext& Context, UE::Editor::DataStorage::RowHandle Row)
+							[this](IQueryContext& Context, DataStorage::RowHandle Row)
 							{
 								bSelectionDirty = true;
 							})
@@ -666,7 +668,7 @@ void FTedsOutlinerImpl::RecompileQueries()
 							Select(
 							TEXT("Row deselected"),
 							FObserver::OnRemove<FTypedElementSelectionColumn>().SetExecutionMode(EExecutionMode::GameThread),
-							[this](IQueryContext& Context, UE::Editor::DataStorage::RowHandle Row)
+							[this](IQueryContext& Context, DataStorage::RowHandle Row)
 							{
 								bSelectionDirty = true;
 							})
@@ -782,5 +784,6 @@ const TOptional<FTedsOutlinerHierarchyData>& FTedsOutlinerImpl::GetHierarchyData
 {
 	return HierarchyData;
 }
+} // namespace UE::Editor::Outliner
 
 #undef LOCTEXT_NAMESPACE
diff --git a/Engine/Plugins/Experimental/EditorDataStorageFeatures/Source/TedsOutliner/Private/TedsOutlinerItem.cpp b/Engine/Plugins/Experimental/EditorDataStorageFeatures/Source/TedsOutliner/Private/TedsOutlinerItem.cpp
--- a/Engine/Plugins/Experimental/EditorDataStorageFeatures/Source/TedsOutliner/Private/TedsOutlinerItem.cpp
+++ b/Engine/Plugins/Experimental/EditorDataStorageFeatures/Source/TedsOutliner/Private/TedsOutlinerItem.cpp
@@ -8,9 +8,11 @@
 
 #define LOCTEXT_NAMESPACE "TedsOutliner"
 
+namespace UE::Editor::Outliner
+{
 const FSceneOutlinerTreeItemType FTedsOutlinerTreeItem::Type(&ISceneOutlinerTreeItem::Type);
 
-FTedsOutlinerTreeItem::FTedsOutlinerTreeItem(const UE::Editor::DataStorage::RowHandle& InRowHandle,
+FTedsOutlinerTreeItem::FTedsOutlinerTreeItem(const DataStorage::RowHandle& InRowHandle,
 	const TSharedRef<const FTedsOutlinerImpl>& InTedsOutlinerImpl)
 	: ISceneOutlinerTreeItem(Type)
 	, RowHandle(InRowHandle)
@@ -65,9 +67,10 @@ void FTedsOutlinerTreeItem::GenerateContextMenu(UToolMenu* Menu, SSceneOutliner&
 	);
 }
 
-UE::Editor::DataStorage::RowHandle FTedsOutlinerTreeItem::GetRowHandle() const
+DataStorage::RowHandle FTedsOutlinerTreeItem::GetRowHandle() const
 {
 	return RowHandle;
 }
+} // namespace UE::Editor::Outliner
 
 #undef LOCTEXT_NAMESPACE
\ No newline at end of file
diff --git a/Engine/Plugins/Experimental/EditorDataStorageFeatures/Source/TedsOutliner/Private/TedsOutlinerMode.cpp b/Engine/Plugins/Experimental/EditorDataStorageFeatures/Source/TedsOutliner/Private/TedsOutlinerMode.cpp
--- a/Engine/Plugins/Experimental/EditorDataStorageFeatures/Source/TedsOutliner/Private/TedsOutlinerMode.cpp
+++ b/Engine/Plugins/Experimental/EditorDataStorageFeatures/Source/TedsOutliner/Private/TedsOutlinerMode.cpp
@@ -17,13 +17,15 @@
 
 #define LOCTEXT_NAMESPACE "TedsOutlinerMode"
 
-namespace UE::EditorDataStorage::Outliner::Private
+namespace UE::Editor::Outliner
 {
-	// Drag drop currently disabled as we are missing data marshaling for hierarchies from TEDS to the world
-	static bool TedsOutlinerDragDropEnabled = false;
-	static FAutoConsoleVariableRef TedsOutlinerDragDropEnabledCvar(TEXT("TEDS.UI.EnableTEDSOutlinerDragDrop"), TedsOutlinerDragDropEnabled, TEXT("Enable drag/drop for the generic TEDS Outliner."));
-	static FName ContextMenuName("TedsOutlinerContextMenu");
-}
+	namespace Private
+	{
+		// Drag drop currently disabled as we are missing data marshaling for hierarchies from TEDS to the world
+		static bool TedsOutlinerDragDropEnabled = false;
+		static FAutoConsoleVariableRef TedsOutlinerDragDropEnabledCvar(TEXT("TEDS.UI.EnableTEDSOutlinerDragDrop"), TedsOutlinerDragDropEnabled, TEXT("Enable drag/drop for the generic TEDS Outliner."));
+		static FName ContextMenuName("TedsOutlinerContextMenu");
+	} // namespace Private
 
 FTedsOutlinerMode::FTedsOutlinerMode(const FTedsOutlinerParams& InParams)
 	: ISceneOutlinerMode(InParams.SceneOutliner)
@@ -58,7 +60,7 @@ void FTedsOutlinerMode::OnSelectionChanged()
 	{
 		if(const FTedsOutlinerTreeItem* TedsItem = InItem.CastTo<FTedsOutlinerTreeItem>())
 		{
-			const UE::Editor::DataStorage::RowHandle RowHandle = TedsItem->GetRowHandle();
+			const DataStorage::RowHandle RowHandle = TedsItem->GetRowHandle();
 
 			if(const FTypedElementSelectionColumn* SelectionColumn = Storage->GetColumn<FTypedElementSelectionColumn>(RowHandle))
 			{
@@ -81,7 +83,7 @@ void FTedsOutlinerMode::OnItemSelectionChanged(FSceneOutlinerTreeItemPtr Item, E
 		return; // Direct selection means we selected from outside the Outliner i.e through TEDS, so we don't need to redo the column addition
 	}
 
-	TArray<UE::Editor::DataStorage::RowHandle> RowHandles;
+	TArray<DataStorage::RowHandle> RowHandles;
 	
 	// The selection in the Outliner changed, update TEDS
 	Selection.ForEachItem([&RowHandles](FSceneOutlinerTreeItemPtr& Item)
@@ -105,7 +107,7 @@ TSharedPtr<FDragDropOperation> FTedsOutlinerMode::CreateDragDropOperation(const
 		return nullptr;
 	}
 	
-	TArray<UE::Editor::DataStorage::RowHandle> DraggedRowHandles;
+	TArray<DataStorage::RowHandle> DraggedRowHandles;
 
 	for(const FSceneOutlinerTreeItemPtr& Item :InTreeItems)
 	{
@@ -125,7 +127,7 @@ bool FTedsOutlinerMode::ParseDragDrop(FSceneOutlinerDragDropPayload& OutPayload,
 	{
 		const FTedsRowDragDropOp& TedsOp = static_cast<const FTedsRowDragDropOp&>(Operation);
 
-		for(UE::Editor::DataStorage::RowHandle RowHandle : TedsOp.DraggedRows)
+		for(DataStorage::RowHandle RowHandle : TedsOp.DraggedRows)
 		{
 			OutPayload.DraggedItems.Add(SceneOutliner->GetTreeItem(RowHandle));
 		}
@@ -146,7 +148,7 @@ FSceneOutlinerDragValidationInfo FTedsOutlinerMode::ValidateDrop(const ISceneOut
 			LOCTEXT("DropDisabled", "Drag/Drop is disabled due to missing hierarchy data!"));
 	}
 	
-	TArray<UE::Editor::DataStorage::RowHandle> DraggedRowHandles;
+	TArray<DataStorage::RowHandle> DraggedRowHandles;
 
 	Payload.ForEachItem<FTedsOutlinerTreeItem>([&DraggedRowHandles](FTedsOutlinerTreeItem& TedsItem)
 		{
@@ -157,7 +159,7 @@ FSceneOutlinerDragValidationInfo FTedsOutlinerMode::ValidateDrop(const ISceneOut
 	// TEDS-Outliner TODO: Need better drag/drop validation and better place for this, TEDS-Outliner does not know about what types these rows are and all types that exist and what attachment is valid
 	if(const FTedsOutlinerTreeItem* TedsItem = DropTarget.CastTo<FTedsOutlinerTreeItem>())
 	{
-		UE::Editor::DataStorage::RowHandle DropTargetRowHandle = TedsItem->GetRowHandle();
+		DataStorage::RowHandle DropTargetRowHandle = TedsItem->GetRowHandle();
 
 		FTypedElementClassTypeInfoColumn* DropTargetTypeInfoColumn = Storage->GetColumn<FTypedElementClassTypeInfoColumn>(DropTargetRowHandle);
 
@@ -169,7 +171,7 @@ FSceneOutlinerDragValidationInfo FTedsOutlinerMode::ValidateDrop(const ISceneOut
 
 		
 		// TEDS-Outliner TODO: Currently we detect parent changes by removing the column and then adding the column back with the new parent 
-		for(UE::Editor::DataStorage::RowHandle RowHandle : DraggedRowHandles)
+		for(DataStorage::RowHandle RowHandle : DraggedRowHandles)
 		{
 			FTypedElementClassTypeInfoColumn* TypeInfoColumn = Storage->GetColumn<FTypedElementClassTypeInfoColumn>(RowHandle);
 
@@ -196,7 +198,7 @@ FSceneOutlinerDragValidationInfo FTedsOutlinerMode::ValidateDrop(const ISceneOut
 		{
 			bool bValidDetach = false;
 			
-			for(UE::Editor::DataStorage::RowHandle RowHandle : DraggedRowHandles)
+			for(DataStorage::RowHandle RowHandle : DraggedRowHandles)
 			{
 				if(Storage->HasColumns(RowHandle, MakeArrayView({HierarchyData.GetValue().HierarchyColumn})))
 				{
@@ -220,12 +222,12 @@ void FTedsOutlinerMode::OnDrop(ISceneOutlinerTreeItem& DropTarget, const FSceneO
 	const TOptional<FTedsOutlinerHierarchyData>& HierarchyData = TedsOutlinerImpl->GetHierarchyData();
 	ITypedElementDataStorageInterface* Storage = TedsOutlinerImpl->GetStorage();
 
-	if(!UE::EditorDataStorage::Outliner::Private::TedsOutlinerDragDropEnabledCvar->GetBool() || !HierarchyData.IsSet())
+	if(!Private::TedsOutlinerDragDropEnabledCvar->GetBool() || !HierarchyData.IsSet())
 	{
 		return;
 	}
 	
-	TArray<UE::Editor::DataStorage::RowHandle> DraggedRowHandles;
+	TArray<DataStorage::RowHandle> DraggedRowHandles;
 
 	Payload.ForEachItem<FTedsOutlinerTreeItem>([&DraggedRowHandles](FTedsOutlinerTreeItem& TedsItem)
 	{
@@ -234,7 +236,7 @@ void FTedsOutlinerMode::OnDrop(ISceneOutlinerTreeItem& DropTarget, const FSceneO
 
 	if(ValidationInfo.CompatibilityType == ESceneOutlinerDropCompatibility::CompatibleDetach)
 	{
-		for(UE::Editor::DataStorage::RowHandle RowHandle : DraggedRowHandles)
+		for(DataStorage::RowHandle RowHandle : DraggedRowHandles)
 		{
 			Storage->RemoveColumn(RowHandle, HierarchyData.GetValue().HierarchyColumn);
 			Storage->AddColumn<FTypedElementSyncBackToWorldTag>(RowHandle);
@@ -243,9 +245,9 @@ void FTedsOutlinerMode::OnDrop(ISceneOutlinerTreeItem& DropTarget, const FSceneO
 	
 	if(const FTedsOutlinerTreeItem* TedsItem = DropTarget.CastTo<FTedsOutlinerTreeItem>())
 	{
-		UE::Editor::DataStorage::RowHandle DropTargetRowHandle = TedsItem->GetRowHandle();
+		DataStorage::RowHandle DropTargetRowHandle = TedsItem->GetRowHandle();
 		
-		for(UE::Editor::DataStorage::RowHandle RowHandle : DraggedRowHandles)
+		for(DataStorage::RowHandle RowHandle : DraggedRowHandles)
 		{
 			// Add the column
 			Storage->AddColumn(RowHandle, HierarchyData.GetValue().HierarchyColumn);
@@ -262,9 +264,9 @@ TSharedPtr<SWidget> FTedsOutlinerMode::CreateContextMenu()
 {
 	UToolMenus* ToolMenus = UToolMenus::Get();
 
-	if (!ToolMenus->IsMenuRegistered(UE::EditorDataStorage::Outliner::Private::ContextMenuName))
+	if (!ToolMenus->IsMenuRegistered(Private::ContextMenuName))
 	{
-		UToolMenu* Menu = ToolMenus->RegisterMenu((UE::EditorDataStorage::Outliner::Private::ContextMenuName));
+		UToolMenu* Menu = ToolMenus->RegisterMenu(Private::ContextMenuName);
 		Menu->AddDynamicSection("DynamicHierarchySection", FNewToolMenuDelegate::CreateLambda([](UToolMenu* InMenu)
 		{
 			if(UTedsOutlinerMenuContext* TedsOutlinerMenuContext = InMenu->FindContext<UTedsOutlinerMenuContext>())
@@ -289,12 +291,13 @@ TSharedPtr<SWidget> FTedsOutlinerMode::CreateContextMenu()
 	FToolMenuContext MenuContext;
 	MenuContext.AddObject(TedsOutlinerMenuContext);
 
-	return UToolMenus::Get()->GenerateWidget(UE::EditorDataStorage::Outliner::Private::ContextMenuName, MenuContext);
+	return UToolMenus::Get()->GenerateWidget(Private::ContextMenuName, MenuContext);
 }
 
 TUniquePtr<ISceneOutlinerHierarchy> FTedsOutlinerMode::CreateHierarchy()
 {
 	return MakeUnique<FTedsOutlinerHierarchy>(this, TedsOutlinerImpl.ToSharedRef());
 }
+} // namespace UE::Editor::Outliner
 
 #undef LOCTEXT_NAMESPACE
diff --git a/Engine/Plugins/Experimental/EditorDataStorageFeatures/Source/TedsOutliner/Private/TedsOutlinerModule.cpp b/Engine/Plugins/Experimental/EditorDataStorageFeatures/Source/TedsOutliner/Private/TedsOutlinerModule.cpp
--- a/Engine/Plugins/Experimental/EditorDataStorageFeatures/Source/TedsOutliner/Private/TedsOutlinerModule.cpp
+++ b/Engine/Plugins/Experimental/EditorDataStorageFeatures/Source/TedsOutliner/Private/TedsOutlinerModule.cpp
@@ -21,9 +21,12 @@
 
 #define LOCTEXT_NAMESPACE "TedsOutlinerModule"
 
-namespace UE::EditorDataStorage::Outliner::Private
+namespace UE::Editor::Outliner
 {
-	static bool bUseNewRevisionControlWidgets = false;
+	namespace Private
+	{
+		static bool bUseNewRevisionControlWidgets = false;
+	} // namespace Private
 
 	void RefreshLevelEditorTedsOutliner(bool bAlwaysInvoke)
 	{
@@ -40,7 +43,7 @@ namespace UE::EditorDataStorage::Outliner::Private
 	
 	static FAutoConsoleVariableRef CVarUseNewRevisionControlWidgets(
 		TEXT("TEDS.UI.UseNewRevisionControlWidgets"),
-		bUseNewRevisionControlWidgets,
+		Private::bUseNewRevisionControlWidgets,
 		TEXT("Use new TEDS-based source control widgets in the Outliner (requires TEDS-Outliner to be enabled)")
 		, FConsoleVariableDelegate::CreateLambda([](IConsoleVariable*)
 		{
@@ -55,13 +58,12 @@ namespace UE::EditorDataStorage::Outliner::Private
 		{
 			RefreshLevelEditorTedsOutliner(true);
 		}));
-}
 
 FTedsOutlinerModule::FTedsOutlinerModule()
 {
 }
 
-TSharedRef<ISceneOutliner> FTedsOutlinerModule::CreateTedsOutliner(const FSceneOutlinerInitializationOptions& InInitOptions, const FTedsOutlinerParams& InInitTedsOptions, UE::Editor::DataStorage::QueryHandle ColumnQuery) const
+TSharedRef<ISceneOutliner> FTedsOutlinerModule::CreateTedsOutliner(const FSceneOutlinerInitializationOptions& InInitOptions, const FTedsOutlinerParams& InInitTedsOptions, DataStorage::QueryHandle ColumnQuery) const
 {
 	UTypedElementRegistry* Registry = UTypedElementRegistry::GetInstance();
 
@@ -109,25 +111,25 @@ void FTedsOutlinerModule::ShutdownModule()
 	IModuleInterface::ShutdownModule();
 }
 
-UE::Editor::DataStorage::QueryHandle FTedsOutlinerModule::GetLevelEditorTedsOutlinerColumnQuery()
+DataStorage::QueryHandle FTedsOutlinerModule::GetLevelEditorTedsOutlinerColumnQuery()
 {
 	UTypedElementRegistry* Registry = UTypedElementRegistry::GetInstance();
 	ITypedElementDataStorageInterface* Storage = Registry->GetMutableDataStorage();
 		
-	using namespace UE::Editor::DataStorage::Queries;
+	using namespace DataStorage::Queries;
 
-	static UE::Editor::DataStorage::QueryHandle ColumnQuery = Storage->RegisterQuery(
+	static DataStorage::QueryHandle ColumnQuery = Storage->RegisterQuery(
 		Select()
 			.ReadOnly<FTypedElementClassTypeInfoColumn, FTypedElementAlertColumn, FTypedElementChildAlertColumn>()
 		.Compile());
 
 	// Query to also include revision control info
-	static UE::Editor::DataStorage::QueryHandle RevisionControlQuery = Storage->RegisterQuery(
+	static DataStorage::QueryHandle RevisionControlQuery = Storage->RegisterQuery(
 		Select()
 			.ReadOnly<FTypedElementClassTypeInfoColumn, FTypedElementPackageReference, FTypedElementAlertColumn>()
 		.Compile());
 
-	return UE::EditorDataStorage::Outliner::Private::bUseNewRevisionControlWidgets ? RevisionControlQuery : ColumnQuery;
+	return Private::bUseNewRevisionControlWidgets ? RevisionControlQuery : ColumnQuery;
 	
 }
 
@@ -142,10 +144,10 @@ TSharedRef<SWidget> FTedsOutlinerModule::CreateLevelEditorTedsOutliner()
 		.Text(LOCTEXT("TEDSPluginNotEnabledText", "You need to enable the Typed Element Data Storage plugin to see the table viewer!"));
 	}
 
-	using namespace UE::Editor::DataStorage::Queries;
+	using namespace DataStorage::Queries;
 
 	// The Outliner is populated with Actors and Entities
-	UE::Editor::DataStorage::FQueryDescription OutlinerQueryDescription =
+	DataStorage::FQueryDescription OutlinerQueryDescription =
 		Select()
 		.Where()
 			.All<FTypedElementClassTypeInfoColumn>() // TEDS-Outliner TODO: Currently looking at all entries with type info in TEDS
@@ -162,7 +164,7 @@ TSharedRef<SWidget> FTedsOutlinerModule::CreateLevelEditorTedsOutliner()
 	Params.bUseDefaultTedsFilters = true;
 
 	// Example Query to filter for actors
-	UE::Editor::DataStorage::FQueryDescription ActorFilterQuery =
+	DataStorage::FQueryDescription ActorFilterQuery =
 		Select()
 		.Where()
 			.All<FTypedElementActorTag>()
@@ -217,7 +219,8 @@ FName FTedsOutlinerModule::GetTedsOutlinerTabName()
 {
 	return TedsOutlinerTabName;
 }
+} // namsepace UE::Editor::Outliner
 
-IMPLEMENT_MODULE(FTedsOutlinerModule, TedsOutliner);
+IMPLEMENT_MODULE(UE::Editor::Outliner::FTedsOutlinerModule, TedsOutliner);
 
 #undef LOCTEXT_NAMESPACE // TedsOutlinerModule
diff --git a/Engine/Plugins/Experimental/EditorDataStorageFeatures/Source/TedsOutliner/Public/TedsOutlinerFilter.h b/Engine/Plugins/Experimental/EditorDataStorageFeatures/Source/TedsOutliner/Public/TedsOutlinerFilter.h
--- a/Engine/Plugins/Experimental/EditorDataStorageFeatures/Source/TedsOutliner/Public/TedsOutlinerFilter.h
+++ b/Engine/Plugins/Experimental/EditorDataStorageFeatures/Source/TedsOutliner/Public/TedsOutlinerFilter.h
@@ -7,6 +7,8 @@
 #include "Elements/Common/TypedElementQueryDescription.h"
 #include "Elements/Framework/TypedElementMetaData.h"
 
+namespace UE::Editor::Outliner
+{
 class FTedsOutlinerImpl;
 
 // TEDS-Outliner TODO: Lots of minor missing functionality for icon, tooltip, color etc
@@ -14,7 +16,7 @@ class FTedsOutlinerFilter : public FFilterBase<SceneOutliner::FilterBarType>
 {
 public:
 	FTedsOutlinerFilter(const FName& InFilterName, const FText& InFilterDisplayName, TSharedPtr<FFilterCategory> InCategory,
-		TSharedRef<FTedsOutlinerImpl> InTedsOutlinerImpl, const UE::Editor::DataStorage::FQueryDescription& InFilterQuery);
+		TSharedRef<FTedsOutlinerImpl> InTedsOutlinerImpl, const DataStorage::FQueryDescription& InFilterQuery);
 
 	/** Returns the system name for this filter */
 	virtual FString GetName() const override;
@@ -57,5 +59,6 @@ protected:
 	FText FilterDisplayName;
 
 	TSharedRef<FTedsOutlinerImpl> TedsOutlinerImpl;
-	const UE::Editor::DataStorage::FQueryDescription FilterQuery;
-};
\ No newline at end of file
+	const DataStorage::FQueryDescription FilterQuery;
+};
+} // namespace UE::Editor::Outliner
diff --git a/Engine/Plugins/Experimental/EditorDataStorageFeatures/Source/TedsOutliner/Public/TedsOutlinerImpl.h b/Engine/Plugins/Experimental/EditorDataStorageFeatures/Source/TedsOutliner/Public/TedsOutlinerImpl.h
--- a/Engine/Plugins/Experimental/EditorDataStorageFeatures/Source/TedsOutliner/Public/TedsOutlinerImpl.h
+++ b/Engine/Plugins/Experimental/EditorDataStorageFeatures/Source/TedsOutliner/Public/TedsOutlinerImpl.h
@@ -12,14 +12,16 @@ class ITypedElementDataStorageCompatibilityInterface;
 struct FTypedElementWidgetConstructor;
 class SWidget;
 
+namespace UE::Editor::Outliner
+{
 // Struct storing information on how hierarchies are handled in the TEDS Outliner
 struct FTedsOutlinerHierarchyData
 {
 	/** A delegate used to get the parent row handle for a given row */
-	DECLARE_DELEGATE_RetVal_OneParam(UE::Editor::DataStorage::RowHandle, FGetParentRowHandle, void* /* InColumnData */);
+	DECLARE_DELEGATE_RetVal_OneParam(DataStorage::RowHandle, FGetParentRowHandle, void* /* InColumnData */);
 	
 	/** A delegate used to set the parent row handle for a given row */
-	DECLARE_DELEGATE_TwoParams(FSetParentRowHandle, void* /* InColumnData */, UE::Editor::DataStorage::RowHandle /* InParentRowHandle */);
+	DECLARE_DELEGATE_TwoParams(FSetParentRowHandle, void* /* InColumnData */, DataStorage::RowHandle /* InParentRowHandle */);
 
 	FTedsOutlinerHierarchyData(const UScriptStruct* InHierarchyColumn, const FGetParentRowHandle& InGetParent, const FSetParentRowHandle& InSetParent)
 		: HierarchyColumn(InHierarchyColumn)
@@ -48,11 +50,11 @@ struct FTedsOutlinerHierarchyData
 					return ParentColumn->Parent;
 				}
 
-				return UE::Editor::DataStorage::InvalidRowHandle;
+				return DataStorage::InvalidRowHandle;
 			});
 
 		const FSetParentRowHandle RowHandleSetter = FSetParentRowHandle::CreateLambda([](void* InColumnData,
-			UE::Editor::DataStorage::RowHandle InRowHandle)
+			DataStorage::RowHandle InRowHandle)
 			{
 				if(FTableRowParentColumn* ParentColumn = static_cast<FTableRowParentColumn *>(InColumnData))
 				{
@@ -78,12 +80,12 @@ struct FTedsOutlinerParams
 	SSceneOutliner* SceneOutliner;
 
 	// The query description that will be used to populate rows in the TEDS-Outliner
-	TAttribute<UE::Editor::DataStorage::FQueryDescription> QueryDescription;
+	TAttribute<DataStorage::FQueryDescription> QueryDescription;
 	
 	// TEDS queries that will be used to create filters in this Outliner
 	// TEDS-Outliner TODO: Can we consolidate this with the SceneOutliner API to create filters? Currently has to be separate because FTEDSOutlinerFilter
 	// needs a reference to the mode which is not possible since filters with the Outliner API are added before the mode is init
-	TMap<FName, const UE::Editor::DataStorage::FQueryDescription> FilterQueries;
+	TMap<FName, const DataStorage::FQueryDescription> FilterQueries;
 
 	// If true, this Outliner will automatically add all TEDS tags and columns as filters
 	bool bUseDefaultTedsFilters;
@@ -132,43 +134,43 @@ public:
 	FIsItemCompatible& IsItemCompatible();
 
 	// Update the selection in TEDS to the input rows, only if SelectionSetName is set
-	void SetSelection(const TArray<UE::Editor::DataStorage::RowHandle>& InSelectedRows);
+	void SetSelection(const TArray<DataStorage::RowHandle>& InSelectedRows);
 
 	// Helper function to create a label widget for a given row
-	TSharedRef<SWidget> CreateLabelWidgetForItem(UE::Editor::DataStorage::RowHandle InRowHandle, const STableRow<FSceneOutlinerTreeItemPtr>& InRow) const;
+	TSharedRef<SWidget> CreateLabelWidgetForItem(DataStorage::RowHandle InRowHandle, const STableRow<FSceneOutlinerTreeItemPtr>& InRow) const;
 
 	// Get the hierarchy data associated with this table viewer
 	const TOptional<FTedsOutlinerHierarchyData>& GetHierarchyData();
 	
 	// Add an external query to the Outliner
-	void AddExternalQuery(FName QueryName, const UE::Editor::DataStorage::FQueryDescription& InQueryDescription);
+	void AddExternalQuery(FName QueryName, const DataStorage::FQueryDescription& InQueryDescription);
 	void RemoveExternalQuery(FName QueryName);
 
 	// Append all external queries into the given query
-	void AppendExternalQueries(UE::Editor::DataStorage::FQueryDescription& OutQuery);
+	void AppendExternalQueries(DataStorage::FQueryDescription& OutQuery);
 
 	// TEDS-Outliner TODO: This should live in TEDS long term
 	// Funtion to combine 2 queries (adds to second query to the first)
-	static void AppendQuery(UE::Editor::DataStorage::FQueryDescription& Query1, const UE::Editor::DataStorage::FQueryDescription& Query2);
+	static void AppendQuery(DataStorage::FQueryDescription& Query1, const DataStorage::FQueryDescription& Query2);
 
 	// Check if the given item's parent has changed (i.e ParentRowHandle does not match what the Outliner reports as the parent)
-	bool HasItemParentChanged(UE::Editor::DataStorage::RowHandle ItemRowHandle, UE::Editor::DataStorage::RowHandle ParentRowHandle) const;
+	bool HasItemParentChanged(DataStorage::RowHandle ItemRowHandle, DataStorage::RowHandle ParentRowHandle) const;
 
 	// Outliner specific functionality
 	void CreateItemsFromQuery(TArray<FSceneOutlinerTreeItemPtr>& OutItems, ISceneOutlinerMode* InMode) const;
 	void CreateChildren(const FSceneOutlinerTreeItemPtr& Item, TArray<FSceneOutlinerTreeItemPtr>& OutChildren) const;
 
 	// Get the parent row for a given row
-	UE::Editor::DataStorage::RowHandle GetParentRow(UE::Editor::DataStorage::RowHandle InRowHandle);
+	DataStorage::RowHandle GetParentRow(DataStorage::RowHandle InRowHandle);
 
 	// Recompile all queries used by this table viewer
 	void RecompileQueries();
 
 protected:
 	
-	void OnItemAdded(UE::Editor::DataStorage::RowHandle ItemRowHandle);
-	void OnItemRemoved(UE::Editor::DataStorage::RowHandle ItemRowHandle);
-	void OnItemMoved(UE::Editor::DataStorage::RowHandle ItemRowHandle);
+	void OnItemAdded(DataStorage::RowHandle ItemRowHandle);
+	void OnItemRemoved(DataStorage::RowHandle ItemRowHandle);
+	void OnItemMoved(DataStorage::RowHandle ItemRowHandle);
 
 	void UnregisterQueries() const;
 	void ClearSelection() const;
@@ -178,7 +180,7 @@ protected:
 	void CreateFilterQueries();
 
 	// Check if this row can be displayed in this table viewer
-	bool CanDisplayRow(UE::Editor::DataStorage::RowHandle ItemRowHandle) const;
+	bool CanDisplayRow(DataStorage::RowHandle ItemRowHandle) const;
 	
 protected:
 	// TEDS Storage Constructs
@@ -189,35 +191,35 @@ protected:
 	FTedsOutlinerParams CreationParams;
 
 	// Widget constructor to create the label widget
-	TArray<TPair<UE::Editor::DataStorage::QueryHandle, TSharedPtr<FTypedElementWidgetConstructor>>> QueryToWidgetConstructorMap;
+	TArray<TPair<DataStorage::QueryHandle, TSharedPtr<FTypedElementWidgetConstructor>>> QueryToWidgetConstructorMap;
 
 	// Widget purposes this table viewer supports
 	TArray<FName> CellWidgetPurposes;
 	
 	// Initial query provided by user
-	TAttribute<UE::Editor::DataStorage::FQueryDescription> InitialQueryDescription;
+	TAttribute<DataStorage::FQueryDescription> InitialQueryDescription;
 
 	// External queries that are currently active (e.g Filters)
-	TMap<FName, UE::Editor::DataStorage::FQueryDescription> ExternalQueries;
+	TMap<FName, DataStorage::FQueryDescription> ExternalQueries;
 
 	// Optional Hierarchy Data
 	TOptional<FTedsOutlinerHierarchyData> HierarchyData;
 
 	// Querys to track row handle collection, addition and removal
-	UE::Editor::DataStorage::QueryHandle RowHandleQuery = UE::Editor::DataStorage::InvalidQueryHandle;
-	UE::Editor::DataStorage::QueryHandle RowAdditionQuery = UE::Editor::DataStorage::InvalidQueryHandle;
-	UE::Editor::DataStorage::QueryHandle RowRemovalQuery = UE::Editor::DataStorage::InvalidQueryHandle;
+	DataStorage::QueryHandle RowHandleQuery = DataStorage::InvalidQueryHandle;
+	DataStorage::QueryHandle RowAdditionQuery = DataStorage::InvalidQueryHandle;
+	DataStorage::QueryHandle RowRemovalQuery = DataStorage::InvalidQueryHandle;
 
 	// Query to get all child rows
-	UE::Editor::DataStorage::QueryHandle ChildRowHandleQuery = UE::Editor::DataStorage::InvalidQueryHandle;
+	DataStorage::QueryHandle ChildRowHandleQuery = DataStorage::InvalidQueryHandle;
 
 	// Query to track when a row's parent gets changed
-	UE::Editor::DataStorage::QueryHandle UpdateParentQuery = UE::Editor::DataStorage::InvalidQueryHandle;
+	DataStorage::QueryHandle UpdateParentQuery = DataStorage::InvalidQueryHandle;
 
 	// Query to get all selected rows, track selection added, track selection removed
-	UE::Editor::DataStorage::QueryHandle SelectedRowsQuery = UE::Editor::DataStorage::InvalidQueryHandle;
-	UE::Editor::DataStorage::QueryHandle SelectionAddedQuery = UE::Editor::DataStorage::InvalidQueryHandle;
-	UE::Editor::DataStorage::QueryHandle SelectionRemovedQuery = UE::Editor::DataStorage::InvalidQueryHandle;
+	DataStorage::QueryHandle SelectedRowsQuery = DataStorage::InvalidQueryHandle;
+	DataStorage::QueryHandle SelectionAddedQuery = DataStorage::InvalidQueryHandle;
+	DataStorage::QueryHandle SelectionRemovedQuery = DataStorage::InvalidQueryHandle;
 	
 	TOptional<FName> SelectionSetName;
 	bool bSelectionDirty = false;
@@ -237,3 +239,4 @@ protected:
 	// Delegate to check if an item is compatible with this table viewer
 	FIsItemCompatible IsItemCompatibleWithTeds;
 };
+} // namsepace UE::Editor::Outliner
diff --git a/Engine/Plugins/Experimental/EditorDataStorageFeatures/Source/TedsOutliner/Public/TedsOutlinerItem.h b/Engine/Plugins/Experimental/EditorDataStorageFeatures/Source/TedsOutliner/Public/TedsOutlinerItem.h
--- a/Engine/Plugins/Experimental/EditorDataStorageFeatures/Source/TedsOutliner/Public/TedsOutlinerItem.h
+++ b/Engine/Plugins/Experimental/EditorDataStorageFeatures/Source/TedsOutliner/Public/TedsOutlinerItem.h
@@ -5,10 +5,12 @@
 #include "Elements/Interfaces/TypedElementDataStorageInterface.h"
 #include "TedsOutlinerMode.h"
 
+namespace UE::Editor::Outliner
+{
 class FTedsOutlinerImpl;
 
 /*
- * A generic item in the TEDS driven Outliner, that uses a UE::Editor::DataStorage::RowHandle to uniquely identify the object it is
+ * A generic item in the TEDS driven Outliner, that uses a DataStorage::RowHandle to uniquely identify the object it is
  * looking at. Functionality should be added through TEDS queries instead of having a different TreeItem type for each
  * type of object you are looking at (i.e. Actor vs Folder)
  * Inherits from ISceneOutlinerItem - which determines what type of item you are looking at. E.G FActorTreeItem for actors
@@ -17,14 +19,14 @@ struct FTedsOutlinerTreeItem : ISceneOutlinerTreeItem
 {
 public:
 	
-	DECLARE_DELEGATE_RetVal_OneParam(bool, FFilterPredicate, const UE::Editor::DataStorage::RowHandle);
+	DECLARE_DELEGATE_RetVal_OneParam(bool, FFilterPredicate, const DataStorage::RowHandle);
 
 	TEDSOUTLINER_API bool Filter(FFilterPredicate Pred) const
 	{
 		return Pred.Execute(RowHandle);
 	}
 
-	TEDSOUTLINER_API FTedsOutlinerTreeItem(const UE::Editor::DataStorage::RowHandle& InRowHandle, const TSharedRef<const FTedsOutlinerImpl>& InTedsOutlinerImpl);
+	TEDSOUTLINER_API FTedsOutlinerTreeItem(const DataStorage::RowHandle& InRowHandle, const TSharedRef<const FTedsOutlinerImpl>& InTedsOutlinerImpl);
 
 	/* Begin ISceneOutlinerTreeItem Implementation */
 	TEDSOUTLINER_API virtual bool IsValid() const override;
@@ -38,9 +40,10 @@ public:
 
 	TEDSOUTLINER_API static const FSceneOutlinerTreeItemType Type;
 
-	TEDSOUTLINER_API UE::Editor::DataStorage::RowHandle GetRowHandle() const;
+	TEDSOUTLINER_API DataStorage::RowHandle GetRowHandle() const;
 
 private:
-	const UE::Editor::DataStorage::RowHandle RowHandle;
+	const DataStorage::RowHandle RowHandle;
 	const TSharedRef<const FTedsOutlinerImpl> TedsOutlinerImpl;
-};
\ No newline at end of file
+};
+} // namespace UE::Editor::Outliner
diff --git a/Engine/Plugins/Experimental/EditorDataStorageFeatures/Source/TedsOutliner/Public/TedsOutlinerMode.h b/Engine/Plugins/Experimental/EditorDataStorageFeatures/Source/TedsOutliner/Public/TedsOutlinerMode.h
--- a/Engine/Plugins/Experimental/EditorDataStorageFeatures/Source/TedsOutliner/Public/TedsOutlinerMode.h
+++ b/Engine/Plugins/Experimental/EditorDataStorageFeatures/Source/TedsOutliner/Public/TedsOutlinerMode.h
@@ -9,6 +9,8 @@
 
 #include "TedsOutlinerMode.generated.h"
 
+namespace UE::Editor::Outliner
+{
 // TEDS-Outliner TODO: This can probably be moved to a more generic location for all TEDS related drag drops?
 class FTedsRowDragDropOp : public FDecoratedDragDropOp
 {
@@ -16,14 +18,14 @@ public:
 	DRAG_DROP_OPERATOR_TYPE(FTedsRowDragDropOp, FDecoratedDragDropOp)
 
 	/** Rows we are dragging */
-	TArray<UE::Editor::DataStorage::RowHandle> DraggedRows;
+	TArray<DataStorage::RowHandle> DraggedRows;
 
-	void Init(const TArray<UE::Editor::DataStorage::RowHandle>& InRowHandles)
+	void Init(const TArray<DataStorage::RowHandle>& InRowHandles)
 	{
 		DraggedRows = InRowHandles;
 	}
 
-	static TSharedRef<FTedsRowDragDropOp> New(const TArray<UE::Editor::DataStorage::RowHandle>& InRowHandles)
+	static TSharedRef<FTedsRowDragDropOp> New(const TArray<DataStorage::RowHandle>& InRowHandles)
 	{
 		TSharedRef<FTedsRowDragDropOp> Operation = MakeShareable(new FTedsRowDragDropOp);
 		
@@ -35,16 +37,6 @@ public:
 	}
 };
 
-// Class to hold the owning scene outliner for a menu
-// TEDS-Outliner TODO: Once menus go through TEDS UI this can be done using the FTableViewerColumn on the widget row instead
-UCLASS()
-class UTedsOutlinerMenuContext : public UObject
-{
-	GENERATED_BODY()
-public:
-	SSceneOutliner* OwningSceneOutliner = nullptr;
-};
-
 /*
  * TEDS driven Outliner mode where the Outliner is populated using the results of the RowHandleQueries passed in.
  * See CreateGenericTEDSOutliner() for example usage
@@ -79,4 +71,15 @@ protected:
 
 	// The actual model for the TEDS Outliner
 	TSharedPtr<FTedsOutlinerImpl> TedsOutlinerImpl;
-};
\ No newline at end of file
+};
+} // namespace UE::Editor::Outliner
+
+// Class to hold the owning scene outliner for a menu
+// TEDS-Outliner TODO: Once menus go through TEDS UI this can be done using the FTableViewerColumn on the widget row instead
+UCLASS()
+class UTedsOutlinerMenuContext : public UObject
+{
+	GENERATED_BODY()
+public:
+	SSceneOutliner* OwningSceneOutliner = nullptr;
+};
diff --git a/Engine/Plugins/Experimental/EditorDataStorageFeatures/Source/TedsOutliner/Public/TedsOutlinerModule.h b/Engine/Plugins/Experimental/EditorDataStorageFeatures/Source/TedsOutliner/Public/TedsOutlinerModule.h
--- a/Engine/Plugins/Experimental/EditorDataStorageFeatures/Source/TedsOutliner/Public/TedsOutlinerModule.h
+++ b/Engine/Plugins/Experimental/EditorDataStorageFeatures/Source/TedsOutliner/Public/TedsOutlinerModule.h
@@ -9,11 +9,14 @@
 
 class ISceneOutliner;
 struct FSceneOutlinerInitializationOptions;
-struct FTedsOutlinerParams;
 class SDockTab;
 class SWidget;
 class FSpawnTabArgs;
 
+namespace UE::Editor::Outliner
+{
+	struct FTedsOutlinerParams;
+
 /**
  * Implements the Scene Outliner module.
  */
@@ -57,3 +60,4 @@ private:
 	FName TedsOutlinerTabName;
 	FDelegateHandle LevelEditorTabManagerChangedHandle;
 };
+} // namespace UE::Editor::Outliner
diff --git a/Engine/Plugins/Experimental/EditorDataStorageFeatures/Source/TedsPropertyEditor/Private/TedsRowPickingMode.cpp b/Engine/Plugins/Experimental/EditorDataStorageFeatures/Source/TedsPropertyEditor/Private/TedsRowPickingMode.cpp
--- a/Engine/Plugins/Experimental/EditorDataStorageFeatures/Source/TedsPropertyEditor/Private/TedsRowPickingMode.cpp
+++ b/Engine/Plugins/Experimental/EditorDataStorageFeatures/Source/TedsPropertyEditor/Private/TedsRowPickingMode.cpp
@@ -4,7 +4,7 @@
 
 #include "TedsOutlinerItem.h"
 
-FTedsRowPickingMode::FTedsRowPickingMode(const FTedsOutlinerParams& Params, FOnSceneOutlinerItemPicked OnItemPickedDelegate)
+FTedsRowPickingMode::FTedsRowPickingMode(const UE::Editor::Outliner::FTedsOutlinerParams& Params, FOnSceneOutlinerItemPicked OnItemPickedDelegate)
 	: FTedsOutlinerMode(Params)
 	, OnItemPicked(OnItemPickedDelegate)
 {
@@ -32,7 +32,7 @@ void FTedsRowPickingMode::OnFilterTextCommited(FSceneOutlinerItemSelection& Sele
 	// In picking mode, we check to see if we have any Typed Element items, and if so, fire
 	// off the notification to whoever is listening. This may often cause the widget itself
 	// to be enqueued for destruction
-	TArray<FTedsOutlinerTreeItem*> OutlinerRows;
+	TArray<UE::Editor::Outliner::FTedsOutlinerTreeItem*> OutlinerRows;
 	Selection.Get(OutlinerRows);
 	if (OutlinerRows.Num() == 1 && OutlinerRows[0])
 	{
diff --git a/Engine/Plugins/Experimental/EditorDataStorageFeatures/Source/TedsPropertyEditor/Private/Widgets/SPropertyMenuTedsRowPicker.cpp b/Engine/Plugins/Experimental/EditorDataStorageFeatures/Source/TedsPropertyEditor/Private/Widgets/SPropertyMenuTedsRowPicker.cpp
--- a/Engine/Plugins/Experimental/EditorDataStorageFeatures/Source/TedsPropertyEditor/Private/Widgets/SPropertyMenuTedsRowPicker.cpp
+++ b/Engine/Plugins/Experimental/EditorDataStorageFeatures/Source/TedsPropertyEditor/Private/Widgets/SPropertyMenuTedsRowPicker.cpp
@@ -35,6 +35,8 @@ void SPropertyMenuTedsRowPicker::Construct(const FArguments& InArgs)
 	{
 		TSharedPtr<SWidget> MenuContent;
 		{
+			using namespace UE::Editor::Outliner;
+
 			// TEDS-Outliner TODO: Taken from private implementation of PropertyEditorAssetConstants.
 			//                     Should be centralized when TEDS is moved to core
 			static const FVector2D ContentBrowserWindowSize(300.0f, 300.0f);
diff --git a/Engine/Plugins/Experimental/EditorDataStorageFeatures/Source/TedsPropertyEditor/Public/TedsRowPickingMode.h b/Engine/Plugins/Experimental/EditorDataStorageFeatures/Source/TedsPropertyEditor/Public/TedsRowPickingMode.h
--- a/Engine/Plugins/Experimental/EditorDataStorageFeatures/Source/TedsPropertyEditor/Public/TedsRowPickingMode.h
+++ b/Engine/Plugins/Experimental/EditorDataStorageFeatures/Source/TedsPropertyEditor/Public/TedsRowPickingMode.h
@@ -12,10 +12,10 @@ DECLARE_DELEGATE_RetVal_OneParam(bool, FOnShouldFilterTedsRow, const UE::Editor:
 /*
 * Picking mode for TEDs Scene Outliner Widgets. Based off of FActorPickingMode
 */
-class TEDSPROPERTYEDITOR_API FTedsRowPickingMode : public FTedsOutlinerMode
+class TEDSPROPERTYEDITOR_API FTedsRowPickingMode : public UE::Editor::Outliner::FTedsOutlinerMode
 {
 public:
-	FTedsRowPickingMode(const FTedsOutlinerParams& Params, FOnSceneOutlinerItemPicked OnItemPickedDelegate);
+	FTedsRowPickingMode(const UE::Editor::Outliner::FTedsOutlinerParams& Params, FOnSceneOutlinerItemPicked OnItemPickedDelegate);
 
 	virtual ~FTedsRowPickingMode() = default;
 public:
diff --git a/Engine/Plugins/Experimental/EditorDataStorageFeatures/Source/TedsSettings/Private/TedsSettingsEditorSubsystem.cpp b/Engine/Plugins/Experimental/EditorDataStorageFeatures/Source/TedsSettings/Private/TedsSettingsEditorSubsystem.cpp
--- a/Engine/Plugins/Experimental/EditorDataStorageFeatures/Source/TedsSettings/Private/TedsSettingsEditorSubsystem.cpp
+++ b/Engine/Plugins/Experimental/EditorDataStorageFeatures/Source/TedsSettings/Private/TedsSettingsEditorSubsystem.cpp
@@ -6,13 +6,13 @@
 #include "TedsSettingsLog.h"
 #include "TedsSettingsManager.h"
 
-namespace UE::EditorDataStorage::Settings
+namespace UE::Editor::Settings::Private
 {
 	static TAutoConsoleVariable<bool> CVarTedsSettingsEnable(
 		TEXT("TEDS.Feature.Settings.Enable"),
 		false,
 		TEXT("When true, settings objects from the ISettingsModule will be mirrored to rows in the editor data storage."));
-}
+} // namespace UE::Editor::Settings::Private
 
 UTedsSettingsEditorSubsystem::UTedsSettingsEditorSubsystem()
 	: UEditorSubsystem()
@@ -23,7 +23,7 @@ UTedsSettingsEditorSubsystem::UTedsSettingsEditorSubsystem()
 
 const bool UTedsSettingsEditorSubsystem::IsEnabled() const
 {
-	return UE::EditorDataStorage::Settings::CVarTedsSettingsEnable.GetValueOnGameThread();
+	return UE::Editor::Settings::Private::CVarTedsSettingsEnable.GetValueOnGameThread();
 }
 
 UTedsSettingsEditorSubsystem::FOnEnabledChanged& UTedsSettingsEditorSubsystem::OnEnabledChanged()
@@ -37,7 +37,7 @@ void UTedsSettingsEditorSubsystem::Initialize(FSubsystemCollectionBase& Collecti
 
 	UE_LOG(LogTedsSettings, Log, TEXT("UTedsSettingsEditorSubsystem::Initialize"));
 
-	UE::EditorDataStorage::Settings::CVarTedsSettingsEnable->SetOnChangedCallback(
+	UE::Editor::Settings::Private::CVarTedsSettingsEnable->SetOnChangedCallback(
 		FConsoleVariableDelegate::CreateLambda([this](IConsoleVariable* Variable)
 		{
 			const bool bIsEnabled = Variable->GetBool();
@@ -54,7 +54,7 @@ void UTedsSettingsEditorSubsystem::Initialize(FSubsystemCollectionBase& Collecti
 			EnabledChangedDelegate.Broadcast();
 		}));
 
-	if (UE::EditorDataStorage::Settings::CVarTedsSettingsEnable.GetValueOnGameThread())
+	if (UE::Editor::Settings::Private::CVarTedsSettingsEnable.GetValueOnGameThread())
 	{
 		SettingsManager->Initialize();
 	}
@@ -64,7 +64,7 @@ void UTedsSettingsEditorSubsystem::Deinitialize()
 {
 	UE_LOG(LogTedsSettings, Log, TEXT("UTedsSettingsEditorSubsystem::Deinitialize"));
 
-	if (UE::EditorDataStorage::Settings::CVarTedsSettingsEnable.GetValueOnGameThread())
+	if (UE::Editor::Settings::Private::CVarTedsSettingsEnable.GetValueOnGameThread())
 	{
 		SettingsManager->Shutdown();
 	}
diff --git a/Engine/Plugins/Experimental/EditorDataStorageFeatures/Source/TedsTableViewer/Private/Widgets/SRowDetails.cpp b/Engine/Plugins/Experimental/EditorDataStorageFeatures/Source/TedsTableViewer/Private/Widgets/SRowDetails.cpp
--- a/Engine/Plugins/Experimental/EditorDataStorageFeatures/Source/TedsTableViewer/Private/Widgets/SRowDetails.cpp
+++ b/Engine/Plugins/Experimental/EditorDataStorageFeatures/Source/TedsTableViewer/Private/Widgets/SRowDetails.cpp
@@ -188,7 +188,7 @@ namespace UE::Editor::DataStorage
 	
 	}
 	
-} // namespace UE::EditorDataStorage
+} // namespace UE::Editor::DataStorage
 
 
 #undef LOCTEXT_NAMESPACE // "SRowDetails"
diff --git a/Engine/Plugins/Experimental/EditorDataStorageFeatures/Source/TedsTableViewer/Public/Widgets/STedsTableViewer.h b/Engine/Plugins/Experimental/EditorDataStorageFeatures/Source/TedsTableViewer/Public/Widgets/STedsTableViewer.h
--- a/Engine/Plugins/Experimental/EditorDataStorageFeatures/Source/TedsTableViewer/Public/Widgets/STedsTableViewer.h
+++ b/Engine/Plugins/Experimental/EditorDataStorageFeatures/Source/TedsTableViewer/Public/Widgets/STedsTableViewer.h
@@ -23,7 +23,7 @@ namespace UE::Editor::DataStorage
 	 * Example usage:
 	 * 
 	 *	SNew(STedsTableViewer)
-     *		.QueryStack(MakeShared<UE::EditorDataStorage::FQueryStackNode_RowView>(&Rows))
+     *		.QueryStack(MakeShared<UE::Editor::DataStorage::FQueryStackNode_RowView>(&Rows))
 	 *		.Columns({FTypedElementLabelColumn::StaticStruct(), FTypedElementClassTypeInfoColumn::StaticStruct());
 	 */
 	class STedsTableViewer : public SCompoundWidget
@@ -31,7 +31,7 @@ namespace UE::Editor::DataStorage
 	public:
 		
 		// Delegate fired when the selection in the table viewer changes
-		DECLARE_DELEGATE_OneParam(FOnSelectionChanged, UE::Editor::DataStorage::RowHandle)
+		DECLARE_DELEGATE_OneParam(FOnSelectionChanged, RowHandle)
 
 		SLATE_BEGIN_ARGS(STedsTableViewer)
 			: _CellWidgetPurposes({TEXT("General.Cell")})
@@ -74,10 +74,10 @@ namespace UE::Editor::DataStorage
 		TEDSTABLEVIEWER_API void AddCustomColumn(const TSharedRef<FTedsTableViewerColumn>& InColumn);
 
 		// Execute the given callback for each row that is selected in the table viewer
-		TEDSTABLEVIEWER_API void ForEachSelectedRow(TFunctionRef<void(UE::Editor::DataStorage::RowHandle)> InCallback) const;
+		TEDSTABLEVIEWER_API void ForEachSelectedRow(TFunctionRef<void(RowHandle)> InCallback) const;
 
 		// Get the row handle for the widget row the table viewer's contents are stored in
-		TEDSTABLEVIEWER_API UE::Editor::DataStorage::RowHandle GetWidgetRowHandle() const;
+		TEDSTABLEVIEWER_API RowHandle GetWidgetRowHandle() const;
 
 	protected:
 		
