Universal Search - Added new label widgets.

Added new label widgets to setup the context menu and double click behavior specifically for the Universal search.

#rnx

diff --git a/Engine/Plugins/Experimental/EditorDataStorageFeatures/Source/TedsAssetData/Private/Widgets/AssetDataLabelWidget.cpp b/Engine/Plugins/Experimental/EditorDataStorageFeatures/Source/TedsAssetData/Private/Widgets/AssetDataLabelWidget.cpp
--- a/Engine/Plugins/Experimental/EditorDataStorageFeatures/Source/TedsAssetData/Private/Widgets/AssetDataLabelWidget.cpp
+++ b/Engine/Plugins/Experimental/EditorDataStorageFeatures/Source/TedsAssetData/Private/Widgets/AssetDataLabelWidget.cpp
@@ -25,6 +25,10 @@ FAssetDataLabelWidgetConstructor::FAssetDataLabelWidgetConstructor()
 {
 }
 
+FAssetDataLabelWidgetConstructor::FAssetDataLabelWidgetConstructor(const UScriptStruct* TypeInfo)
+	: FSimpleWidgetConstructor(TypeInfo)
+{
+}
 
 TSharedPtr<SWidget> FAssetDataLabelWidgetConstructor::CreateWidget(ITypedElementDataStorageInterface* DataStorage,
 	ITypedElementDataStorageUiInterface* DataStorageUi, UE::Editor::DataStorage::RowHandle TargetRow, UE::Editor::DataStorage::RowHandle WidgetRow,
diff --git a/Engine/Plugins/Experimental/EditorDataStorageFeatures/Source/TedsAssetData/Public/Widgets/AssetDataLabelWidget.h b/Engine/Plugins/Experimental/EditorDataStorageFeatures/Source/TedsAssetData/Public/Widgets/AssetDataLabelWidget.h
--- a/Engine/Plugins/Experimental/EditorDataStorageFeatures/Source/TedsAssetData/Public/Widgets/AssetDataLabelWidget.h
+++ b/Engine/Plugins/Experimental/EditorDataStorageFeatures/Source/TedsAssetData/Public/Widgets/AssetDataLabelWidget.h
@@ -32,6 +32,7 @@ struct FAssetDataLabelWidgetConstructor : public FSimpleWidgetConstructor
 
 public:
 	TEDSASSETDATA_API FAssetDataLabelWidgetConstructor();
+	TEDSASSETDATA_API explicit FAssetDataLabelWidgetConstructor(const UScriptStruct* TypeInfo);
 	~FAssetDataLabelWidgetConstructor() override = default;
 
 	TEDSASSETDATA_API virtual TSharedPtr<SWidget> CreateWidget(
diff --git a/Engine/Plugins/Experimental/EditorDataStorageFeatures/Source/TedsTableViewer/Private/Widgets/UObjectLabelWidget.cpp b/Engine/Plugins/Experimental/EditorDataStorageFeatures/Source/TedsTableViewer/Private/Widgets/UObjectLabelWidget.cpp
--- a/Engine/Plugins/Experimental/EditorDataStorageFeatures/Source/TedsTableViewer/Private/Widgets/UObjectLabelWidget.cpp
+++ b/Engine/Plugins/Experimental/EditorDataStorageFeatures/Source/TedsTableViewer/Private/Widgets/UObjectLabelWidget.cpp
@@ -28,6 +28,11 @@ FUObjectLabelWidgetConstructor::FUObjectLabelWidgetConstructor()
 {
 }
 
+FUObjectLabelWidgetConstructor::FUObjectLabelWidgetConstructor(const UScriptStruct* TypeInfo)
+	: FSimpleWidgetConstructor(TypeInfo)
+{
+}
+
 TSharedPtr<SWidget> FUObjectLabelWidgetConstructor::CreateWidget(ITypedElementDataStorageInterface* DataStorage,
 	ITypedElementDataStorageUiInterface* DataStorageUi, RowHandle TargetRow, RowHandle WidgetRow,
 	const UE::Editor::DataStorage::FMetaDataView& Arguments)
diff --git a/Engine/Plugins/Experimental/EditorDataStorageFeatures/Source/TedsTableViewer/Public/Widgets/UObjectLabelWidget.h b/Engine/Plugins/Experimental/EditorDataStorageFeatures/Source/TedsTableViewer/Public/Widgets/UObjectLabelWidget.h
--- a/Engine/Plugins/Experimental/EditorDataStorageFeatures/Source/TedsTableViewer/Public/Widgets/UObjectLabelWidget.h
+++ b/Engine/Plugins/Experimental/EditorDataStorageFeatures/Source/TedsTableViewer/Public/Widgets/UObjectLabelWidget.h
@@ -32,6 +32,7 @@ struct FUObjectLabelWidgetConstructor : public FSimpleWidgetConstructor
 
 public:
 	TEDSTABLEVIEWER_API FUObjectLabelWidgetConstructor();
+	TEDSTABLEVIEWER_API explicit FUObjectLabelWidgetConstructor(const UScriptStruct* TypeInfo);
 	~FUObjectLabelWidgetConstructor() override = default;
 
 	TEDSTABLEVIEWER_API virtual TSharedPtr<SWidget> CreateWidget(
