TedsSettingsManager : Remove container and category rows in UnregisterSettings

This fixes a bug where we could get duplicate container / category rows when toggling the "TEDS.Feature.Settings.Enable" CVar multiple times.

diff --git a/Engine/Plugins/Experimental/EditorDataStorageFeatures/Source/TedsSettings/Private/TedsSettingsManager.cpp b/Engine/Plugins/Experimental/EditorDataStorageFeatures/Source/TedsSettings/Private/TedsSettingsManager.cpp
--- a/Engine/Plugins/Experimental/EditorDataStorageFeatures/Source/TedsSettings/Private/TedsSettingsManager.cpp
+++ b/Engine/Plugins/Experimental/EditorDataStorageFeatures/Source/TedsSettings/Private/TedsSettingsManager.cpp
@@ -156,11 +156,18 @@ void FTedsSettingsManager::RegisterSettingsContainer(const FName& ContainerName)
 
 	ISettingsContainerPtr ContainerPtr = SettingsModule->GetContainer(ContainerName);
 
-	RowHandle ContainerRow = DataStorage->AddRow(SettingsContainerTable);
-	DataStorage->AddColumn<FNameColumn>(ContainerRow, { .Name = ContainerName });
-	DataStorage->AddColumn<FDisplayNameColumn>(ContainerRow, { .DisplayName = ContainerPtr->GetDisplayName() });
-	DataStorage->AddColumn<FDescriptionColumn>(ContainerRow, { .Description = ContainerPtr->GetDescription() });
-	DataStorage->AddColumn<FSettingsContainerTag>(ContainerRow);
+	uint64 ContainerIndexHash = GenerateIndexHash(ContainerPtr.Get());
+	RowHandle ContainerRow = DataStorage->FindIndexedRow(ContainerIndexHash);
+	if (ContainerRow == InvalidRowHandle)
+	{
+		ContainerRow = DataStorage->AddRow(SettingsContainerTable);
+		DataStorage->AddColumn<FNameColumn>(ContainerRow, { .Name = ContainerName });
+		DataStorage->AddColumn<FDisplayNameColumn>(ContainerRow, { .DisplayName = ContainerPtr->GetDisplayName() });
+		DataStorage->AddColumn<FDescriptionColumn>(ContainerRow, { .Description = ContainerPtr->GetDescription() });
+		DataStorage->AddColumn<FSettingsContainerTag>(ContainerRow);
+
+		DataStorage->IndexRow(ContainerIndexHash, ContainerRow);
+	}
 
 	TArray<ISettingsCategoryPtr> Categories;
 	ContainerPtr->GetCategories(Categories);
@@ -184,6 +191,8 @@ void FTedsSettingsManager::RegisterSettingsContainer(const FName& ContainerName)
 
 void FTedsSettingsManager::UnregisterSettings()
 {
+	using namespace UE::Editor::DataStorage;
+
 	TRACE_CPUPROFILER_EVENT_SCOPE(TedsSettingsManager.UnregisterSettings);
 
 	ISettingsModule* SettingsModule = FModuleManager::GetModulePtr<ISettingsModule>("Settings");
@@ -192,6 +201,9 @@ void FTedsSettingsManager::UnregisterSettings()
 	UTypedElementRegistry* TypedElementRegistry = UTypedElementRegistry::GetInstance();
 	check(TypedElementRegistry);
 
+	ITypedElementDataStorageInterface* DataStorage = TypedElementRegistry->GetMutableDataStorage();
+	check(DataStorage);
+
 	ITypedElementDataStorageCompatibilityInterface* DataStorageCompatibility = TypedElementRegistry->GetMutableDataStorageCompatibility();
 	check(DataStorageCompatibility);
 
@@ -228,6 +240,22 @@ void FTedsSettingsManager::UnregisterSettings()
 					UE_LOG(LogTedsSettings, Log, TEXT("Removed Settings Section : '%s'"), *SectionPtr->GetName().ToString());
 				}
 			}
+
+			uint64 CategoryIndexHash = GenerateIndexHash(CategoryPtr.Get());
+			RowHandle CategoryRow = DataStorage->FindIndexedRow(CategoryIndexHash);
+			if (CategoryRow != InvalidRowHandle)
+			{
+				DataStorage->RemoveRow(CategoryRow);
+				DataStorage->RemoveIndex(CategoryIndexHash);
+			}
+		}
+
+		uint64 ContainerIndexHash = GenerateIndexHash(ContainerPtr.Get());
+		RowHandle ContainerRow = DataStorage->FindIndexedRow(ContainerIndexHash);
+		if (ContainerRow != InvalidRowHandle)
+		{
+			DataStorage->RemoveRow(ContainerRow);
+			DataStorage->RemoveIndex(ContainerIndexHash);
 		}
 	}
 }
