Fix invalid TEDS interface usage
* The original issue that tipped off this change was use of IQueryContext::GetColumn from within a singular callback.  This function doesn't return the column associated with the Row of the callback, it returns the first column in the chunk that is being processed.
* More investigation shows that the effort to deduplicate some code via HasItemParentChanged, OnItemMoved and CanDisplayRow, leads to a query callback calling back in through the TEDS interface.  This bypasses the column access safeguards that lock tables during processing.  Example can be seen with the TedsOutlinerColumn.
* Refactored into pure functions that are explicit about the function's dependencies.

diff --git a/Engine/Plugins/Experimental/EditorDataStorageFeatures/Source/TedsOutliner/Private/TedsOutlinerImpl.cpp b/Engine/Plugins/Experimental/EditorDataStorageFeatures/Source/TedsOutliner/Private/TedsOutlinerImpl.cpp
--- a/Engine/Plugins/Experimental/EditorDataStorageFeatures/Source/TedsOutliner/Private/TedsOutlinerImpl.cpp
+++ b/Engine/Plugins/Experimental/EditorDataStorageFeatures/Source/TedsOutliner/Private/TedsOutlinerImpl.cpp
@@ -27,6 +27,57 @@
 
 namespace UE::Editor::Outliner
 {
+namespace QueryUtils
+{
+	static bool CanDisplayRow(DataStorage::IQueryContext& Context, const FTedsOutlinerColumn& TedsOutlinerColumn, DataStorage::RowHandle Row, SSceneOutliner& SceneOutliner)
+	{
+		/*
+		 * Don't display widgets that are created for rows in this table viewer. Widgets are only created for rows that are currently visible, so if we
+		 * display the rows for them we are now adding/removing rows to the table viewer based on currently visible rows. But adding rows can cause
+		 * scrolling and change the currently visible rows which in turn again adds/removes widget rows. This chain keeps continuing which can cause
+		 * flickering/scrolling issues in the table viewer.
+		 */
+		if (Context.HasColumn<FTypedElementSlateWidgetReferenceColumn>(Row))
+		{
+			// Check if this widget row belongs to the same table viewer it is being displayed in
+			if (const TSharedPtr<ISceneOutliner> TableViewer = TedsOutlinerColumn.Outliner.Pin())
+			{
+				return &SceneOutliner != TableViewer.Get();
+			}
+			
+		}
+		return true;
+	}
+
+	static bool HasItemParentChanged(DataStorage::IQueryContext& Context, DataStorage::RowHandle Row, DataStorage::RowHandle ParentRowHandle, SSceneOutliner& SceneOutliner)
+	{
+		const FSceneOutlinerTreeItemPtr Item = SceneOutliner.GetTreeItem(Row, true);
+
+		// If the item doesn't exist, it doesn't make sense to say its parent changed
+		if (!Item)
+		{
+			return false;
+		}
+										
+		const FSceneOutlinerTreeItemPtr ParentItem = Item->GetParent();
+
+		// If the item doesn't have a parent, but ParentRowHandle is valid: The item just got added a parent so we want to dirty it
+		if (!ParentItem)
+		{
+			return Context.IsRowAvailable(ParentRowHandle);
+		}
+										
+		const FTedsOutlinerTreeItem* TedsParentItem = ParentItem->CastTo<FTedsOutlinerTreeItem>();
+
+		if (TedsParentItem)
+		{
+			// return true if the row handle of the parent item doesn't match what we are given, i.e the parent has changed
+			return TedsParentItem->GetRowHandle() != ParentRowHandle;
+		}
+
+		return false;
+	};
+}
 FTedsOutlinerImpl::FTedsOutlinerImpl(const FTedsOutlinerParams& InParams, ISceneOutlinerMode* InMode)
 	: CreationParams(InParams)
 	, CellWidgetPurposes(InParams.CellWidgetPurposes)
@@ -311,34 +362,6 @@ void FTedsOutlinerImpl::AppendExternalQueries(DataStorage::FQueryDescription& Ou
 	}
 }
 
-bool FTedsOutlinerImpl::HasItemParentChanged(DataStorage::RowHandle InRowHandle, DataStorage::RowHandle ParentRowHandle) const
-{
-	const FSceneOutlinerTreeItemPtr Item = SceneOutliner->GetTreeItem(InRowHandle, true);
-
-	// If the item doesn't exist, it doesn't make sense to say its parent changed
-	if (!Item)
-	{
-		return false;
-	}
-	
-	const FSceneOutlinerTreeItemPtr ParentItem = Item->GetParent();
-
-	// If the item doesn't have a parent, but ParentRowHandle is valid: The item just got added a parent so we want to dirty it
-	if (!ParentItem)
-	{
-		return Storage->IsRowAvailable(ParentRowHandle);
-	}
-	
-	const FTedsOutlinerTreeItem* TedsParentItem = ParentItem->CastTo<FTedsOutlinerTreeItem>();
-
-	if (TedsParentItem)
-	{
-		// return true if the row handle of the parent item doesn't match what we are given, i.e the parent has changed
-		return TedsParentItem->GetRowHandle() != ParentRowHandle;
-	}
-
-	return false;
-}
 
 bool FTedsOutlinerImpl::CanDisplayRow(DataStorage::RowHandle ItemRowHandle) const
 {
@@ -354,7 +377,7 @@ bool FTedsOutlinerImpl::CanDisplayRow(DataStorage::RowHandle ItemRowHandle) cons
 		if (const FTedsOutlinerColumn* TedsOutlinerColumn = Storage->GetColumn<FTedsOutlinerColumn>(ItemRowHandle))
 		{
 			if (const TSharedPtr<ISceneOutliner> TableViewer = TedsOutlinerColumn->Outliner.Pin())
-			{
+		{
 				return SceneOutliner != TableViewer.Get();
 			}
 		}
@@ -536,19 +559,6 @@ void FTedsOutlinerImpl::OnItemRemoved(DataStorage::RowHandle ItemRowHandle)
 	HierarchyChangedEvent.Broadcast(EventData);
 }
 
-void FTedsOutlinerImpl::OnItemMoved(DataStorage::RowHandle ItemRowHandle)
-{
-	if (!CanDisplayRow(ItemRowHandle))
-	{
-		return;
-	}
-	
-	FSceneOutlinerHierarchyChangedData EventData;
-	EventData.Type = FSceneOutlinerHierarchyChangedData::Moved;
-	EventData.ItemIDs.Add(ItemRowHandle);
-	HierarchyChangedEvent.Broadcast(EventData);
-}
-
 void FTedsOutlinerImpl::RecompileQueries()
 {
 	using namespace UE::Editor::DataStorage::Queries;
@@ -614,18 +624,32 @@ void FTedsOutlinerImpl::RecompileQueries()
 			TEXT("Update item parent"),
 			FProcessor(EQueryTickPhase::DuringPhysics, Storage->GetQueryTickGroupName(EQueryTickGroups::Update))
 				.SetExecutionMode(EExecutionMode::GameThread),
-			[this](IQueryContext& Context, DataStorage::RowHandle Row)
+			[this](IQueryContext& Context, const DataStorage::RowHandle* Rows, const FTedsOutlinerColumn* TedsOutlinerColumns)
 			{
-				DataStorage::RowHandle ParentRowHandle = InvalidRowHandle;
-
-				if (const FTableRowParentColumn* ParentColumn = Context.GetColumn<FTableRowParentColumn>())
-				{
-					ParentRowHandle = ParentColumn->Parent;
-				}
+				const FTableRowParentColumn* ParentColumnBegin = Context.GetColumn<FTableRowParentColumn>();
 				
-				if (HasItemParentChanged(Row, ParentRowHandle))
+				for (int32 RowIndex = 0; RowIndex < (int32)Context.GetRowCount(); ++RowIndex)
 				{
-					OnItemMoved(Row);
+					RowHandle Row = Rows[RowIndex];
+					DataStorage::RowHandle ParentRowHandle = InvalidRowHandle;
+					if (ParentColumnBegin)
+					{
+						ParentRowHandle = ParentColumnBegin[RowIndex].Parent;
+					}
+									
+					if (QueryUtils::HasItemParentChanged(Context, Row, ParentRowHandle, *SceneOutliner))
+					{
+						const FTedsOutlinerColumn& TedsOutlinerColumn = TedsOutlinerColumns[RowIndex];
+						if (!QueryUtils::CanDisplayRow(Context, TedsOutlinerColumn, Row, *SceneOutliner))
+						{
+							return;
+						}
+						
+						FSceneOutlinerHierarchyChangedData EventData;
+						EventData.Type = FSceneOutlinerHierarchyChangedData::Moved;
+						EventData.ItemIDs.Add(Row);
+						HierarchyChangedEvent.Broadcast(EventData);
+					}
 				}
 			})
 			.ReadOnly<FTableRowParentColumn>(EOptional::Yes)
diff --git a/Engine/Plugins/Experimental/EditorDataStorageFeatures/Source/TedsOutliner/Public/TedsOutlinerImpl.h b/Engine/Plugins/Experimental/EditorDataStorageFeatures/Source/TedsOutliner/Public/TedsOutlinerImpl.h
--- a/Engine/Plugins/Experimental/EditorDataStorageFeatures/Source/TedsOutliner/Public/TedsOutlinerImpl.h
+++ b/Engine/Plugins/Experimental/EditorDataStorageFeatures/Source/TedsOutliner/Public/TedsOutlinerImpl.h
@@ -153,9 +153,6 @@ public:
 	// Funtion to combine 2 queries (adds to second query to the first)
 	static void AppendQuery(DataStorage::FQueryDescription& Query1, const DataStorage::FQueryDescription& Query2);
 
-	// Check if the given item's parent has changed (i.e ParentRowHandle does not match what the Outliner reports as the parent)
-	bool HasItemParentChanged(DataStorage::RowHandle ItemRowHandle, DataStorage::RowHandle ParentRowHandle) const;
-
 	// Outliner specific functionality
 	void CreateItemsFromQuery(TArray<FSceneOutlinerTreeItemPtr>& OutItems, ISceneOutlinerMode* InMode) const;
 	void CreateChildren(const FSceneOutlinerTreeItemPtr& Item, TArray<FSceneOutlinerTreeItemPtr>& OutChildren) const;
@@ -170,7 +167,6 @@ protected:
 	
 	void OnItemAdded(DataStorage::RowHandle ItemRowHandle);
 	void OnItemRemoved(DataStorage::RowHandle ItemRowHandle);
-	void OnItemMoved(DataStorage::RowHandle ItemRowHandle);
 
 	void UnregisterQueries() const;
 	void ClearSelection() const;
