TedsSettingsManager : In UpdateSettingsCategory, always update section rows as their settings object may have changed.

Store an index for each row using a hash of the ContainerName + CategoryName + SectionName so that we can find rows even if they are pending or if their settings object has changed.

This fixes an issue where the same Container / Category / Section is registered twice with 2 different settings objects.  In this case, the end result should be 1 row in TEDS with the settings object from the second call to RegisterSettings.  Added a test case for regression testing.

diff --git a/Engine/Plugins/Experimental/EditorDataStorageFeatures/Source/TedsSettings/Private/TedsSettingsManager.cpp b/Engine/Plugins/Experimental/EditorDataStorageFeatures/Source/TedsSettings/Private/TedsSettingsManager.cpp
--- a/Engine/Plugins/Experimental/EditorDataStorageFeatures/Source/TedsSettings/Private/TedsSettingsManager.cpp
+++ b/Engine/Plugins/Experimental/EditorDataStorageFeatures/Source/TedsSettings/Private/TedsSettingsManager.cpp
@@ -19,6 +19,19 @@
 #include "TedsSettingsLog.h"
 #include "UObject/UnrealType.h"
 
+namespace UE::Editor::Settings::Private
+{
+	static UE::Editor::DataStorage::IndexHash GenerateIndexHash(const FName& ContainerName, const FName& CategoryName, const FName& SectionName)
+	{
+		constexpr static const char SeedName[] = "ISettingsSection";
+		static uint64 Seed = CityHash64(SeedName, sizeof(SeedName) - 1);
+
+		uint64 Hash = CityHash128to64({ Seed, UE::Editor::DataStorage::GenerateIndexHash(ContainerName) });
+		Hash = CityHash128to64({ Hash, UE::Editor::DataStorage::GenerateIndexHash(CategoryName) });
+		return CityHash128to64({ Hash, UE::Editor::DataStorage::GenerateIndexHash(SectionName) });
+	}
+}
+
 FTedsSettingsManager::FTedsSettingsManager()
 	: bIsInitialized{ false }
 	, SelectAllSettingsQuery{ UE::Editor::DataStorage::InvalidQueryHandle }
@@ -192,6 +205,7 @@ void FTedsSettingsManager::RegisterSettingsContainer(const FName& ContainerName)
 void FTedsSettingsManager::UnregisterSettings()
 {
 	using namespace UE::Editor::DataStorage;
+	using namespace UE::Editor::Settings::Private;
 
 	TRACE_CPUPROFILER_EVENT_SCOPE(TedsSettingsManager.UnregisterSettings);
 
@@ -237,6 +251,8 @@ void FTedsSettingsManager::UnregisterSettings()
 				{
 					DataStorageCompatibility->RemoveCompatibleObject(SettingsObjectPtr);
 
+					DataStorage->RemoveIndex(GenerateIndexHash(ContainerName, CategoryPtr->GetName(), SectionPtr->GetName()));
+
 					UE_LOG(LogTedsSettings, Log, TEXT("Removed Settings Section : '%s'"), *SectionPtr->GetName().ToString());
 				}
 			}
@@ -263,6 +279,7 @@ void FTedsSettingsManager::UnregisterSettings()
 void FTedsSettingsManager::UpdateSettingsCategory(TSharedPtr<ISettingsCategory> SettingsCategory, UE::Editor::DataStorage::RowHandle ContainerRow, const bool bQueryExistingRows)
 {
 	using namespace UE::Editor::DataStorage;
+	using namespace UE::Editor::Settings::Private;
 
 	TRACE_CPUPROFILER_EVENT_SCOPE(TedsSettingsManager.UpdateSettingsCategory);
 
@@ -333,7 +350,7 @@ void FTedsSettingsManager::UpdateSettingsCategory(TSharedPtr<ISettingsCategory>
 	const bool bIgnoreVisibility = true;
 	SettingsCategory->GetSections(NewSections, bIgnoreVisibility);
 
-	// Iterate the category and add rows for any section not in the old sections list.
+	// Iterate the category and add rows for all sections ( replace any existing row for the section as its object may have changed )
 	for (ISettingsSectionPtr SectionPtr : NewSections)
 	{
 		const FName& SectionName = SectionPtr->GetName();
@@ -342,21 +359,29 @@ void FTedsSettingsManager::UpdateSettingsCategory(TSharedPtr<ISettingsCategory>
 		{
 			NewSectionNames.Emplace(SectionName);
 
-			if (OldSectionNames.Contains(SectionName))
+			uint64 SectionIndexHash = GenerateIndexHash(ContainerName, CategoryName, SectionName);
+			
+			RowHandle OldSectionRow = DataStorage->FindIndexedRow(SectionIndexHash);
+			if (OldSectionRow != InvalidRowHandle)
 			{
 				UE_LOG(LogTedsSettings, Verbose, TEXT("Settings Section : '%s' is already in data storage"), *SectionName.ToString());
 
-				continue;
+				// Remove the row, the SettingsObjectPtr may have changed so we need to re-add the row with the new object
+				DataStorage->RemoveRow(OldSectionRow);
+
+				UE_LOG(LogTedsSettings, Log, TEXT("Removed Settings Section : '%s'"), *SectionName.ToString());
 			}
 
-			RowHandle NewRow = DataStorageCompatibility->AddCompatibleObject(SettingsObjectPtr);
+			RowHandle NewSectionRow = DataStorageCompatibility->AddCompatibleObject(SettingsObjectPtr);
 
-			DataStorage->AddColumn<FSettingsContainerReferenceColumn>(NewRow, { .ContainerName = ContainerName, .ContainerRow = ContainerRow });
-			DataStorage->AddColumn<FSettingsCategoryReferenceColumn>(NewRow, { .CategoryName = CategoryName, .CategoryRow = CategoryRow });
-			DataStorage->AddColumn<FNameColumn>(NewRow, { .Name = SectionName });
-			DataStorage->AddColumn<FDisplayNameColumn>(NewRow, { .DisplayName = SectionPtr->GetDisplayName() });
-			DataStorage->AddColumn<FDescriptionColumn>(NewRow, { .Description = SectionPtr->GetDescription() });
-			DataStorage->AddColumn<FSettingsSectionTag>(NewRow);
+			DataStorage->AddColumn<FSettingsSectionTag>(NewSectionRow);
+			DataStorage->AddColumn<FSettingsContainerReferenceColumn>(NewSectionRow, { .ContainerName = ContainerName, .ContainerRow = ContainerRow });
+			DataStorage->AddColumn<FSettingsCategoryReferenceColumn>(NewSectionRow, { .CategoryName = CategoryName, .CategoryRow = CategoryRow });
+			DataStorage->AddColumn<FNameColumn>(NewSectionRow, { .Name = SectionName });
+			DataStorage->AddColumn<FDisplayNameColumn>(NewSectionRow, { .DisplayName = SectionPtr->GetDisplayName() });
+			DataStorage->AddColumn<FDescriptionColumn>(NewSectionRow, { .Description = SectionPtr->GetDescription() });
+			
+			DataStorage->IndexRow(SectionIndexHash, NewSectionRow);
 
 			UE_LOG(LogTedsSettings, Log, TEXT("Added Settings Section : '%s'"), *SectionName.ToString());
 		}
@@ -377,6 +402,8 @@ void FTedsSettingsManager::UpdateSettingsCategory(TSharedPtr<ISettingsCategory>
 
 		DataStorage->RemoveRow(OldRowHandle);
 
+		DataStorage->RemoveIndex(GenerateIndexHash(ContainerName, CategoryName, OldSectionName));
+
 		UE_LOG(LogTedsSettings, Log, TEXT("Removed Settings Section : '%s'"), *OldSectionName.ToString());
 	}
 }
diff --git a/Engine/Plugins/Experimental/EditorDataStorageFeatures/Source/TedsSettings/Tests/TedsSettingsTests.cpp b/Engine/Plugins/Experimental/EditorDataStorageFeatures/Source/TedsSettings/Tests/TedsSettingsTests.cpp
--- a/Engine/Plugins/Experimental/EditorDataStorageFeatures/Source/TedsSettings/Tests/TedsSettingsTests.cpp
+++ b/Engine/Plugins/Experimental/EditorDataStorageFeatures/Source/TedsSettings/Tests/TedsSettingsTests.cpp
@@ -189,6 +189,52 @@ namespace UE::Editor::Settings::Tests
 				});
 			});
 		});
+
+		Describe("RegisterSettings same container/category/section twice with different UObjects", [this]()
+		{
+			LatentIt("Should result in only a single row in editor data storage", [this](const FDoneDelegate& Done)
+			{
+				const FName& ContainerName = FName(TEXT("TestContainer"));
+				const FName& CategoryName = FName(TEXT("TestCategory"));
+				const FName& SectionName = FName(TEXT("TestSection"));
+
+				TObjectPtr<UTestSettings> TestSettingsObject1 = NewObject<UTestSettings>();
+				TObjectPtr<UTestSettings> TestSettingsObject2 = NewObject<UTestSettings>();
+
+				SettingsModule->RegisterSettings(ContainerName, CategoryName, SectionName, FText(), FText(), TestSettingsObject1);
+				SettingsModule->RegisterSettings(ContainerName, CategoryName, SectionName, FText(), FText(), TestSettingsObject2);
+
+				DataStorage::RowHandle RowHandle = DataStorageCompatibility->FindRowWithCompatibleObject(TestSettingsObject1);
+				TestEqual(TEXT("RowHandle"), RowHandle, DataStorage::InvalidRowHandle);
+
+				if (RowHandle != DataStorage::InvalidRowHandle)
+				{
+					TestRowHandles.Push(RowHandle);
+				}
+
+				RowHandle = DataStorageCompatibility->FindRowWithCompatibleObject(TestSettingsObject2);
+				TestNotEqual(TEXT("RowHandle"), RowHandle, DataStorage::InvalidRowHandle);
+
+				if (RowHandle == DataStorage::InvalidRowHandle)
+				{
+					Done.Execute();
+					return;
+				}
+
+				TestRowHandles.Push(RowHandle);
+
+				AwaitRowHandleThenVerify(RowHandle, Done, [this, RowHandle, ContainerName, CategoryName, SectionName]()
+				{
+					uint32 AfterRowCount = CountSettingsRowsInDataStorage();
+
+					TestEqual(TEXT("RowCount"), AfterRowCount, BeforeRowCount + 1);
+
+					TestEqual(TEXT("ContainerName"), DataStorage->GetColumn<FSettingsContainerReferenceColumn>(RowHandle)->ContainerName, ContainerName);
+					TestEqual(TEXT("CategoryName"), DataStorage->GetColumn<FSettingsCategoryReferenceColumn>(RowHandle)->CategoryName, CategoryName);
+					TestEqual(TEXT("SectionName"), DataStorage->GetColumn<FNameColumn>(RowHandle)->Name, SectionName);
+				});
+			});
+		});
 	}
 } // namespace UE::Editor::Settings::Tests
 
