First pass at adding the asset data tags and values into teds.
All thedata can be populated via this cvar TEDS.AssetDataStorage.Metadata.

Stuff that doesn't work yet:
Missing attribute metadata for some asset types. (DataAsset that are base on a blueprint type and some types that are not created  by a factory or that doesn't have an asset definition for their asset type).
Update of assets will keep around the old tags and values that are no-longer present for now.
Perf is not there yet.

diff --git a/Engine/Plugins/Experimental/EditorDataStorageFeatures/Source/TedsAssetData/Private/CB/TagsMetadataCache.cpp b/Engine/Plugins/Experimental/EditorDataStorageFeatures/Source/TedsAssetData/Private/CB/TagsMetadataCache.cpp
new file mode 100644
--- /dev/null
+++ b/Engine/Plugins/Experimental/EditorDataStorageFeatures/Source/TedsAssetData/Private/CB/TagsMetadataCache.cpp
@@ -0,0 +1,93 @@
+// Copyright Epic Games, Inc. All Rights Reserved.
+
+#include "CB/TagsMetadataCache.h"
+
+#include "UObject/AssetRegistryTagsContext.h"
+#include "UObject/LinkerLoad.h"
+#include "UObject/NameTypes.h"
+#include "UObject/UnrealType.h"
+
+namespace UE::Editor::AssetData::Private
+{
+
+void FTagsMetadataCache::CacheClass(FTopLevelAssetPath InClassName)
+{
+	UClass* FoundClass = FindObject<UClass>(InClassName);
+
+	if (!FoundClass)
+	{
+		// Look for class redirectors
+		FString NewPath = FLinkerLoad::FindNewPathNameForClass(InClassName.ToString(), false);
+		if (!NewPath.IsEmpty())
+		{
+			FoundClass = FindObject<UClass>(nullptr, *NewPath);
+		}
+	}
+
+	if (FoundClass)
+	{
+		FClassPropertiesCache ClassProperties;
+		FTopLevelAssetPath ClassPath(FoundClass);
+
+		const UObject* CDO = FoundClass->GetDefaultObject();
+		FAssetRegistryTagsContextData TagsContext(CDO, EAssetRegistryTagsCaller::Uncategorized);
+		CDO->GetAssetRegistryTags(TagsContext);
+
+		TMap<FName, UObject::FAssetRegistryTagMetadata> TagsMetadata;
+		CDO->GetAssetRegistryTagMetadata(TagsMetadata);
+
+		ClassProperties.TagNameToCachedProperty.Reserve(TagsContext.Tags.Num());
+
+		for (const TPair<FName, UObject::FAssetRegistryTag>& TagPair : TagsContext.Tags)
+		{
+			TSharedPtr< UE::Editor::AssetData::FItemAttributeMetadata> TagCache = MakeShared<UE::Editor::AssetData::FItemAttributeMetadata>();
+			ClassProperties.TagNameToCachedProperty.Add(TagPair.Key, TagCache);
+			TagCache->TagType = TagPair.Value.Type;
+			TagCache->DisplayFlags = TagPair.Value.DisplayFlags;
+
+			if (const UObject::FAssetRegistryTagMetadata* TagMetaData = TagsMetadata.Find(TagPair.Key))
+			{
+				TagCache->DisplayName = TagMetaData->DisplayName;
+				TagCache->TooltipText = TagMetaData->TooltipText;
+				TagCache->Suffix = TagMetaData->Suffix;
+				TagCache->ImportantValue = TagMetaData->ImportantValue;
+			}
+			else
+			{
+				// If the tag name corresponds to a property name, use the property tooltip
+				const FProperty* Property = FindFProperty<FProperty>(FoundClass, TagPair.Key);
+				TagCache->TooltipText = Property ? Property->GetToolTipText() : FText::FromString(FName::NameToDisplayString(TagPair.Key.ToString(), false));
+			}
+
+			// Ensure a display name for this tag
+			if (TagCache->DisplayName.IsEmpty())
+			{
+				if (const FProperty* TagField = FindFProperty<FProperty>(FoundClass, TagPair.Key))
+				{
+					// Take the display name from the corresponding property if possible
+					TagCache->DisplayName = TagField->GetDisplayNameText();
+				}
+				else
+				{
+					// We have no type information by this point, so no idea if it's a bool :(
+					TagCache->DisplayName = FText::AsCultureInvariant(FName::NameToDisplayString(TagPair.Key.ToString(), /*bIsBool*/false));
+				}
+			}
+		}
+
+		ClassPathToCachedClass.Add(ClassPath, MoveTemp(ClassProperties));
+	}
+}
+
+const FTagsMetadataCache::FClassPropertiesCache* FTagsMetadataCache::FindCacheForClass(FTopLevelAssetPath InClassName) const
+{
+	return ClassPathToCachedClass.Find(InClassName);
+}
+
+TSharedPtr<UE::Editor::AssetData::FItemAttributeMetadata> FTagsMetadataCache::FClassPropertiesCache::GetCacheForTag(const FName InTagName) const
+{
+	const TSharedPtr<UE::Editor::AssetData::FItemAttributeMetadata>* PropertyCachePtr = TagNameToCachedProperty.Find(InTagName);
+	return PropertyCachePtr ? *PropertyCachePtr : TSharedPtr<UE::Editor::AssetData::FItemAttributeMetadata>();
+}
+
+}
\ No newline at end of file
diff --git a/Engine/Plugins/Experimental/EditorDataStorageFeatures/Source/TedsAssetData/Private/CB/TagsMetadataCache.h b/Engine/Plugins/Experimental/EditorDataStorageFeatures/Source/TedsAssetData/Private/CB/TagsMetadataCache.h
new file mode 100644
--- /dev/null
+++ b/Engine/Plugins/Experimental/EditorDataStorageFeatures/Source/TedsAssetData/Private/CB/TagsMetadataCache.h
@@ -0,0 +1,50 @@
+// Copyright Epic Games, Inc. All Rights Reserved.
+
+#pragma once
+
+#include "TedsAssetDataStructs.h"
+
+#include "Containers/Map.h"
+#include "Containers/Set.h"
+#include "Templates/SharedPointer.h"
+#include "UObject/TopLevelAssetPath.h"
+#include "UObject/NameTypes.h"
+
+namespace UE::Editor::AssetData::Private
+{
+class FTagsMetadataCache
+{
+public:
+
+	struct FClassPropertiesCache
+	{
+	public:
+		TSharedPtr<UE::Editor::AssetData::FItemAttributeMetadata> GetCacheForTag(const FName InTagName) const;
+
+	private:
+		friend class FTagsMetadataCache;
+
+		TMap<FName, TSharedPtr<UE::Editor::AssetData::FItemAttributeMetadata>> TagNameToCachedProperty;
+	};
+
+
+	void CacheClass(FTopLevelAssetPath InClassName);
+
+	const FClassPropertiesCache* FindCacheForClass(FTopLevelAssetPath InClassName) const;
+
+	template<typename TContainer>
+	void BatchCacheClasses(const TContainer& InClasses)
+	{
+		ClassPathToCachedClass.Reserve(ClassPathToCachedClass.Num() + InClasses.Num());
+
+		for (const FTopLevelAssetPath& ClassPath : InClasses)
+		{
+			CacheClass(ClassPath);
+		}
+	}
+
+	private:
+		TMap<FTopLevelAssetPath, FClassPropertiesCache> ClassPathToCachedClass;
+};
+
+}
\ No newline at end of file
diff --git a/Engine/Plugins/Experimental/EditorDataStorageFeatures/Source/TedsAssetData/Private/CB/TedsAssetDataCBDataSource.cpp b/Engine/Plugins/Experimental/EditorDataStorageFeatures/Source/TedsAssetData/Private/CB/TedsAssetDataCBDataSource.cpp
--- a/Engine/Plugins/Experimental/EditorDataStorageFeatures/Source/TedsAssetData/Private/CB/TedsAssetDataCBDataSource.cpp
+++ b/Engine/Plugins/Experimental/EditorDataStorageFeatures/Source/TedsAssetData/Private/CB/TedsAssetDataCBDataSource.cpp
@@ -4,63 +4,73 @@
 
 #include "TedsAssetDataColumns.h"
 
+#include "AssetDefinitionRegistry.h"
 #include "AssetRegistry/IAssetRegistry.h"
+#include "Blueprint/BlueprintSupport.h"
 #include "ContentBrowserDataUtils.h"
 #include "Elements/Common/TypedElementQueryTypes.h"
 #include "Elements/Framework/TypedElementQueryBuilder.h"
 #include "Elements/Interfaces/TypedElementDataStorageInterface.h"
+#include "Factories/Factory.h"
+#include "HAL/IConsoleManager.h"
 #include "Interfaces/IPluginManager.h"
 #include "Interfaces/IPluginManager.h"
+#include "Internationalization/Text.h"
 #include "Misc/PathViews.h"
 #include "PluginDescriptor.h"
 #include "Settings/ContentBrowserSettings.h"
+#include "TedsAssetDataModule.h"
+#include "Templates/Function.h"
+#include "UObject/CoreRedirects.h"
 #include "UObject/NameTypes.h"
+#include "UObject/ObjectMacros.h"
+#include "UObject/SoftObjectPtr.h"
+#include "UObject/TopLevelAssetPath.h"
+#include "UObject/UObjectIterator.h"
 
 namespace UE::Editor::AssetData::Private
 {
-	using namespace UE::Editor::DataStorage;
+
+TAutoConsoleVariable<bool> CVarTEDSAssetDataCBSourceIncludeTagsAndValues(TEXT("TEDS.AssetDataStorage.Metadata"), false, TEXT("When true we will add the meta data for the asset showable in the CB")
+	, FConsoleVariableDelegate::CreateLambda([](IConsoleVariable* Variable)
+	{
+		const bool bIsEnabled = Variable->GetBool();
+		FTedsAssetDataModule& Module = FTedsAssetDataModule::GetChecked();
+		if (bIsEnabled)
+		{
+			Module.EnableAssetDataMetadataStorage();
+		}
+		else
+		{
+			Module.DisableAssetDataMetadataStorage();
+		}
+	}));
 
 FTedsAssetDataCBDataSource::FTedsAssetDataCBDataSource(ITypedElementDataStorageInterface& InDatabase)
 	: Database(InDatabase)
 {
+	using namespace UE::Editor::DataStorage;
 	using namespace Queries;
 
+
+	bPopulateMetadataColumns = CVarTEDSAssetDataCBSourceIncludeTagsAndValues.GetValueOnGameThread();
+	AssetRegistry = IAssetRegistry::Get();
+
 	InitVirtualPathProcessor();
 
-	auto GenerateVirtualPaths = [DataSource = static_cast<const FTedsAssetDataCBDataSource*>(this)](const FStringView InAssetPath, FNameBuilder& OutVirtualizedPath) -> bool
+	if (bPopulateMetadataColumns)
 	{
-		if (!ContentBrowserDataUtils::PathPassesAttributeFilter(InAssetPath, 0, EContentBrowserItemAttributeFilter::IncludeAll))
-		{
-			return false;
-		}
-
-		DataSource->VirtualPathProcessor.ConvertInternalPathToVirtualPath(InAssetPath, OutVirtualizedPath);
-		return true;
-	};
+		TagsMetadataCache = MakeUnique<FTagsMetadataCache>();
+		PrepopulateTagsMetadataCache();
+	}
 
 	ProcessPathQuery = Database.RegisterQuery(
 		Select(
 			TEXT("FTedsAssetDataCBDataSource: Process Path updates"),
 			FProcessor(EQueryTickPhase::DuringPhysics, Database.GetQueryTickGroupName(EQueryTickGroups::Update)),
-			[GenerateVirtualPaths](IQueryContext& Context, const RowHandle* Rows, const FAssetPathColumn_Experimental* PathColumn)
+			[DataSource = const_cast<const FTedsAssetDataCBDataSource*>(this)](IQueryContext& Context, const RowHandle* Rows, const FAssetPathColumn_Experimental* PathColumn)
 			{
-				int32 NumOfRowToProcess = Context.GetRowCount();
-
-				FNameBuilder InternalPath;
-				FNameBuilder VirtualPath;
-
-				for (int32 Index = 0; Index < NumOfRowToProcess; ++Index)
-				{
-					PathColumn[Index].Path.ToString(InternalPath);
-
-					if (GenerateVirtualPaths(InternalPath, VirtualPath))
-					{
-						FVirtualPathColumn_Experimental VirtualPathColumn;
-						VirtualPathColumn.VirtualPath = *VirtualPath;
-						// Todo investigate for a batch add Maybe?
-						Context.AddColumn(Rows[Index], MoveTemp(VirtualPathColumn));
-					}
-				}
+				DataSource->ProcessPathQueryCallback(Context, Rows, PathColumn);
 			})
 		.Where()
 			.All<FUpdatedPathTag>()
@@ -70,138 +80,61 @@ FTedsAssetDataCBDataSource::FTedsAssetDataCBDataSource(ITypedElementDataStorageI
 		Select(
 			TEXT("FTedsAssetDataCBDataSource:: Process Asset Data Path Update"),
 			FProcessor(EQueryTickPhase::DuringPhysics, Database.GetQueryTickGroupName(EQueryTickGroups::Update)),
-			[GenerateVirtualPaths](IQueryContext& Context, const RowHandle Row, const FAssetDataColumn_Experimental& AssetDataColumn)
+			[DataSource = const_cast<const FTedsAssetDataCBDataSource*>(this)](IQueryContext& Context, const RowHandle* Row, const FAssetDataColumn_Experimental* AssetDataColumn)
 			{
-				int32 NumOfRowToProcess = Context.GetRowCount();
-
-				FNameBuilder InternalPath;
-				FNameBuilder VirtualPath;
-
-				AssetDataColumn.AssetData.AppendObjectPath(InternalPath);
-
-				if (GenerateVirtualPaths(InternalPath, VirtualPath))
-				{
-					FVirtualPathColumn_Experimental VirtualPathColumn;
-					VirtualPathColumn.VirtualPath = *VirtualPath;
-					// Todo investigate for a batch add Maybe?
-					Context.AddColumn(Row, MoveTemp(VirtualPathColumn));
-				}
+				DataSource->ProcessAssetDataPathUpdateQueryCallback(Context, Row, AssetDataColumn);
 			})
 		.Where()
-				.All<FUpdatedPathTag>()
+			.All<FUpdatedPathTag>()
 			.None<FUpdatedAssetDataTag>()
 		.Compile());
 
 
-	// For now just add the columns one by one but this should be rework to work in batch
-	auto AddAssetDataColumns = [](IQueryContext& InContext, RowHandle Row, const FAssetData& InAssetData, const FAssetPackageData* PackageData)
-	{
-		// Not optimized at all but we would like to have the data in sooner for testing purposes.
-		
-		if (InAssetData.HasAnyPackageFlags(PKG_NotExternallyReferenceable))
-		{
-			// Private Asset
-			InContext.AddColumns<FAssetTag, FPrivateAssetTag>(Row);
-		}
-		else
-		{
-			InContext.AddColumns<FAssetTag, FPublicAssetTag>(Row);
-		}
-
-		if (PackageData)
-		{
-			FDiskSizeColumn DiskSizeColumn;
-			DiskSizeColumn.DiskSize = PackageData->DiskSize;
-			InContext.AddColumn(Row, MoveTemp(DiskSizeColumn));
-		}
-
-		// Should add those or should we develop tooling to query the asset data via Teds?
-		FAssetClassColumn AssetClassColumn;
-		AssetClassColumn.ClassPath = InAssetData.AssetClassPath;
-		InContext.AddColumn(Row, MoveTemp(AssetClassColumn));
-
-		FItemNameColumn_Experimental ItemNameColumn;
-		ItemNameColumn.Name = InAssetData.AssetName;
-		InContext.AddColumn(Row, MoveTemp(ItemNameColumn));
-	};
-
 	ProcessAssetDataAndPathUpdateQuery = Database.RegisterQuery(
 		Select(
 			TEXT("FTedsAssetDataCBDataSource: Process Asset Data and Path Updates"),
 			FProcessor(EQueryTickPhase::DuringPhysics, Database.GetQueryTickGroupName(EQueryTickGroups::Update)),
-			[GenerateVirtualPaths, AddAssetDataColumns, AssetRegistry = static_cast<const IAssetRegistry*>(&IAssetRegistry::GetChecked())](IQueryContext& Context, const RowHandle* Rows, const FAssetDataColumn_Experimental* AssetDataColumn)
+			[DataSource = const_cast<const FTedsAssetDataCBDataSource*>(this)](IQueryContext& Context, const RowHandle* Rows, const FAssetDataColumn_Experimental* AssetDataColumn)
 			{
-				const int32 RowCount = Context.GetRowCount();
-
-				TArray<FName, TInlineAllocator<32>> PackageNames;
-				TArray<TPair<RowHandle, const FAssetData*>, TInlineAllocator<32>> RowsAndAssetData;
-				PackageNames.Reserve(RowCount);
-				RowsAndAssetData.Reserve(RowCount);
-
-				FNameBuilder InternalPath;
-				FNameBuilder VirtualPath;
-
-				for (int32 Index = 0; Index < RowCount; ++Index)
-				{ 
-					if (GenerateVirtualPaths(InternalPath, VirtualPath))
-					{
-						const FAssetData& AssetData = AssetDataColumn[Index].AssetData;
-						const RowHandle Row = Rows[Index];
-
-						InternalPath.Reset();
-						AssetData.AppendObjectPath(InternalPath);
-						FVirtualPathColumn_Experimental VirtualPathColumn;
-						VirtualPathColumn.VirtualPath = *VirtualPath;
-						// Todo investigate for a batch add Maybe?
-						Context.AddColumn(Row, MoveTemp(VirtualPathColumn));
-
-						PackageNames.Add(AssetData.PackageName);
-						RowsAndAssetData.Emplace(Row, &AssetData);
-					}
-				}
-
-				TArray<TOptional<FAssetPackageData>> AssetPackageDatas = AssetRegistry->GetAssetPackageDatasCopy(PackageNames);
-
-				for (int32 Index = 0; Index < AssetPackageDatas.Num(); ++Index)
-				{
-					const TPair<RowHandle, const FAssetData*>& Pair = RowsAndAssetData[Index];
-					const TOptional<FAssetPackageData>& AssetPackageData = AssetPackageDatas[Index];
-					AddAssetDataColumns(Context, Pair.Key, *Pair.Value, AssetPackageData.GetPtrOrNull());
-				}
-			}
-		)
+				DataSource->ProcessAssetDataAndPathUpdateQueryCallback(Context, Rows, AssetDataColumn);
+			})
 		.Where()
 			.All<FUpdatedAssetDataTag, FUpdatedPathTag>()
-		.Compile());
+			.Compile());
 
 	ProcessAssetDataUpdateQuery = Database.RegisterQuery(
 		Select(
 			TEXT("FTedsAssetDataCBDataSource: Process Asset Data updates"),
 			FProcessor(EQueryTickPhase::DuringPhysics, Database.GetQueryTickGroupName(EQueryTickGroups::Update)),
-			[AddAssetDataColumns,  AssetRegistry = static_cast<const IAssetRegistry*>(&IAssetRegistry::GetChecked())](IQueryContext& Context, const RowHandle* Rows, const FAssetDataColumn_Experimental* AssetDataColumn)
+			[DataSource = const_cast<const FTedsAssetDataCBDataSource*>(this)](IQueryContext& Context, const RowHandle* Rows, const FAssetDataColumn_Experimental* AssetDataColumn)
 			{
-				const int32 RowCount = Context.GetRowCount();
-				TArray<FName, TInlineAllocator<32>> PackageNames;
-				PackageNames.Reserve(RowCount);
-
-				for (int32 Index = 0; Index < RowCount; ++Index)
-				{
-					PackageNames.Add(AssetDataColumn[Index].AssetData.PackageName);
-				}
-
-				TArray<TOptional<FAssetPackageData>> AssetPackageDatas = AssetRegistry->GetAssetPackageDatasCopy(PackageNames);
-
-				for (int32 Index = 0; Index < RowCount; ++Index)
-				{
-					AddAssetDataColumns(Context, Rows[Index], AssetDataColumn[Index].AssetData, AssetPackageDatas[Index].GetPtrOrNull());
-				}
+				DataSource->ProcessAssetDataUpdateQueryCallback(Context, Rows, AssetDataColumn);
 			}
 		)
 		.Where()
-			.All<FUpdatedAssetDataTag>()
+			.All<FUpdatedAssetDataTag, FVirtualPathColumn_Experimental>()
 			.None<FUpdatedPathTag>()
-		.Compile());
+		.Compile()
+		);
+
 
+	ReprocessesAssetDataColumns = Database.RegisterQuery(
+		Select(
+			TEXT("FTedsAssetData: Remove Updated Asset Tag"),
+			FPhaseAmble(FPhaseAmble::ELocation::Preamble, EQueryTickPhase::PrePhysics)
+				.MakeActivatable(RepopulateAssetDataColumns),
+			[DataSource = const_cast<const FTedsAssetDataCBDataSource*>(this)](IQueryContext& Context, const RowHandle* Rows)
+			{
+				TConstArrayView<RowHandle> RowsArrayView(Rows, Context.GetRowCount());
+				if (!DataSource->bPopulateMetadataColumns)
+				{
+					Context.RemoveColumns<FItemTextAttributeColumn_Experimental, FItemStringAttributeColumn_Experimental>(RowsArrayView);
+				}
+				Context.AddColumns<FUpdatedAssetDataTag>(RowsArrayView);
+			})
+		.Where()
+			.All<FAssetTag>()
+			.Compile());
 }
 
 FTedsAssetDataCBDataSource::~FTedsAssetDataCBDataSource()
@@ -209,6 +142,7 @@ FTedsAssetDataCBDataSource::~FTedsAssetDataCBDataSource()
 	// Not needed on a editor shut down
 	if (!IsEngineExitRequested())
 	{
+		Database.UnregisterQuery(ReprocessesAssetDataColumns);
 		Database.UnregisterQuery(ProcessAssetDataUpdateQuery);
 		Database.UnregisterQuery(ProcessAssetDataAndPathUpdateQuery);
 		Database.UnregisterQuery(ProcessAssetDataPathUpdateQuery);
@@ -256,6 +190,310 @@ void FTedsAssetDataCBDataSource::OnPluginUnmounted(IPlugin& InPlugin)
 	VirtualPathProcessor.PluginNameToCachedData.Remove(InPlugin.GetName());
 }
 
+void FTedsAssetDataCBDataSource::EnableMetadataStorage(bool bEnable)
+{
+	if (bEnable != bPopulateMetadataColumns)
+	{
+		bPopulateMetadataColumns = bEnable;
+
+		if (bPopulateMetadataColumns)
+		{ 
+			TagsMetadataCache = MakeUnique<FTagsMetadataCache>();
+			PrepopulateTagsMetadataCache();
+		}
+		else
+		{
+			TagsMetadataCache.Reset();
+		}
+
+		// Make sure the CVar match the current state
+		CVarTEDSAssetDataCBSourceIncludeTagsAndValues.AsVariable()->Set(bEnable);
+
+		// Force a update of the asset data columns
+		Database.ActivateQueries(RepopulateAssetDataColumns);
+	}
+}
+
+bool FTedsAssetDataCBDataSource::GenerateVirtualPath(const FStringView InAssetPath, FNameBuilder& OutVirtualizedPath) const
+{
+	if (!ContentBrowserDataUtils::PathPassesAttributeFilter(InAssetPath, 0, EContentBrowserItemAttributeFilter::IncludeAll))
+	{
+		return false;
+	}
+
+	VirtualPathProcessor.ConvertInternalPathToVirtualPath(InAssetPath, OutVirtualizedPath);
+	return true;
+}
+
+void FTedsAssetDataCBDataSource::AddAssetDataColumns(UE::Editor::DataStorage::IQueryContext& Context, DataStorage::RowHandle Row, const FAssetData& AssetData, const FAssetPackageData* OptionalPackageData) const
+{
+	// For now just add the columns one by one but this should be rework to work in batch
+	// Not optimized at all but we would like to have the data in sooner for testing purposes.
+	if (AssetData.HasAnyPackageFlags(PKG_NotExternallyReferenceable))
+	{
+		// Private Asset
+		Context.AddColumns<FAssetTag, FPrivateAssetTag>(Row);
+	}
+	else
+	{
+		Context.AddColumns<FAssetTag, FPublicAssetTag>(Row);
+	}
+
+	if (OptionalPackageData)
+	{
+		FDiskSizeColumn DiskSizeColumn;
+		DiskSizeColumn.DiskSize = OptionalPackageData->DiskSize;
+		Context.AddColumn(Row, MoveTemp(DiskSizeColumn));
+	}
+
+	FAssetClassColumn AssetClassColumn;
+	AssetClassColumn.ClassPath = AssetData.AssetClassPath;
+	Context.AddColumn(Row, MoveTemp(AssetClassColumn));
+
+	FItemNameColumn_Experimental ItemNameColumn;
+	ItemNameColumn.Name = AssetData.AssetName;
+	Context.AddColumn(Row, MoveTemp(ItemNameColumn));
+
+	if (bPopulateMetadataColumns)
+	{
+		static const FTopLevelAssetPath BlueprintAssetClass = FTopLevelAssetPath(TEXT("/Script/Engine"), TEXT("Blueprint"));
+
+		// The population of the cache still need some work.
+		const FTagsMetadataCache::FClassPropertiesCache* ClassPropertyTagCache = TagsMetadataCache->FindCacheForClass(AssetData.AssetClassPath);
+		const FTagsMetadataCache::FClassPropertiesCache* ParentClassPropertyTagCache = nullptr;
+
+		if (!ClassPropertyTagCache)
+		{
+			FCoreRedirectObjectName RedirectedName = FCoreRedirects::GetRedirectedName(ECoreRedirectFlags::Type_Class, AssetData.AssetClassPath);
+			if (RedirectedName.IsValid())
+			{
+				ClassPropertyTagCache = TagsMetadataCache->FindCacheForClass(FTopLevelAssetPath(RedirectedName.PackageName, RedirectedName.ObjectName));
+			}
+		}
+
+
+		if (AssetData.AssetClassPath == BlueprintAssetClass)
+		{
+			// Non functional at the moment need to revisit the caching for those.
+			FAssetTagValueRef ParentClassRef = AssetData.TagsAndValues.FindTag(FBlueprintTags::ParentClassPath);
+			if (ParentClassRef.IsSet())
+			{
+				ParentClassPropertyTagCache = TagsMetadataCache->FindCacheForClass(ParentClassRef.AsExportPath().ToTopLevelAssetPath());
+			}
+
+			if (!ParentClassPropertyTagCache)
+			{
+				FAssetTagValueRef NativeParentClassRef = AssetData.TagsAndValues.FindTag(FBlueprintTags::NativeParentClassPath);
+				if (NativeParentClassRef.IsSet())
+				{
+					ParentClassPropertyTagCache = TagsMetadataCache->FindCacheForClass(NativeParentClassRef.AsExportPath().ToTopLevelAssetPath());
+				}
+			}
+		}
+
+		
+		for (const TPair<FName, FAssetTagValueRef>& TagAndValue : AssetData.TagsAndValues)
+		{ 
+			TSharedPtr<const UE::Editor::AssetData::FItemAttributeMetadata> AttributeMetadata;
+			if (ParentClassPropertyTagCache)
+			{
+				AttributeMetadata = ParentClassPropertyTagCache->GetCacheForTag(TagAndValue.Key);
+			}
+
+			if (!AttributeMetadata && ClassPropertyTagCache)
+			{
+				AttributeMetadata = ClassPropertyTagCache->GetCacheForTag(TagAndValue.Key);
+			}
+
+			// Todo revisit to see if we can save some memory here. 
+			FString TagValue = TagAndValue.Value.AsString();
+			bool bAddedColumn = false;
+			if (FTextStringHelper::IsComplexText(*TagValue))
+			{
+				FText TmpText;
+				if (FTextStringHelper::ReadFromBuffer(*TagValue, TmpText))
+				{
+					FItemTextAttributeColumn_Experimental AttributeColumn;
+					AttributeColumn.Value = MoveTemp(TmpText);
+					AttributeColumn.AttributeMetadata = MoveTemp(AttributeMetadata);
+					Context.AddColumn(Row, TagAndValue.Key, MoveTemp(AttributeColumn));
+					bAddedColumn = true;
+				}
+			}
+
+			if (!bAddedColumn)
+			{
+				FItemStringAttributeColumn_Experimental AttributeColumn;
+				AttributeColumn.Value = MoveTemp(TagValue);
+				AttributeColumn.AttributeMetadata = MoveTemp(AttributeMetadata);
+				Context.AddColumn(Row, TagAndValue.Key, MoveTemp(AttributeColumn));
+			}
+		}
+
+
+	}
+}
+
+void FTedsAssetDataCBDataSource::ProcessPathQueryCallback(DataStorage::IQueryContext& Context, const DataStorage::RowHandle* Rows, const FAssetPathColumn_Experimental* PathColumn) const
+{
+	int32 NumOfRowToProcess = Context.GetRowCount();
+
+	FNameBuilder InternalPath;
+	FNameBuilder VirtualPath;
+
+	for (int32 Index = 0; Index < NumOfRowToProcess; ++Index)
+	{
+		if (GenerateVirtualPath(InternalPath, VirtualPath))
+		{
+			FVirtualPathColumn_Experimental VirtualPathColumn;
+			VirtualPathColumn.VirtualPath = *VirtualPath;
+			// Todo investigate for a batch add Maybe?
+			Context.AddColumn(Rows[Index], MoveTemp(VirtualPathColumn));
+		}
+	}
+}
+
+void FTedsAssetDataCBDataSource::ProcessAssetDataPathUpdateQueryCallback(DataStorage::IQueryContext& Context, const DataStorage::RowHandle* Rows, const FAssetDataColumn_Experimental* AssetDataColumn) const
+{
+	int32 NumOfRowToProcess = Context.GetRowCount();
+
+	FNameBuilder InternalPath;
+	FNameBuilder VirtualPath;
+
+	for (int32 Index = 0; Index < NumOfRowToProcess; ++Index)
+	{ 
+		InternalPath.Reset();
+		AssetDataColumn[Index].AssetData.AppendObjectPath(InternalPath);
+		if (GenerateVirtualPath(InternalPath, VirtualPath))
+		{
+			FVirtualPathColumn_Experimental VirtualPathColumn;
+			VirtualPathColumn.VirtualPath = *VirtualPath;
+			// Todo investigate for a batch add Maybe?
+			Context.AddColumn(Rows[Index], MoveTemp(VirtualPathColumn));
+		}
+	}
+}
+
+void FTedsAssetDataCBDataSource::ProcessAssetDataAndPathUpdateQueryCallback(DataStorage::IQueryContext& Context, const DataStorage::RowHandle* Rows, const FAssetDataColumn_Experimental* AssetDataColumn) const
+{
+	const int32 RowCount = Context.GetRowCount();
+
+	TArray<FName, TInlineAllocator<32>> PackageNames;
+	TArray<TPair<DataStorage::RowHandle, const FAssetData*>, TInlineAllocator<32>> RowsAndAssetData;
+	PackageNames.Reserve(RowCount);
+	RowsAndAssetData.Reserve(RowCount);
+
+	FNameBuilder InternalPath;
+	FNameBuilder VirtualPath;
+
+	for (int32 Index = 0; Index < RowCount; ++Index)
+	{ 
+		const FAssetData& AssetData = AssetDataColumn[Index].AssetData;
+
+		if (ContentBrowserDataUtils::IsPrimaryAsset(AssetData))
+		{ 
+			InternalPath.Reset();
+			AssetData.AppendObjectPath(InternalPath);
+
+			if (GenerateVirtualPath(InternalPath, VirtualPath))
+			{
+				const DataStorage::RowHandle Row = Rows[Index];
+
+				FVirtualPathColumn_Experimental VirtualPathColumn;
+				VirtualPathColumn.VirtualPath = *VirtualPath;
+				// Todo investigate for a batch add Maybe?
+				Context.AddColumn(Row, MoveTemp(VirtualPathColumn));
+
+				PackageNames.Add(AssetData.PackageName);
+				RowsAndAssetData.Emplace(Row, &AssetData);
+			}
+		}
+	}
+
+	TArray<TOptional<FAssetPackageData>> AssetPackageDatas = AssetRegistry->GetAssetPackageDatasCopy(PackageNames);
+
+	for (int32 Index = 0; Index < RowsAndAssetData.Num(); ++Index)
+	{
+		const TPair<DataStorage::RowHandle, const FAssetData*>& Pair = RowsAndAssetData[Index];
+		const TOptional<FAssetPackageData>& AssetPackageData = AssetPackageDatas[Index];
+		AddAssetDataColumns(Context, Pair.Key, *Pair.Value, AssetPackageData.GetPtrOrNull());
+	}
+}
+
+void FTedsAssetDataCBDataSource::ProcessAssetDataUpdateQueryCallback(DataStorage::IQueryContext& Context, const DataStorage::RowHandle* Rows, const FAssetDataColumn_Experimental* AssetDataColumn) const
+{
+	const int32 RowCount = Context.GetRowCount();
+	TArray<FName, TInlineAllocator<32>> PackageNames;
+	PackageNames.Reserve(RowCount);
+
+	for (int32 Index = 0; Index < RowCount; ++Index)
+	{
+		PackageNames.Add(AssetDataColumn[Index].AssetData.PackageName);
+	}
+
+	TArray<TOptional<FAssetPackageData>> AssetPackageDatas = AssetRegistry->GetAssetPackageDatasCopy(PackageNames);
+
+	for (int32 Index = 0; Index < RowCount; ++Index)
+	{
+		AddAssetDataColumns(Context, Rows[Index], AssetDataColumn[Index].AssetData, AssetPackageDatas[Index].GetPtrOrNull());
+	}
+}
+
+void FTedsAssetDataCBDataSource::PrepopulateTagsMetadataCache()
+{
+	if (FTagsMetadataCache* TagsMetadata = TagsMetadataCache.Get())
+	{
+		TSet<FTopLevelAssetPath> ClassesPath;
+
+		// Try to populate the tags meta cache to avoid some costly operations when doing the initial population by using the knows asset types
+		TArray<TSoftClassPtr<UObject>> ClassesWithAssetDefinition = UAssetDefinitionRegistry::Get()->GetAllRegisteredAssetClasses();
+		ClassesPath.Reserve(ClassesWithAssetDefinition.Num());
+
+		TArray<UClass*> ChildrenClasses;
+
+		ClassesWithAssetDefinition.Contains(TSoftClassPtr<UObject>(FSoftObjectPath(TEXT("/Script/Engine.UserDefinedStruct"))));
+		for (const TSoftClassPtr<UObject>& SoftClass : ClassesWithAssetDefinition)
+		{
+			// Check 
+			if (UClass* Class = SoftClass.Get())
+			{
+				if (!Class->HasAllClassFlags(EClassFlags::CLASS_Abstract))
+				{
+					ClassesPath.Add(Class->GetClassPathName());
+				}
+
+				ChildrenClasses.Reset();
+				GetDerivedClasses(Class, ChildrenClasses);
+				
+				for (UClass* ChildClass : ChildrenClasses)
+				{
+					if (!ChildClass->HasAllClassFlags(EClassFlags::CLASS_Abstract))
+					{
+						ClassesPath.Add(ChildClass->GetClassPathName());
+					}
+				}
+			}
+		}
+
+		// Also check the factories for asset types
+		for (TObjectIterator<UClass> It; It; ++It)
+		{
+			UClass* Class = *It;
+			if (!Class->IsChildOf<UFactory>() || Class->HasAnyClassFlags(CLASS_Abstract))
+			{
+				continue;
+			}
+
+			if (UClass* AssetClass = Class->GetDefaultObject<UFactory>()->GetSupportedClass())
+			{
+				ClassesPath.Add(AssetClass->GetClassPathName());
+			}
+		}
+
+		TagsMetadata->BatchCacheClasses(ClassesPath);
+	}
+}
+
 void FTedsAssetDataCBDataSource::FVirtualPathProcessor::ConvertInternalPathToVirtualPath(const FStringView InternalPath, FStringBuilderBase& OutVirtualPath) const
 {
 	OutVirtualPath.Reset();
diff --git a/Engine/Plugins/Experimental/EditorDataStorageFeatures/Source/TedsAssetData/Private/CB/TedsAssetDataCBDataSource.h b/Engine/Plugins/Experimental/EditorDataStorageFeatures/Source/TedsAssetData/Private/CB/TedsAssetDataCBDataSource.h
--- a/Engine/Plugins/Experimental/EditorDataStorageFeatures/Source/TedsAssetData/Private/CB/TedsAssetDataCBDataSource.h
+++ b/Engine/Plugins/Experimental/EditorDataStorageFeatures/Source/TedsAssetData/Private/CB/TedsAssetDataCBDataSource.h
@@ -2,56 +2,94 @@
 
 #pragma once
 
+#include "CB/TagsMetadataCache.h"
+#include "Containers/StringFwd.h"
+#include "Containers/UnrealString.h"
 #include "Elements/Common/TypedElementCommonTypes.h"
 #include "Elements/Common/TypedElementHandles.h"
-
-#include "Containers/StringFwd.h"
 #include "Interfaces/IPluginManager.h"
 
+class FAssetPackageData;
+class FAssetPropertyTagCache;
+class IAssetRegistry;
 class IPlugin;
 class ITypedElementDataStorageInterface;
 
-namespace UE::Editor
+struct FAssetData;
+struct FAssetDataColumn_Experimental;
+struct FAssetPathColumn_Experimental;
+
+namespace UE::Editor::DataStorage
 {
-	namespace AssetData::Private
-	{
+	struct IQueryContext;
+}
+
+namespace UE::Editor::AssetData::Private
+{
+
+class FTedsAssetDataCBDataSource
+{
+public:
+	explicit FTedsAssetDataCBDataSource(ITypedElementDataStorageInterface& InDatabase);
+	~FTedsAssetDataCBDataSource();
+	
+	// Toggle the storage of the metadata from the asset registry into teds 
+	void EnableMetadataStorage(bool bEnable);
 
-		class FTedsAssetDataCBDataSource
+private:
+	struct FVirtualPathProcessor
+	{
+		struct FCachedPluginData
 		{
-		public:
-			explicit FTedsAssetDataCBDataSource(ITypedElementDataStorageInterface& InDatabase);
-			~FTedsAssetDataCBDataSource();
+			EPluginLoadedFrom LoadedFrom;
+			FString EditorCustomVirtualPath;
+		};
 
-		private:
-			ITypedElementDataStorageInterface& Database;
-			DataStorage::QueryHandle ProcessPathQuery;
-			DataStorage::QueryHandle ProcessAssetDataPathUpdateQuery;
-			DataStorage::QueryHandle ProcessAssetDataAndPathUpdateQuery;
-			DataStorage::QueryHandle ProcessAssetDataUpdateQuery;
+		TMap<FString, FCachedPluginData> PluginNameToCachedData;
 
-			struct FVirtualPathProcessor
-			{
-				struct FCachedPluginData
-				{
-					EPluginLoadedFrom LoadedFrom;
-					FString EditorCustomVirtualPath;
-				};
+		bool bShowAllFolder = false;
+		bool bOrganizeFolders = false;
 
-				TMap<FString, FCachedPluginData> PluginNameToCachedData;
+		void ConvertInternalPathToVirtualPath(const FStringView InternalPath, FStringBuilderBase& OutVirtualPath) const;
+	};
 
-				bool bShowAllFolder = false;
-				bool bOrganizeFolders = false;
+	void InitVirtualPathProcessor();
+	void PrepopulateTagsMetadataCache();
 
-				void ConvertInternalPathToVirtualPath(const FStringView InternalPath, FStringBuilderBase& OutVirtualPath) const;
-			};
 
-			FVirtualPathProcessor VirtualPathProcessor;
+	void OnPluginContentMounted(IPlugin& InPlugin);
+	void OnPluginUnmounted(IPlugin& InPlugin);
 
-			void InitVirtualPathProcessor();
+	/**
+	 * Check if we should generate a virtual a path to the path and if so do generate a virtualized path for it.
+	 * @param InAssetPath The internal path to virtualize
+	 * @param OutVirtualizedPath The virtualized path generated if needed.
+	 * @return True if the path virtual path was generated.
+	 */
+	bool GenerateVirtualPath(const FStringView InAssetPath, FNameBuilder& OutVirtualizedPath) const;
 
-			void OnPluginContentMounted(IPlugin& InPlugin);
-			void OnPluginUnmounted(IPlugin& InPlugin);
-		};
+	void AddAssetDataColumns(DataStorage::IQueryContext& Context, DataStorage::RowHandle Row, const FAssetData& AssetData, const FAssetPackageData* OptionalPackageData) const;
+
+	void ProcessPathQueryCallback(DataStorage::IQueryContext& Context, const DataStorage::RowHandle* Rows, const FAssetPathColumn_Experimental* PathColumn) const;
+	void ProcessAssetDataPathUpdateQueryCallback(DataStorage::IQueryContext& Context, const DataStorage::RowHandle* Rows, const FAssetDataColumn_Experimental* AssetDataColumn) const;
+	void ProcessAssetDataAndPathUpdateQueryCallback(DataStorage::IQueryContext& Context, const DataStorage::RowHandle* Rows, const FAssetDataColumn_Experimental* AssetDataColumn) const;
+	void ProcessAssetDataUpdateQueryCallback(DataStorage::IQueryContext& Context, const DataStorage::RowHandle*Rows, const FAssetDataColumn_Experimental* AssetDataColumn) const;
+
+	ITypedElementDataStorageInterface& Database;
+	DataStorage::QueryHandle ProcessPathQuery;
+	DataStorage::QueryHandle ProcessAssetDataPathUpdateQuery;
+	DataStorage::QueryHandle ProcessAssetDataAndPathUpdateQuery;
+	DataStorage::QueryHandle ProcessAssetDataUpdateQuery;
+	DataStorage::QueryHandle ReprocessesAssetDataColumns;
+
+	FVirtualPathProcessor VirtualPathProcessor; 
+	bool bPopulateMetadataColumns = false;
+	FName RepopulateAssetDataColumns = TEXT("RepopulateAssetDataColumnsQuery");
+
+	TUniquePtr<FTagsMetadataCache> TagsMetadataCache;
+
+	const IAssetRegistry* AssetRegistry = nullptr;
+
+};
 
-	} // namespace AssetData::Private
 } // End of namespace UE::Editor
diff --git a/Engine/Plugins/Experimental/EditorDataStorageFeatures/Source/TedsAssetData/Private/TedsAssetData.cpp b/Engine/Plugins/Experimental/EditorDataStorageFeatures/Source/TedsAssetData/Private/TedsAssetData.cpp
--- a/Engine/Plugins/Experimental/EditorDataStorageFeatures/Source/TedsAssetData/Private/TedsAssetData.cpp
+++ b/Engine/Plugins/Experimental/EditorDataStorageFeatures/Source/TedsAssetData/Private/TedsAssetData.cpp
@@ -182,7 +182,7 @@ FTedsAssetData::FTedsAssetData(ITypedElementDataStorageInterface& InDatabase)
 	PathsTable = Database.FindTable(FName(TEXT("Editor_AssetRegistryPathsTable")));
 	if (PathsTable == InvalidTableHandle)
 	{
-		PathsTable = Database.RegisterTable<FAssetPathColumn_Experimental, FChildrenAssetPathColumn_Experimental, FParentAssetPathColumn_Experimental, FAssetsInPathColumn_Experimental>(FName(TEXT("Editor_AssetRegistryPathsTable")));
+		PathsTable = Database.RegisterTable<FAssetPathColumn_Experimental, FChildrenAssetPathColumn_Experimental, FParentAssetPathColumn_Experimental, FAssetsInPathColumn_Experimental, FUpdatedPathTag>(FName(TEXT("Editor_AssetRegistryPathsTable")));
 	}
 
 	AssetsDataTable = Database.FindTable(FName(TEXT("Editor_AssetRegistryAssetDataTable")));
diff --git a/Engine/Plugins/Experimental/EditorDataStorageFeatures/Source/TedsAssetData/Private/TedsAssetDataModule.cpp b/Engine/Plugins/Experimental/EditorDataStorageFeatures/Source/TedsAssetData/Private/TedsAssetDataModule.cpp
--- a/Engine/Plugins/Experimental/EditorDataStorageFeatures/Source/TedsAssetData/Private/TedsAssetDataModule.cpp
+++ b/Engine/Plugins/Experimental/EditorDataStorageFeatures/Source/TedsAssetData/Private/TedsAssetDataModule.cpp
@@ -114,4 +114,20 @@ void FTedsAssetDataModule::InitAssetRegistryStorage()
 	AssetRegistryStorage = MakeUnique<Private::FTedsAssetData>(MutableDataStorage);
 }
 
-} // namespace UE::Editor::AssetData
\ No newline at end of file
+void FTedsAssetDataModule::EnableAssetDataMetadataStorage()
+{
+	if (AssetDataCBDataSource)
+	{
+		AssetDataCBDataSource->EnableMetadataStorage(true);
+	}
+}
+
+void FTedsAssetDataModule::DisableAssetDataMetadataStorage()
+{
+	if (AssetDataCBDataSource)
+	{
+		AssetDataCBDataSource->EnableMetadataStorage(false);
+	}
+}
+
+}// namespace UE::Editor::AssetData
diff --git a/Engine/Plugins/Experimental/EditorDataStorageFeatures/Source/TedsAssetData/Private/TedsAssetDataModule.h b/Engine/Plugins/Experimental/EditorDataStorageFeatures/Source/TedsAssetData/Private/TedsAssetDataModule.h
--- a/Engine/Plugins/Experimental/EditorDataStorageFeatures/Source/TedsAssetData/Private/TedsAssetDataModule.h
+++ b/Engine/Plugins/Experimental/EditorDataStorageFeatures/Source/TedsAssetData/Private/TedsAssetDataModule.h
@@ -22,11 +22,11 @@ public:
 
 	static FTedsAssetDataModule* Get();
 	static FTedsAssetDataModule& GetChecked();
-
 	void EnableTedsAssetRegistryStorage();
 	void DisableTedsAssetRegistryStorage();
 	bool IsTedsAssetRegistryStorageEnabled() const;
-
+	void EnableAssetDataMetadataStorage();
+	void DisableAssetDataMetadataStorage();
 
 	/**
 	 * Process now any pending event that might make the Teds database out of sync with the asset registry.
diff --git a/Engine/Plugins/Experimental/EditorDataStorageFeatures/Source/TedsAssetData/Public/TedsAssetDataColumns.h b/Engine/Plugins/Experimental/EditorDataStorageFeatures/Source/TedsAssetData/Public/TedsAssetDataColumns.h
--- a/Engine/Plugins/Experimental/EditorDataStorageFeatures/Source/TedsAssetData/Public/TedsAssetDataColumns.h
+++ b/Engine/Plugins/Experimental/EditorDataStorageFeatures/Source/TedsAssetData/Public/TedsAssetDataColumns.h
@@ -2,10 +2,16 @@
 
 #pragma once
 
+#include "TedsAssetDataStructs.h"
+
 #include "AssetRegistry/AssetData.h"
+#include "Containers/UnrealString.h"
+#include "Containers/VersePath.h"
 #include "Elements/Common/TypedElementCommonTypes.h"
 #include "Elements/Interfaces/TypedElementQueryStorageInterfaces.h"
-#include "Containers/VersePath.h"
+#include "Internationalization/Text.h"
+#include "Templates/SharedPointer.h"
+#include "UObject/NameTypes.h"
 
 #include "TedsAssetDataColumns.generated.h"
 
@@ -150,4 +156,30 @@ struct FItemNameColumn_Experimental : public FEditorDataStorageColumn
 
 	UPROPERTY(meta = (Searchable))
 	FName Name;
+};
+
+USTRUCT()
+struct FItemAttributeBaseColumn_Experimental : public FEditorDataStorageColumn
+{
+	GENERATED_BODY()
+
+	TSharedPtr<const UE::Editor::AssetData::FItemAttributeMetadata> AttributeMetadata;
+};
+
+USTRUCT()
+struct FItemTextAttributeColumn_Experimental : public FItemAttributeBaseColumn_Experimental
+{
+	GENERATED_BODY()
+
+	UPROPERTY()
+	FText Value;
+};
+
+USTRUCT()
+struct FItemStringAttributeColumn_Experimental : public FItemAttributeBaseColumn_Experimental
+{
+	GENERATED_BODY()
+
+	UPROPERTY()
+	FString Value;
 };
\ No newline at end of file
diff --git a/Engine/Plugins/Experimental/EditorDataStorageFeatures/Source/TedsAssetData/Public/TedsAssetDataStructs.h b/Engine/Plugins/Experimental/EditorDataStorageFeatures/Source/TedsAssetData/Public/TedsAssetDataStructs.h
new file mode 100644
--- /dev/null
+++ b/Engine/Plugins/Experimental/EditorDataStorageFeatures/Source/TedsAssetData/Public/TedsAssetDataStructs.h
@@ -0,0 +1,33 @@
+// Copyright Epic Games, Inc. All Rights Reserved.
+
+#pragma once
+
+#include "UObject/Object.h"
+#include "Containers/UnrealString.h"
+#include "Internationalization/Text.h"
+
+namespace UE::Editor::AssetData
+{ 
+
+struct FItemAttributeMetadata
+{
+	/** The kind of data represented by this tag value */
+	UObject::FAssetRegistryTag::ETagType TagType = UObject::FAssetRegistryTag::TT_Hidden;
+
+	/** Flags giving hints at how to display this tag value in the UI (see ETagDisplay) */
+	uint32 DisplayFlags = UObject::FAssetRegistryTag::TD_None;
+
+	/** Resolved display name of the associated tag */
+	FText DisplayName;
+
+	/** Optional tooltip of the associated tag */
+	FText TooltipText;
+
+	/** Optional suffix to apply to values of the tag attribute in the UI */
+	FText Suffix;
+
+	/** Optional value which denotes which values should be considered "important" in the UI */
+	FString ImportantValue;
+};
+
+}
diff --git a/Engine/Plugins/Experimental/EditorDataStorageFeatures/Source/TedsAssetData/TedsAssetData.build.cs b/Engine/Plugins/Experimental/EditorDataStorageFeatures/Source/TedsAssetData/TedsAssetData.build.cs
--- a/Engine/Plugins/Experimental/EditorDataStorageFeatures/Source/TedsAssetData/TedsAssetData.build.cs
+++ b/Engine/Plugins/Experimental/EditorDataStorageFeatures/Source/TedsAssetData/TedsAssetData.build.cs
@@ -19,15 +19,17 @@ public class TedsAssetData : ModuleRules
 			PrivateDependencyModuleNames.AddRange(
 				new string[]
 				{
+					"AssetDefinition",
 					"AssetRegistry",
 					"ContentBrowserData",
 					"Core",
 					"CoreUObject",
 					"Engine",
+					"Slate",
+					"SlateCore",
 					"TypedElementFramework",
 					"UnrealEd",
-					"SlateCore",
-					"Slate"
+					"UnrealEd",
 				});
 
 			PublicDependencyModuleNames.AddRange(
diff --git a/Engine/Source/Editor/AssetDefinition/Private/AssetDefinitionRegistry.cpp b/Engine/Source/Editor/AssetDefinition/Private/AssetDefinitionRegistry.cpp
--- a/Engine/Source/Editor/AssetDefinition/Private/AssetDefinitionRegistry.cpp
+++ b/Engine/Source/Editor/AssetDefinition/Private/AssetDefinitionRegistry.cpp
@@ -68,6 +68,13 @@ TArray<TObjectPtr<UAssetDefinition>> UAssetDefinitionRegistry::GetAllAssetDefini
 	return AllAssetDefinitions;
 }
 
+TArray<TSoftClassPtr<UObject>> UAssetDefinitionRegistry::GetAllRegisteredAssetClasses() const
+{
+	TArray<TSoftClassPtr<UObject>> AllRegisteredClasses;
+	AssetDefinitions.GenerateKeyArray(AllRegisteredClasses);
+	return AllRegisteredClasses;
+}
+
 uint64 UAssetDefinitionRegistry::GetAssetDefinitionVersion() const
 {
 	return Version;
diff --git a/Engine/Source/Editor/AssetDefinition/Public/AssetDefinitionRegistry.h b/Engine/Source/Editor/AssetDefinition/Public/AssetDefinitionRegistry.h
--- a/Engine/Source/Editor/AssetDefinition/Public/AssetDefinitionRegistry.h
+++ b/Engine/Source/Editor/AssetDefinition/Public/AssetDefinitionRegistry.h
@@ -30,6 +30,7 @@ public:
 	// Gets the current version of the AssetDefinitions.   Version is updated whenever an AssetDefinition is Registered/Unregistered
 	uint64 GetAssetDefinitionVersion() const;
 	TArray<TObjectPtr<UAssetDefinition>> GetAllAssetDefinitions() const;
+	TArray<TSoftClassPtr<UObject>> GetAllRegisteredAssetClasses() const;
 
 	/**
 	 * Normally UAssetDefinitionRegistry are registered automatically by their CDO.  The only reason you need to do this is if
diff --git a/Engine/Source/Editor/ContentBrowserData/Private/ContentBrowserDataUtils.cpp b/Engine/Source/Editor/ContentBrowserData/Private/ContentBrowserDataUtils.cpp
--- a/Engine/Source/Editor/ContentBrowserData/Private/ContentBrowserDataUtils.cpp
+++ b/Engine/Source/Editor/ContentBrowserData/Private/ContentBrowserDataUtils.cpp
@@ -2,6 +2,7 @@
 
 #include "ContentBrowserDataUtils.h"
 
+#include "AssetRegistry/AssetData.h"
 #include "Containers/StringView.h"
 #include "Interfaces/IPluginManager.h"
 #include "Internationalization/Internationalization.h"
@@ -279,4 +280,32 @@ FText ContentBrowserDataUtils::GetFolderItemDisplayNameOverride(const FName InFo
 	return FolderDisplayNameOverride;
 }
 
+bool ContentBrowserDataUtils::IsPrimaryAsset(const FAssetData& InAssetData)
+{
+	// Both GetOptionalOuterPathName and IsUAsset currently do not work on cooked assets
+	//
+	// GetOptionalOuterPathName is not serialized to the asset registry during cook
+	// IsUAsset when called on compiled blueprint class compares Name_C vs Name and returns false
+	if (InAssetData.HasAnyPackageFlags(PKG_Cooked | PKG_FilterEditorOnly))
+	{
+		// Check for the asset being a redirector first, as currently only class 
+		// redirectors emit non-primary assets from the Asset Registry
+		return !InAssetData.IsRedirector() || InAssetData.IsUAsset();
+	}
+	else
+	{
+		// External assets are not displayed in the Content Browser or other asset pickers
+		bool bIsExternalAsset = !InAssetData.GetOptionalOuterPathName().IsNone();
+		return !bIsExternalAsset && InAssetData.IsUAsset();
+	}
+}
+
+bool ContentBrowserDataUtils::IsPrimaryAsset(UObject* InObject)
+{
+		// External assets are not displayed in the Content Browser or other asset pickers
+	const bool bIsExternalAsset = InObject->IsPackageExternal();
+	
+	return !bIsExternalAsset && FAssetData::IsUAsset(InObject);
+}
+
 #undef LOCTEXT_NAMESPACE
diff --git a/Engine/Source/Editor/ContentBrowserData/Public/ContentBrowserDataUtils.h b/Engine/Source/Editor/ContentBrowserData/Public/ContentBrowserDataUtils.h
--- a/Engine/Source/Editor/ContentBrowserData/Public/ContentBrowserDataUtils.h
+++ b/Engine/Source/Editor/ContentBrowserData/Public/ContentBrowserDataUtils.h
@@ -10,6 +10,10 @@
 #include "Internationalization/Text.h"
 #include "UObject/NameTypes.h"
 
+class UObject;
+
+struct FAssetData;
+
 namespace ContentBrowserDataUtils
 {
 	/** Returns number of folders in forward slashed path (Eg, 1 for "/Path", 2 for "/Path/Name") */
@@ -46,4 +50,8 @@ namespace ContentBrowserDataUtils
 	 * @return Override display name or empty string
 	 */
 	CONTENTBROWSERDATA_API FText GetFolderItemDisplayNameOverride(const FName InFolderPath, const FString& InFolderItemName, const bool bIsClassesFolder, const bool bIsCookedPath = false);
+
+	CONTENTBROWSERDATA_API bool IsPrimaryAsset(const FAssetData& InAssetData);
+
+	CONTENTBROWSERDATA_API bool IsPrimaryAsset(UObject* InObject);
 }
