Add a new TedsStyling module to store stylesheet contents in TEDS

diff --git a/Engine/Plugins/Experimental/EditorDataStorage/Source/TedsUI/Private/Widgets/SlateColorWidget.cpp b/Engine/Plugins/Experimental/EditorDataStorage/Source/TedsUI/Private/Widgets/SlateColorWidget.cpp
new file mode 100644
--- /dev/null
+++ b/Engine/Plugins/Experimental/EditorDataStorage/Source/TedsUI/Private/Widgets/SlateColorWidget.cpp
@@ -0,0 +1,87 @@
+// Copyright Epic Games, Inc. All Rights Reserved.
+
+#include "Widgets/SlateColorWidget.h"
+
+#include "Elements/Columns/TypedElementCompatibilityColumns.h"
+#include "Elements/Columns/TypedElementMiscColumns.h"
+#include "Elements/Columns/TypedElementSlateWidgetColumns.h"
+#include "Elements/Framework/TypedElementAttributeBinding.h"
+#include "Internationalization/Text.h"
+#include "Widgets/Colors/SColorBlock.h"
+#include "Widgets/Text/STextBlock.h"
+#include "Widgets/Colors/SColorPicker.h"
+
+namespace UE::SlateColorWidget::Private
+{
+	FReply SummonColorPicker(const FGeometry& Geometry, const FPointerEvent& PointerEvent,
+		ITypedElementDataStorageInterface* DataStorage, UE::Editor::DataStorage::RowHandle TargetRow)
+	{
+		if (PointerEvent.GetEffectingButton() != EKeys::LeftMouseButton)
+		{
+			return FReply::Unhandled();
+		}
+	
+		FColorPickerArgs PickerArgs;
+		PickerArgs.bUseAlpha = true;
+		PickerArgs.OnColorCommitted = FOnLinearColorValueChanged::CreateLambda([TargetRow, DataStorage](FLinearColor Color)
+		{
+			DataStorage->GetColumn<FSlateColorColumn>(TargetRow)->Color = Color;
+			DataStorage->AddColumn(TargetRow, FTypedElementSyncBackToWorldTag::StaticStruct());
+		});
+
+		FSlateColor ColumnColor = DataStorage->GetColumn<FSlateColorColumn>(TargetRow)->Color;
+		// If the color isn't set to a specified value, we just show it as white by default in the picker
+		FLinearColor InitColor = ColumnColor.IsColorSpecified() ? ColumnColor.GetSpecifiedColor() : FLinearColor::White;
+	
+		PickerArgs.InitialColor = InitColor;
+		PickerArgs.bOnlyRefreshOnOk = true;
+		OpenColorPicker(PickerArgs);
+		return FReply::Handled();
+	}
+}
+
+void USlateColorWidgetFactory::RegisterWidgetConstructors(ITypedElementDataStorageInterface& DataStorage,
+														ITypedElementDataStorageUiInterface& DataStorageUi) const
+{
+	using namespace UE::Editor::DataStorage::Queries;
+	DataStorageUi.RegisterWidgetFactory<FSlateColorWidgetConstructor>(TEXT("General.Cell"), TColumn<FSlateColorColumn>());
+}
+
+FSlateColorWidgetConstructor::FSlateColorWidgetConstructor()
+	: FSimpleWidgetConstructor(StaticStruct())
+{
+}
+
+
+
+TSharedPtr<SWidget> FSlateColorWidgetConstructor::CreateWidget(ITypedElementDataStorageInterface* DataStorage,
+                                                               ITypedElementDataStorageUiInterface* DataStorageUi, RowHandle TargetRow, RowHandle WidgetRow,
+                                                               const UE::Editor::DataStorage::FMetaDataView& Arguments)
+{
+	UE::Editor::DataStorage::FAttributeBinder Binder(TargetRow, DataStorage);
+
+	FPointerEventHandler OnMouseButtonDown;
+	const bool* IsEditable = Arguments.FindForColumn<FSlateColorColumn>(UE::Editor::DataStorage::IsEditableName).TryGetExact<bool>();
+
+	// If the column is editable, summon a color picker on click
+	if(IsEditable && *IsEditable)
+	{
+		OnMouseButtonDown = FPointerEventHandler::CreateStatic(&UE::SlateColorWidget::Private::SummonColorPicker, DataStorage, TargetRow);
+	}
+	
+	return SNew(SHorizontalBox)
+			+SHorizontalBox::Slot()
+			.AutoWidth()
+			.HAlign(HAlign_Center)
+			.VAlign(VAlign_Center)
+			[
+				SNew(SColorBlock)
+					.Color(Binder.BindData(&FSlateColorColumn::Color, [](const FSlateColor& Color)
+					{
+						// If the color isn't set to a specified value, we just show it as white by default in the picker
+						return Color.IsColorSpecified() ? Color.GetSpecifiedColor() : FLinearColor::White;
+					}))
+					.Size(FVector2D(64, 16))
+					.OnMouseButtonDown(OnMouseButtonDown)
+			];
+}
diff --git a/Engine/Plugins/Experimental/EditorDataStorage/Source/TedsUI/Private/Widgets/SlateColorWidget.h b/Engine/Plugins/Experimental/EditorDataStorage/Source/TedsUI/Private/Widgets/SlateColorWidget.h
new file mode 100644
--- /dev/null
+++ b/Engine/Plugins/Experimental/EditorDataStorage/Source/TedsUI/Private/Widgets/SlateColorWidget.h
@@ -0,0 +1,41 @@
+// Copyright Epic Games, Inc. All Rights Reserved.
+
+#pragma once
+
+#include "Elements/Interfaces/TypedElementDataStorageFactory.h"
+#include "Elements/Interfaces/TypedElementDataStorageUiInterface.h"
+
+#include "SlateColorWidget.generated.h"
+
+class ITypedElementDataStorageInterface;
+class UScriptStruct;
+
+UCLASS()
+class USlateColorWidgetFactory : public UTypedElementDataStorageFactory
+{
+	GENERATED_BODY()
+
+public:
+	~USlateColorWidgetFactory() override = default;
+
+	void RegisterWidgetConstructors(ITypedElementDataStorageInterface& DataStorage,
+		ITypedElementDataStorageUiInterface& DataStorageUi) const override;
+};
+
+// Widget to show and edit the color column in TEDS
+USTRUCT()
+struct FSlateColorWidgetConstructor : public FSimpleWidgetConstructor
+{
+	GENERATED_BODY()
+
+public:
+	FSlateColorWidgetConstructor();
+	~FSlateColorWidgetConstructor() override = default;
+
+	virtual TSharedPtr<SWidget> CreateWidget(
+		ITypedElementDataStorageInterface* DataStorage,
+		ITypedElementDataStorageUiInterface* DataStorageUi,
+		RowHandle TargetRow,
+		RowHandle WidgetRow, 
+		const UE::Editor::DataStorage::FMetaDataView& Arguments) override;
+};
\ No newline at end of file
diff --git a/Engine/Plugins/Experimental/EditorDataStorage/Source/TedsUI/TedsUI.Build.cs b/Engine/Plugins/Experimental/EditorDataStorage/Source/TedsUI/TedsUI.Build.cs
--- a/Engine/Plugins/Experimental/EditorDataStorage/Source/TedsUI/TedsUI.Build.cs
+++ b/Engine/Plugins/Experimental/EditorDataStorage/Source/TedsUI/TedsUI.Build.cs
@@ -30,6 +30,8 @@ public class TedsUI : ModuleRules
 				new string[]
 				{
 					"ToolWidgets",
+					"AppFramework",
+					"InputCore"
 				});
 
 			DynamicallyLoadedModuleNames.AddRange(new string[] {});
diff --git a/Engine/Plugins/Experimental/EditorDataStorageFeatures/EditorDataStorageFeatures.uplugin b/Engine/Plugins/Experimental/EditorDataStorageFeatures/EditorDataStorageFeatures.uplugin
--- a/Engine/Plugins/Experimental/EditorDataStorageFeatures/EditorDataStorageFeatures.uplugin
+++ b/Engine/Plugins/Experimental/EditorDataStorageFeatures/EditorDataStorageFeatures.uplugin
@@ -42,11 +42,14 @@
 		{
 			"Name": "TedsTableViewer",
 			"Type": "Editor"
-		}
-	,
+		},
 		{
 			"Name": "TedsContentBrowser",
 			"Type": "Editor"
+		},
+		{
+			"Name": "TedsStyling",
+			"Type": "Editor"
 		}
 	],
 	"Plugins": [
diff --git a/Engine/Plugins/Experimental/EditorDataStorageFeatures/Source/TedsStyling/Private/TedsStylingFactory.cpp b/Engine/Plugins/Experimental/EditorDataStorageFeatures/Source/TedsStyling/Private/TedsStylingFactory.cpp
new file mode 100644
--- /dev/null
+++ b/Engine/Plugins/Experimental/EditorDataStorageFeatures/Source/TedsStyling/Private/TedsStylingFactory.cpp
@@ -0,0 +1,138 @@
+// Copyright Epic Games, Inc. All Rights Reserved.
+
+#include "TedsStylingFactory.h"
+
+#include "Columns/TedsStylingColumns.h"
+#include "Elements/Columns/TypedElementMiscColumns.h"
+#include "Elements/Columns/TypedElementSlateWidgetColumns.h"
+#include "Elements/Framework/TypedElementIndexHasher.h"
+#include "Elements/Framework/TypedElementRegistry.h"
+#include "Elements/Interfaces/TypedElementDataStorageInterface.h"
+#include "HAL/IConsoleManager.h"
+#include "Styling/SlateBrush.h"
+#include "Styling/SlateStyle.h"
+#include "Styling/SlateStyleRegistry.h"
+
+namespace UE::TedsStylingFactory::Local
+{
+	static const FName TableName("Editor_StyleTable");
+	
+	static FAutoConsoleCommand RegisterStylesCommand
+	(
+		TEXT("TEDS.Feature.AddSlateStyleRows"),
+		TEXT("Add all slate styles found in all registered stylesheets to TEDS"),
+		FConsoleCommandDelegate::CreateLambda([]() { UTedsStylingFactory::RegisterAllKnownStyles(); })
+	);
+
+}
+
+void UTedsStylingFactory::RegisterTables(ITypedElementDataStorageInterface& DataStorage)
+{
+	DataStorage.RegisterTable(TTypedElementColumnTypeList<FNameColumn, FSlateStyleSetColumn, FSlateStyleTag>(), UE::TedsStylingFactory::Local::TableName);
+}
+
+void UTedsStylingFactory::RegisterQueries(ITypedElementDataStorageInterface& DataStorage)
+{
+	Super::RegisterQueries(DataStorage);
+}
+
+void UTedsStylingFactory::RegisterAllKnownStyles()
+{
+	ITypedElementDataStorageInterface* Interface = UTypedElementRegistry::GetInstance()->GetMutableDataStorage();
+
+	auto RegisterBrushesForStyle = [Interface](const ISlateStyle& Style)
+	{
+		// Get all styles belonging to this styleset
+		TSet<FName> StyleKeys = Style.GetStyleKeys();
+
+		for (const FName& StyleKey : StyleKeys)
+		{
+			// Since there is no way to check which type of style a specific key is, we just go through each one by one and check for now
+			// This is not the most performant and can be improved in the future by exposing more internals from ISlateStyle if needed
+			// We also only support brushes and colors currently
+			bool bRegisteredSyleKey = false;
+
+			// Check if this style key is a brush
+			if (const FSlateBrush* Brush = Style.GetBrush(StyleKey))
+			{
+				if(Brush != Style.GetDefaultBrush())
+				{
+					RegisterBrush(Interface, StyleKey, Brush, Style);
+					bRegisteredSyleKey = true;
+				}
+			}
+
+			if(!bRegisteredSyleKey)
+			{
+				// Check if this style key is a color
+				FSlateColor Color = Style.GetSlateColor(StyleKey);
+				if (Color != FStyleDefaults::GetSlateColor())
+				{
+					RegisterColor(Interface, StyleKey, Color, Style);
+					bRegisteredSyleKey = true;
+				}
+			}
+		}
+	};
+
+	// Iterate all known style sheets and register all their members
+	FSlateStyleRegistry::IterateAllStyles(
+		[RegisterBrushesForStyle](const ISlateStyle& Style)
+		{
+			RegisterBrushesForStyle(Style);
+			return true;
+		}
+	);
+}
+
+void UTedsStylingFactory::RegisterBrush(ITypedElementDataStorageInterface* DataStorage, const FName& StyleName, const FSlateBrush* Brush, const ISlateStyle& OwnerStyle)
+{
+	// We currently don't store the FSlateBrush* itself to avoid storing a raw pointer that we can query on demand anyways in case we want to support
+	// swapping brushes or unloading styles on demand in the future
+	const UE::Editor::DataStorage::RowHandle Row = AddOrGetStyleRow(DataStorage, StyleName, OwnerStyle);
+	
+	if (DataStorage->IsRowAssigned(Row))
+	{
+		const FName ResourceName = Brush->GetResourceName();
+
+		if (ResourceName != NAME_None)
+		{
+			DataStorage->AddColumn(Row, FSlateStylePathColumn{.StylePath = ResourceName});
+		}
+
+		DataStorage->AddColumn<FSlateBrushTag>(Row);
+	}
+}
+
+void UTedsStylingFactory::RegisterColor(ITypedElementDataStorageInterface* DataStorage, const FName& StyleName, const FSlateColor& Color,
+	const ISlateStyle& OwnerStyle)
+{
+	const UE::Editor::DataStorage::RowHandle Row = AddOrGetStyleRow(DataStorage, StyleName, OwnerStyle);
+	
+	if (DataStorage->IsRowAssigned(Row))
+	{
+		DataStorage->AddColumn(Row, FSlateColorColumn{.Color = Color});
+	}
+
+}
+
+UE::Editor::DataStorage::RowHandle UTedsStylingFactory::AddOrGetStyleRow(ITypedElementDataStorageInterface* DataStorage,
+	const FName& StyleName, const ISlateStyle& OwnerStyle)
+{
+	const FName StyleSetName = OwnerStyle.GetStyleSetName();
+	const UE::Editor::DataStorage::IndexHash Index = UE::Editor::DataStorage::GenerateIndexHash(ISlateStyle::Join(StyleSetName, TCHAR_TO_ANSI(*StyleName.ToString())));
+	UE::Editor::DataStorage::RowHandle Row = DataStorage->FindIndexedRow(Index);
+
+	if (!DataStorage->IsRowAssigned(Row))
+	{
+		static UE::Editor::DataStorage::RowHandle Table = DataStorage->FindTable(UE::TedsStylingFactory::Local::TableName);
+		Row = DataStorage->AddRow(Table);
+		DataStorage->IndexRow(Index, Row);
+
+		DataStorage->GetColumn<FNameColumn>(Row)->Name = StyleName;
+		DataStorage->GetColumn<FSlateStyleSetColumn>(Row)->StyleSetName = StyleSetName;
+	}
+
+	return Row;
+
+}
diff --git a/Engine/Plugins/Experimental/EditorDataStorageFeatures/Source/TedsStyling/Private/TedsStylingFactory.h b/Engine/Plugins/Experimental/EditorDataStorageFeatures/Source/TedsStyling/Private/TedsStylingFactory.h
new file mode 100644
--- /dev/null
+++ b/Engine/Plugins/Experimental/EditorDataStorageFeatures/Source/TedsStyling/Private/TedsStylingFactory.h
@@ -0,0 +1,31 @@
+// Copyright Epic Games, Inc. All Rights Reserved.
+
+#pragma once
+
+#include "Elements/Interfaces/TypedElementDataStorageFactory.h"
+#include "Elements/Common/TypedElementHandles.h"
+
+#include "TedsStylingFactory.generated.h"
+
+struct FSlateBrush;
+struct FSlateColor;
+class ISlateStyle;
+
+UCLASS()
+class UTedsStylingFactory : public UTypedElementDataStorageFactory
+{
+	GENERATED_BODY()
+
+public:
+	~UTedsStylingFactory() override = default;
+
+	void RegisterTables(ITypedElementDataStorageInterface& DataStorage) override;
+	virtual void RegisterQueries(ITypedElementDataStorageInterface& DataStorage) override;
+
+	static void RegisterAllKnownStyles();
+	
+private:
+	static void RegisterBrush(ITypedElementDataStorageInterface* DataStorage, const FName& StyleName, const FSlateBrush* Brush, const ISlateStyle& OwnerStyle);
+	static void RegisterColor(ITypedElementDataStorageInterface* DataStorage, const FName& StyleName, const FSlateColor& Color, const ISlateStyle& OwnerStyle);
+	static UE::Editor::DataStorage::RowHandle AddOrGetStyleRow(ITypedElementDataStorageInterface* DataStorage, const FName& StyleName, const ISlateStyle& OwnerStyle);
+};
diff --git a/Engine/Plugins/Experimental/EditorDataStorageFeatures/Source/TedsStyling/Private/TedsStylingModule.cpp b/Engine/Plugins/Experimental/EditorDataStorageFeatures/Source/TedsStyling/Private/TedsStylingModule.cpp
new file mode 100644
--- /dev/null
+++ b/Engine/Plugins/Experimental/EditorDataStorageFeatures/Source/TedsStyling/Private/TedsStylingModule.cpp
@@ -0,0 +1,5 @@
+// Copyright Epic Games, Inc. All Rights Reserved.
+
+#include "Modules/ModuleManager.h"
+
+IMPLEMENT_MODULE(FDefaultModuleImpl, TedsStyling)
diff --git a/Engine/Plugins/Experimental/EditorDataStorageFeatures/Source/TedsStyling/Private/Widgets/SlateBrushWidget.cpp b/Engine/Plugins/Experimental/EditorDataStorageFeatures/Source/TedsStyling/Private/Widgets/SlateBrushWidget.cpp
new file mode 100644
--- /dev/null
+++ b/Engine/Plugins/Experimental/EditorDataStorageFeatures/Source/TedsStyling/Private/Widgets/SlateBrushWidget.cpp
@@ -0,0 +1,75 @@
+// Copyright Epic Games, Inc. All Rights Reserved.
+
+#include "Widgets/SlateBrushWidget.h"
+
+#include "Columns/TedsStylingColumns.h"
+#include "Elements/Columns/TypedElementMiscColumns.h"
+#include "Elements/Framework/TypedElementAttributeBinding.h"
+#include "Styling/SlateStyleRegistry.h"
+#include "Widgets/Images/SImage.h"
+#include "Widgets/Layout/SBox.h"
+#include "Widgets/SBoxPanel.h"
+#include "Widgets/Text/STextBlock.h"
+
+void USlateStylePreviewWidget::RegisterWidgetConstructors(ITypedElementDataStorageInterface& DataStorage,
+                                                          ITypedElementDataStorageUiInterface& DataStorageUi) const
+{
+	using namespace UE::Editor::DataStorage::Queries;
+	DataStorageUi.RegisterWidgetFactory<FSlateStylePreviewWidgetConstructor>(TEXT("General.RowLabel"), TColumn<FNameColumn>() && TColumn<FSlateStyleTag>());
+	DataStorageUi.RegisterWidgetFactory<FSlateStylePreviewWidgetConstructor>(TEXT("General.Cell"), TColumn<FNameColumn>() && TColumn<FSlateStyleTag>());
+}
+
+FSlateStylePreviewWidgetConstructor::FSlateStylePreviewWidgetConstructor()
+	: Super(StaticStruct())
+{
+}
+
+TSharedPtr<SWidget> FSlateStylePreviewWidgetConstructor::CreateWidget(ITypedElementDataStorageInterface* DataStorage,
+	ITypedElementDataStorageUiInterface* DataStorageUi, UE::Editor::DataStorage::RowHandle TargetRow, UE::Editor::DataStorage::RowHandle WidgetRow,
+	const UE::Editor::DataStorage::FMetaDataView& Arguments)
+{
+	UE::Editor::DataStorage::FAttributeBinder Binder(TargetRow, DataStorage);
+	
+	return SNew(SHorizontalBox)
+			+ SHorizontalBox::Slot()
+			.HAlign(HAlign_Left)
+			.VAlign(VAlign_Center)
+			.AutoWidth()
+			[
+				SNew(SBox)
+				.MaxDesiredHeight(16.0f)
+				.MaxDesiredWidth(16.0f)
+				[
+					// Currently we cannot differentiate between image brushes and other brushes in TEDS so we only show image brushes as an SImage
+					// and don't support showing a preview of the rest
+					SNew(SImage)
+					.Image_Lambda([DataStorage, TargetRow]()
+					{
+						const FNameColumn* StyleNameColumn = DataStorage->GetColumn<FNameColumn>(TargetRow);
+						const FSlateStyleSetColumn* StyleSetColumn = DataStorage->GetColumn<FSlateStyleSetColumn>(TargetRow);
+
+						if (StyleNameColumn && StyleSetColumn && DataStorage->HasColumns<FSlateBrushTag>(TargetRow))
+						{
+							if(const ISlateStyle* Style = FSlateStyleRegistry::FindSlateStyle(StyleSetColumn->StyleSetName))
+							{
+								if(const FSlateBrush* Brush = Style->GetBrush(StyleNameColumn->Name))
+								{
+									return Brush;
+								}
+							}
+						}
+						return FAppStyle::GetBrush("NoBrush");
+					})
+				]
+			]
+			+ SHorizontalBox::Slot()
+			.HAlign(HAlign_Left)
+			.VAlign(VAlign_Center)
+			.Padding(5.0f, 0.0f, 0.0f, 0.0f)
+			.AutoWidth()
+			[
+				SNew(STextBlock)
+				.Text(Binder.BindText(&FNameColumn::Name))
+			];
+	
+}
diff --git a/Engine/Plugins/Experimental/EditorDataStorageFeatures/Source/TedsStyling/Private/Widgets/SlateBrushWidget.h b/Engine/Plugins/Experimental/EditorDataStorageFeatures/Source/TedsStyling/Private/Widgets/SlateBrushWidget.h
new file mode 100644
--- /dev/null
+++ b/Engine/Plugins/Experimental/EditorDataStorageFeatures/Source/TedsStyling/Private/Widgets/SlateBrushWidget.h
@@ -0,0 +1,41 @@
+// Copyright Epic Games, Inc. All Rights Reserved.
+
+#pragma once
+
+#include "Elements/Interfaces/TypedElementDataStorageFactory.h"
+#include "Elements/Interfaces/TypedElementDataStorageUiInterface.h"
+
+#include "SlateBrushWidget.generated.h"
+
+class ITypedElementDataStorageInterface;
+class UScriptStruct;
+
+UCLASS()
+class USlateStylePreviewWidget : public UTypedElementDataStorageFactory
+{
+	GENERATED_BODY()
+
+public:
+	~USlateStylePreviewWidget() override = default;
+
+	void RegisterWidgetConstructors(ITypedElementDataStorageInterface& DataStorage,
+		ITypedElementDataStorageUiInterface& DataStorageUi) const override;
+};
+
+// Widget to show a slate brush drawn as an SImage
+USTRUCT()
+struct FSlateStylePreviewWidgetConstructor : public FSimpleWidgetConstructor
+{
+	GENERATED_BODY()
+
+public:
+	FSlateStylePreviewWidgetConstructor();
+	~FSlateStylePreviewWidgetConstructor() override = default;
+
+	virtual TSharedPtr<SWidget> CreateWidget(
+		ITypedElementDataStorageInterface* DataStorage,
+		ITypedElementDataStorageUiInterface* DataStorageUi,
+		UE::Editor::DataStorage::RowHandle TargetRow,
+		UE::Editor::DataStorage::RowHandle WidgetRow,
+		const UE::Editor::DataStorage::FMetaDataView& Arguments) override;
+};
\ No newline at end of file
diff --git a/Engine/Plugins/Experimental/EditorDataStorageFeatures/Source/TedsStyling/Public/Columns/TedsStylingColumns.h b/Engine/Plugins/Experimental/EditorDataStorageFeatures/Source/TedsStyling/Public/Columns/TedsStylingColumns.h
new file mode 100644
--- /dev/null
+++ b/Engine/Plugins/Experimental/EditorDataStorageFeatures/Source/TedsStyling/Public/Columns/TedsStylingColumns.h
@@ -0,0 +1,50 @@
+// Copyright Epic Games, Inc. All Rights Reserved.
+
+#pragma once
+
+#include "Elements/Interfaces/TypedElementDataStorageInterface.h"
+
+#include "TedsStylingColumns.generated.h"
+
+/**
+ * Owner style set
+ */
+USTRUCT(meta = (DisplayName = "Style Set"))
+struct FSlateStyleSetColumn final : public FEditorDataStorageColumn
+{
+	GENERATED_BODY()
+
+	UPROPERTY(meta = (Searchable))
+	FName StyleSetName;
+};
+
+/**
+ * Path to style resource on disk
+ */
+USTRUCT(meta = (DisplayName = "Style Path"))
+struct FSlateStylePathColumn final : public FEditorDataStorageColumn
+{
+	GENERATED_BODY()
+
+	UPROPERTY(meta = (Searchable))
+	FName StylePath;
+};
+
+/**
+ * Tag on rows that are part of a stylesheet
+ */
+USTRUCT(meta = (DisplayName = "Slate Style"))
+struct FSlateStyleTag final : public FEditorDataStorageTag
+{
+	GENERATED_BODY()
+};
+
+/**
+ * Tag on rows that are brushes in the stylesheet
+ * TEDS UI TODO: Once TEDS UI supports dynamic columns, we can remove this tag and make FSlateStyleTagDynamic
+ */
+USTRUCT(meta = (DisplayName = "Slate Brush"))
+struct FSlateBrushTag final : public FEditorDataStorageTag
+{
+	GENERATED_BODY()
+};
diff --git a/Engine/Plugins/Experimental/EditorDataStorageFeatures/Source/TedsStyling/TedsStyling.Build.cs b/Engine/Plugins/Experimental/EditorDataStorageFeatures/Source/TedsStyling/TedsStyling.Build.cs
new file mode 100644
--- /dev/null
+++ b/Engine/Plugins/Experimental/EditorDataStorageFeatures/Source/TedsStyling/TedsStyling.Build.cs
@@ -0,0 +1,33 @@
+// Copyright Epic Games, Inc. All Rights Reserved.
+
+using UnrealBuildTool;
+
+public class TedsStyling : ModuleRules
+{
+	public TedsStyling(ReadOnlyTargetRules Target) : base(Target)
+	{
+		PCHUsage = ModuleRules.PCHUsageMode.UseExplicitOrSharedPCHs;
+
+		if (Target.bBuildEditor)
+		{
+			PublicIncludePaths.AddRange(new string[] { });
+			PrivateIncludePaths.AddRange(new string[] { });
+
+			PrivateDependencyModuleNames.AddRange(
+				new string[]
+				{
+					"Core",
+					"CoreUObject",
+					"Engine",
+					"SlateCore",
+					"Slate"
+				});
+
+			PublicDependencyModuleNames.AddRange(
+				new string[]
+				{
+					"TypedElementFramework",
+				});
+		}
+	}
+}
diff --git a/Engine/Source/Runtime/TypedElementFramework/Public/Elements/Columns/TypedElementSlateWidgetColumns.h b/Engine/Source/Runtime/TypedElementFramework/Public/Elements/Columns/TypedElementSlateWidgetColumns.h
--- a/Engine/Source/Runtime/TypedElementFramework/Public/Elements/Columns/TypedElementSlateWidgetColumns.h
+++ b/Engine/Source/Runtime/TypedElementFramework/Public/Elements/Columns/TypedElementSlateWidgetColumns.h
@@ -3,6 +3,7 @@
 #pragma once
 
 #include "Elements/Common/TypedElementCommonTypes.h"
+#include "Styling/SlateColor.h"
 #include "Templates/SharedPointer.h"
 
 #include "TypedElementSlateWidgetColumns.generated.h"
@@ -70,3 +71,16 @@ struct FDescriptionColumn final : public FEditorDataStorageColumn
 	UPROPERTY(meta = (Searchable))
 	FText Description;
 };
+
+/**
+ * A color for this row. Can be used by widget rows to determine widget color or non-widget rows to attach a logical color to their data
+ * (e.g asset colors)
+ */
+USTRUCT(meta = (DisplayName = "Color"))
+struct FSlateColorColumn final : public FEditorDataStorageColumn
+{
+	GENERATED_BODY()
+
+	UPROPERTY()
+	FSlateColor Color = FSlateColor::UseForeground();
+};
