TEDS Content Browser:

Add folder and asset colors to TEDS
Remove the ItemName_ExperimentalColumn and replace with FNameColumn
Add extension utils to the CB to get/set folder colors
Add a new FFolder tag used by path rows
Remove test asset data and repurpose AssetProcessors.cpp to sync asset/folder data to and from TEDS
Modify the AssetDataLabel widget to show folder/asset colors

TEDS UI:
Add a new FSlateColor column and create a widget to view/edit it

diff --git a/Engine/Plugins/Experimental/EditorDataStorageFeatures/Source/TedsAssetData/Private/CB/TedsAssetDataCBDataSource.cpp b/Engine/Plugins/Experimental/EditorDataStorageFeatures/Source/TedsAssetData/Private/CB/TedsAssetDataCBDataSource.cpp
--- a/Engine/Plugins/Experimental/EditorDataStorageFeatures/Source/TedsAssetData/Private/CB/TedsAssetDataCBDataSource.cpp
+++ b/Engine/Plugins/Experimental/EditorDataStorageFeatures/Source/TedsAssetData/Private/CB/TedsAssetDataCBDataSource.cpp
@@ -18,6 +18,7 @@
 #include "Internationalization/Text.h"
 #include "Misc/PathViews.h"
 #include "PluginDescriptor.h"
+#include "Elements/Columns/TypedElementMiscColumns.h"
 #include "Settings/ContentBrowserSettings.h"
 #include "TedsAssetDataModule.h"
 #include "Templates/Function.h"
@@ -89,7 +90,6 @@ FTedsAssetDataCBDataSource::FTedsAssetDataCBDataSource(ITypedElementDataStorageI
 			.None<FUpdatedAssetDataTag>()
 		.Compile());
 
-
 	ProcessAssetDataAndPathUpdateQuery = Database.RegisterQuery(
 		Select(
 			TEXT("FTedsAssetDataCBDataSource: Process Asset Data and Path Updates"),
@@ -250,7 +250,7 @@ void FTedsAssetDataCBDataSource::AddAssetDataColumns(UE::Editor::DataStorage::IQ
 	AssetClassColumn.ClassPath = AssetData.AssetClassPath;
 	Context.AddColumn(Row, MoveTemp(AssetClassColumn));
 
-	FItemNameColumn_Experimental ItemNameColumn;
+	FNameColumn ItemNameColumn;
 	ItemNameColumn.Name = AssetData.AssetName;
 	Context.AddColumn(Row, MoveTemp(ItemNameColumn));
 
diff --git a/Engine/Plugins/Experimental/EditorDataStorageFeatures/Source/TedsAssetData/Private/Processors/AssetProcessors.cpp b/Engine/Plugins/Experimental/EditorDataStorageFeatures/Source/TedsAssetData/Private/Processors/AssetProcessors.cpp
--- a/Engine/Plugins/Experimental/EditorDataStorageFeatures/Source/TedsAssetData/Private/Processors/AssetProcessors.cpp
+++ b/Engine/Plugins/Experimental/EditorDataStorageFeatures/Source/TedsAssetData/Private/Processors/AssetProcessors.cpp
@@ -2,16 +2,92 @@
 
 #include "Processors/AssetProcessors.h"
 
-#include "Elements/Columns/TypedElementLabelColumns.h"
-#include "Elements/Columns/TypedElementTypeInfoColumns.h"
+#include "AssetDefinition.h"
+#include "AssetDefinitionRegistry.h"
+#include "CollectionManagerTypes.h"
+#include "ContentBrowserDataSubsystem.h"
+#include "ContentBrowserModule.h"
+#include "IContentBrowserDataModule.h"
+#include "Elements/Columns/TypedElementFolderColumns.h"
+#include "Elements/Columns/TypedElementSlateWidgetColumns.h"
+#include "Elements/Framework/TypedElementQueryBuilder.h"
 #include "Elements/Interfaces/TypedElementQueryStorageInterfaces.h"
+#include "Experimental/ContentBrowserExtensionUtils.h"
 #include "TedsAssetDataColumns.h"
+#include "Elements/Columns/TypedElementMiscColumns.h"
+#include "Elements/Framework/TypedElementIndexHasher.h"
 
-void UTedsAssetDataFactory::RegisterTables(ITypedElementDataStorageInterface& DataStorage) 
+void UTedsAssetDataFactory::RegisterQueries(ITypedElementDataStorageInterface& DataStorage)
 {
-	DataStorage.RegisterTable(
-		TTypedElementColumnTypeList<
-			FTypedElementLabelColumn, FTypedElementClassTypeInfoColumn, FAssetPathColumn_Experimental,
-			FAssetTag, FDiskSizeColumn, FVersePathColumn>(),
-		FName("Editor_PlaceholderAssetTable"));
+	using namespace UE::Editor::DataStorage::Queries;
+
+	DataStorage.RegisterQuery(
+		Select(
+			TEXT("TedsAssetDataFactory: Sync folder color from world"),
+			FProcessor(EQueryTickPhase::PostPhysics, DataStorage.GetQueryTickGroupName(EQueryTickGroups::SyncExternalToDataStorage))
+			.SetExecutionMode(EExecutionMode::GameThread),
+			[](IQueryContext& Context, const RowHandle* Rows, const FAssetPathColumn_Experimental* AssetPathColumn, FSlateColorColumn* ColorColumn)
+			{
+				const int32 NumOfRowToProcess = Context.GetRowCount();
+
+				for (int32 Index = 0; Index < NumOfRowToProcess; ++Index)
+				{
+					if (TOptional<FLinearColor> Color = UE::Editor::ContentBrowser::ExtensionUtils::GetFolderColor(AssetPathColumn[Index].Path))
+					{
+						ColorColumn[Index].Color = Color.GetValue();
+					}
+				}
+			}
+		)
+		.Where()
+			.All<FFolderTag, FUpdatedPathTag>()
+		.Compile()
+		);
+
+	DataStorage.RegisterQuery(
+		Select(
+			TEXT("TedsAssetDataFactory: Sync folder color back to world"),
+			FProcessor(EQueryTickPhase::PrePhysics, DataStorage.GetQueryTickGroupName(EQueryTickGroups::SyncDataStorageToExternal))
+			.SetExecutionMode(EExecutionMode::GameThread),
+			[](IQueryContext& Context, const RowHandle* Rows, const FAssetPathColumn_Experimental* PathColumn, const FSlateColorColumn* ColorColumn)
+			{
+				const int32 NumOfRowToProcess = Context.GetRowCount();
+
+				for (int32 Index = 0; Index < NumOfRowToProcess; ++Index)
+				{
+					UE::Editor::ContentBrowser::ExtensionUtils::SetFolderColor(PathColumn[Index].Path, ColorColumn[Index].Color.GetSpecifiedColor());
+				}
+			}
+		)
+		.Where()
+			.All<FFolderTag, FTypedElementSyncBackToWorldTag>()
+		.Compile()
+		);
+}
+
+void UTedsAssetDataFactory::PreRegister(ITypedElementDataStorageInterface& DataStorage)
+{
+	if (FContentBrowserModule* ContentBrowserModule = FModuleManager::Get().GetModulePtr<FContentBrowserModule>("ContentBrowser"))
+	{
+		ContentBrowserModule->GetOnSetFolderColor().AddUObject(this, &UTedsAssetDataFactory::OnSetFolderColor, &DataStorage);
+	}
+}
+
+void UTedsAssetDataFactory::PreShutdown(ITypedElementDataStorageInterface& DataStorage)
+{
+	if (FContentBrowserModule* ContentBrowserModule = FModuleManager::Get().GetModulePtr<FContentBrowserModule>("ContentBrowser"))
+	{
+		ContentBrowserModule->GetOnSetFolderColor().RemoveAll(this);
+	}
+}
+
+void UTedsAssetDataFactory::OnSetFolderColor(const FString& Path, ITypedElementDataStorageInterface* DataStorage)
+{
+	const UE::Editor::DataStorage::IndexHash PathHash = UE::Editor::DataStorage::GenerateIndexHash(FName(Path));
+	const UE::Editor::DataStorage::RowHandle Row = DataStorage->FindIndexedRow(PathHash);
+
+	if(DataStorage->IsRowAvailable(Row))
+	{
+		DataStorage->AddColumn<FUpdatedPathTag>(Row);
+	}
 }
diff --git a/Engine/Plugins/Experimental/EditorDataStorageFeatures/Source/TedsAssetData/Private/Processors/AssetProcessors.h b/Engine/Plugins/Experimental/EditorDataStorageFeatures/Source/TedsAssetData/Private/Processors/AssetProcessors.h
--- a/Engine/Plugins/Experimental/EditorDataStorageFeatures/Source/TedsAssetData/Private/Processors/AssetProcessors.h
+++ b/Engine/Plugins/Experimental/EditorDataStorageFeatures/Source/TedsAssetData/Private/Processors/AssetProcessors.h
@@ -16,7 +16,11 @@ class UTedsAssetDataFactory : public UTypedElementDataStorageFactory
 public:
 	~UTedsAssetDataFactory() override = default;
 
-	void RegisterTables(ITypedElementDataStorageInterface& DataStorage) override;
+	virtual void RegisterQueries(ITypedElementDataStorageInterface& DataStorage) override;
+	virtual void PreRegister(ITypedElementDataStorageInterface& DataStorage) override;
+	virtual void PreShutdown(ITypedElementDataStorageInterface& DataStorage) override;
 
-private:
+protected:
+
+	void OnSetFolderColor(const FString& Path, ITypedElementDataStorageInterface* DataStorage);
 };
diff --git a/Engine/Plugins/Experimental/EditorDataStorageFeatures/Source/TedsAssetData/Private/TedsAssetData.cpp b/Engine/Plugins/Experimental/EditorDataStorageFeatures/Source/TedsAssetData/Private/TedsAssetData.cpp
--- a/Engine/Plugins/Experimental/EditorDataStorageFeatures/Source/TedsAssetData/Private/TedsAssetData.cpp
+++ b/Engine/Plugins/Experimental/EditorDataStorageFeatures/Source/TedsAssetData/Private/TedsAssetData.cpp
@@ -8,6 +8,9 @@
 #include "Async/ParallelFor.h"
 #include "Containers/Array.h"
 #include "Containers/ChunkedArray.h"
+#include "Elements/Columns/TypedElementFolderColumns.h"
+#include "Elements/Columns/TypedElementMiscColumns.h"
+#include "Elements/Columns/TypedElementSlateWidgetColumns.h"
 #include "Elements/Common/TypedElementQueryTypes.h"
 #include "Elements/Framework/TypedElementIndexHasher.h"
 #include "Elements/Framework/TypedElementQueryBuilder.h"
@@ -77,6 +80,7 @@ struct FPopulatePathRowArgs
 	FName AssetRegistryPath;
 	IndexHash AssetRegistryPathHash;
 	IndexHash ParentAssetRegistryPathHash;
+	FName AssetName;
 	uint32 PathDepth;
 
 	operator bool() const 
@@ -118,15 +122,18 @@ FPopulatePathRowArgs ThreadSafe_PopulatePathRowArgs(IndexHash AssetRegistryPathH
 	IndexHash ParentAssetRegistryPathHash = InvalidRowHandle;
 	int32 CharacterIndex;
 	uint32 Depth;
+	FString AssetName;
 	GetPathDepthAndParentFolderIndex(PathAsString, Depth, CharacterIndex);
 	if (CharacterIndex != INDEX_NONE)
 	{
 		FStringView ParentPath = PathAsString.Left(CharacterIndex);
 		ParentAssetRegistryPathHash = GenerateIndexHash(FName(ParentPath));
+		AssetName = PathAsString.RightChop(CharacterIndex);
 	}
 
 	FPopulatePathRowArgs Args;
 	Args.AssetRegistryPath = InAssetRegistryPath;
+	Args.AssetName = FName(AssetName);
 	Args.AssetRegistryPathHash = AssetRegistryPathHash;
 	Args.ParentAssetRegistryPathHash = ParentAssetRegistryPathHash;
 	Args.PathDepth = Depth;
@@ -154,6 +161,7 @@ void PopulatePathDataTableRow(FPopulatePathRowArgs&& InPopulatePathRowArgs, ITyp
 	}
 
 	Database.GetColumn<FAssetPathColumn_Experimental>(InRowHandle)->Path = InPopulatePathRowArgs.AssetRegistryPath;
+	Database.GetColumn<FNameColumn>(InRowHandle)->Name = InPopulatePathRowArgs.AssetName;
 }
 
 
@@ -182,7 +190,7 @@ FTedsAssetData::FTedsAssetData(ITypedElementDataStorageInterface& InDatabase)
 	PathsTable = Database.FindTable(FName(TEXT("Editor_AssetRegistryPathsTable")));
 	if (PathsTable == InvalidTableHandle)
 	{
-		PathsTable = Database.RegisterTable<FAssetPathColumn_Experimental, FChildrenAssetPathColumn_Experimental, FParentAssetPathColumn_Experimental, FAssetsInPathColumn_Experimental, FUpdatedPathTag>(FName(TEXT("Editor_AssetRegistryPathsTable")));
+		PathsTable = Database.RegisterTable<FFolderTag, FAssetPathColumn_Experimental, FNameColumn, FUpdatedPathTag, FSlateColorColumn, FChildrenAssetPathColumn_Experimental, FParentAssetPathColumn_Experimental, FAssetsInPathColumn_Experimental>(FName(TEXT("Editor_AssetRegistryPathsTable")));
 	}
 
 	AssetsDataTable = Database.FindTable(FName(TEXT("Editor_AssetRegistryAssetDataTable")));
diff --git a/Engine/Plugins/Experimental/EditorDataStorageFeatures/Source/TedsAssetData/Private/Widgets/AssetDataLabelWidget.cpp b/Engine/Plugins/Experimental/EditorDataStorageFeatures/Source/TedsAssetData/Private/Widgets/AssetDataLabelWidget.cpp
--- a/Engine/Plugins/Experimental/EditorDataStorageFeatures/Source/TedsAssetData/Private/Widgets/AssetDataLabelWidget.cpp
+++ b/Engine/Plugins/Experimental/EditorDataStorageFeatures/Source/TedsAssetData/Private/Widgets/AssetDataLabelWidget.cpp
@@ -2,8 +2,12 @@
 
 #include "Widgets/AssetDataLabelWidget.h"
 
-#include "TedsAssetDataColumns.h"
+#include "AssetDefinition.h"
+#include "AssetDefinitionRegistry.h"
 #include "Elements/Framework/TypedElementAttributeBinding.h"
+#include "TedsAssetDataColumns.h"
+#include "Elements/Columns/TypedElementMiscColumns.h"
+#include "Elements/Columns/TypedElementSlateWidgetColumns.h"
 #include "Widgets/Images/SImage.h"
 #include "Widgets/Layout/SSpacer.h"
 #include "Widgets/Text/STextBlock.h"
@@ -17,7 +21,7 @@ void UAssetDataLabelWidgetFactory::RegisterWidgetConstructors(ITypedElementDataS
 	using namespace UE::Editor::DataStorage::Queries;
 	DataStorageUi.RegisterWidgetFactory<FAssetDataLabelWidgetConstructor>(
 		TEXT("General.RowLabel"),
-		TColumn<FItemNameColumn_Experimental>() && TColumn<FAssetTag>());
+		TColumn<FNameColumn>() && (TColumn<FAssetTag>() || TColumn<FAssetPathColumn_Experimental>()));
 }
 
 FAssetDataLabelWidgetConstructor::FAssetDataLabelWidgetConstructor()
@@ -37,15 +41,37 @@ TSharedPtr<SWidget> FAssetDataLabelWidgetConstructor::CreateWidget(ITypedElement
 	if (DataStorage->IsRowAvailable(TargetRow))
 	{
 		UE::Editor::DataStorage::FAttributeBinder Binder(TargetRow, DataStorage);
-		
+
+		bool bIsAsset = DataStorage->HasColumns<FAssetDataColumn_Experimental>(TargetRow);
+
+		TAttribute<FSlateColor> ColorAndOpacityAttribute;
+
+		// For assets, grab the color from the asset definition
+		if(bIsAsset)
+		{
+			ColorAndOpacityAttribute = Binder.BindData(&FAssetDataColumn_Experimental::AssetData, [](const FAssetData& AssetData)
+			{
+				if(const UAssetDefinition* AssetDefinition = UAssetDefinitionRegistry::Get()->GetAssetDefinitionForAsset(AssetData))
+				{
+					return FSlateColor(AssetDefinition->GetAssetColor());
+				}
+
+				return FSlateColor::UseForeground();
+			});
+		}
+		// For folders, use the color column directly
+		else
+		{
+			ColorAndOpacityAttribute = Binder.BindData(&FSlateColorColumn::Color, FSlateColor::UseForeground());
+		}
+
 		return SNew(SHorizontalBox)
 			+SHorizontalBox::Slot()
 				.AutoWidth()
 			[
-				// TODO: Asset color in TEDS
 				SNew(SImage)
-					.Image(FAppStyle::GetBrush(FName(TEXT("ContentBrowser.ColumnViewAssetIcon"))))
-					.ColorAndOpacity(FSlateColor::UseForeground())
+					.Image(bIsAsset ? FAppStyle::GetBrush(FName(TEXT("ContentBrowser.ColumnViewAssetIcon"))) : FAppStyle::GetBrush(FName(TEXT("ContentBrowser.ColumnViewFolderIcon"))))
+					.ColorAndOpacity(ColorAndOpacityAttribute)
 			]
 			+SHorizontalBox::Slot()
 				.AutoWidth()
@@ -57,10 +83,10 @@ TSharedPtr<SWidget> FAssetDataLabelWidgetConstructor::CreateWidget(ITypedElement
 				.FillWidth(1.0f)
 			[
 				SNew(STextBlock)
-				.Text(Binder.BindText(&FItemNameColumn_Experimental::Name))
+				.Text(Binder.BindText(&FNameColumn::Name))
 				.ToolTipText_Lambda([DataStorage, DataRow = TargetRow]()
 					{
-						return FAssetDataLabelWidgetConstructor::ConstructToolTip(DataStorage, DataRow);
+						return ConstructToolTip(DataStorage, DataRow);
 					})
 			];
 	}
@@ -74,7 +100,7 @@ FText FAssetDataLabelWidgetConstructor::ConstructToolTip(
 	ITypedElementDataStorageInterface* DataStorage, UE::Editor::DataStorage::RowHandle DataRow)
 {
 	TStringBuilder<1024> ToolTip;
-	if (FItemNameColumn_Experimental* ItemName = DataStorage->GetColumn<FItemNameColumn_Experimental>(DataRow))
+	if (FNameColumn* ItemName = DataStorage->GetColumn<FNameColumn>(DataRow))
 	{
 		ToolTip.Append(ItemName->Name.ToString());
 		ToolTip.AppendChar(TEXT('\n'));
diff --git a/Engine/Plugins/Experimental/EditorDataStorageFeatures/Source/TedsAssetData/Public/TedsAssetDataColumns.h b/Engine/Plugins/Experimental/EditorDataStorageFeatures/Source/TedsAssetData/Public/TedsAssetDataColumns.h
--- a/Engine/Plugins/Experimental/EditorDataStorageFeatures/Source/TedsAssetData/Public/TedsAssetDataColumns.h
+++ b/Engine/Plugins/Experimental/EditorDataStorageFeatures/Source/TedsAssetData/Public/TedsAssetDataColumns.h
@@ -149,15 +149,6 @@ struct FVirtualPathColumn_Experimental : public FEditorDataStorageColumn
 	FName VirtualPath;
 };
 
-USTRUCT(meta = (DisplayName = "Name"))
-struct FItemNameColumn_Experimental : public FEditorDataStorageColumn
-{
-	GENERATED_BODY()
-
-	UPROPERTY(meta = (Searchable))
-	FName Name;
-};
-
 USTRUCT()
 struct FItemAttributeBaseColumn_Experimental : public FEditorDataStorageColumn
 {
diff --git a/Engine/Plugins/Experimental/EditorDataStorageFeatures/Source/TedsAssetData/TedsAssetData.build.cs b/Engine/Plugins/Experimental/EditorDataStorageFeatures/Source/TedsAssetData/TedsAssetData.build.cs
--- a/Engine/Plugins/Experimental/EditorDataStorageFeatures/Source/TedsAssetData/TedsAssetData.build.cs
+++ b/Engine/Plugins/Experimental/EditorDataStorageFeatures/Source/TedsAssetData/TedsAssetData.build.cs
@@ -30,6 +30,8 @@ public class TedsAssetData : ModuleRules
 					"TypedElementFramework",
 					"UnrealEd",
 					"UnrealEd",
+					"AssetDefinition",
+					"ContentBrowser"
 				});
 
 			PublicDependencyModuleNames.AddRange(
diff --git a/Engine/Plugins/Experimental/EditorDataStorageFeatures/Source/TedsContentBrowser/Private/TedsContentBrowserModule.cpp b/Engine/Plugins/Experimental/EditorDataStorageFeatures/Source/TedsContentBrowser/Private/TedsContentBrowserModule.cpp
--- a/Engine/Plugins/Experimental/EditorDataStorageFeatures/Source/TedsContentBrowser/Private/TedsContentBrowserModule.cpp
+++ b/Engine/Plugins/Experimental/EditorDataStorageFeatures/Source/TedsContentBrowser/Private/TedsContentBrowserModule.cpp
@@ -12,6 +12,7 @@
 #include "Modules/ModuleManager.h"
 #include "QueryStack/FQueryStackNode_RowView.h"
 #include "TedsAssetDataColumns.h"
+#include "Elements/Columns/TypedElementMiscColumns.h"
 #include "Widgets/STedsTableViewer.h"
 #include "Widgets/Views/SListView.h"
 
@@ -134,7 +135,7 @@ namespace UE::Editor::ContentBrowser
 					.QueryStack(RowQueryStack)
 					.CellWidgetPurposes({TEXT("General.RowLabel"), TEXT("General.Cell")})
 					// Default list of columns to display
-					.Columns({ FItemNameColumn_Experimental::StaticStruct(), FAssetTag::StaticStruct(), FAssetPathColumn_Experimental::StaticStruct(),
+					.Columns({ FNameColumn::StaticStruct(), FAssetTag::StaticStruct(), FAssetPathColumn_Experimental::StaticStruct(),
 						FDiskSizeColumn::StaticStruct(), FVirtualPathColumn_Experimental::StaticStruct() })
 					.ListSelectionMode(ESelectionMode::Multi)
 					.OnSelectionChanged_Lambda([this](RowHandle Row)
diff --git a/Engine/Source/Developer/AssetTools/Private/AssetViewUtils.cpp b/Engine/Source/Developer/AssetTools/Private/AssetViewUtils.cpp
--- a/Engine/Source/Developer/AssetTools/Private/AssetViewUtils.cpp
+++ b/Engine/Source/Developer/AssetTools/Private/AssetViewUtils.cpp
@@ -29,6 +29,7 @@
 #include "AssetRegistry/AssetRegistryModule.h"
 #include "IAssetTools.h"
 #include "AssetToolsModule.h"
+#include "ContentBrowserModule.h"
 #include "Settings/EditorExperimentalSettings.h"
 
 #include "PackagesDialog.h"
@@ -1291,6 +1292,9 @@ void AssetViewUtils::SetPathColor(const FString& FolderPath, TOptional<FLinearCo
 
 		// Update the map too
 		PathColors.Add(InPath, InFolderColor);
+
+		FContentBrowserModule& ContentBrowserModule = FModuleManager::Get().GetModuleChecked<FContentBrowserModule>(TEXT("ContentBrowser"));
+		ContentBrowserModule.GetOnSetFolderColor().Broadcast(InPath);
 	};
 
 	auto RemoveColorInternal = [](const FString& InPath)
diff --git a/Engine/Source/Editor/ContentBrowser/Private/Experimental/ContentBrowserExtensionUtils.cpp b/Engine/Source/Editor/ContentBrowser/Private/Experimental/ContentBrowserExtensionUtils.cpp
new file mode 100644
--- /dev/null
+++ b/Engine/Source/Editor/ContentBrowser/Private/Experimental/ContentBrowserExtensionUtils.cpp
@@ -0,0 +1,38 @@
+// Copyright Epic Games, Inc. All Rights Reserved.
+
+#include "Experimental/ContentBrowserExtensionUtils.h"
+#include "ContentBrowserDataSubsystem.h"
+#include "ContentBrowserUtils.h"
+#include "CollectionViewUtils.h"
+#include "ContentBrowserItemPath.h"
+#include "IContentBrowserDataModule.h"
+
+namespace UE::Editor::ContentBrowser::ExtensionUtils
+{
+	TOptional<FLinearColor> GetFolderColor(const FName& FolderPath)
+	{
+		const FName VirtualPath = IContentBrowserDataModule::Get().GetSubsystem()->ConvertInternalPathToVirtual(FolderPath);
+
+		FName CollectionName;
+		ECollectionShareType::Type CollectionFolderShareType = ECollectionShareType::CST_All;
+
+		if(ContentBrowserUtils::IsCollectionPath(VirtualPath.ToString(), &CollectionName, &CollectionFolderShareType))
+		{
+			if(TOptional<FLinearColor> Color = CollectionViewUtils::GetCustomColor(CollectionName, CollectionFolderShareType))
+			{
+				return Color;
+			}
+		}
+		if (TOptional<FLinearColor> Color = ContentBrowserUtils::GetPathColor(FolderPath.ToString()))
+		{
+			return Color.GetValue();
+		}
+
+		return TOptional<FLinearColor>();
+	}
+
+	void SetFolderColor(const FName& FolderPath, const FLinearColor& FolderColor)
+	{
+		ContentBrowserUtils::SetPathColor(FolderPath.ToString(), FolderColor);
+	}
+}
diff --git a/Engine/Source/Editor/ContentBrowser/Public/ContentBrowserDelegates.h b/Engine/Source/Editor/ContentBrowser/Public/ContentBrowserDelegates.h
--- a/Engine/Source/Editor/ContentBrowser/Public/ContentBrowserDelegates.h
+++ b/Engine/Source/Editor/ContentBrowser/Public/ContentBrowserDelegates.h
@@ -131,6 +131,9 @@ DECLARE_DELEGATE_OneParam(FOnObjectPathChosenForSave, const FString& /*ObjectPat
 /** Called when custom behavior is needed for allowing folders to toggle their private content edit state */
 DECLARE_DELEGATE_RetVal_OneParam(bool, FIsFolderShowPrivateContentToggleableDelegate, const FStringView /*VirtualFolderPath*/);
 
+/** Called when a custom folder color is set on a folder */
+DECLARE_MULTICAST_DELEGATE_OneParam(FOnSetFolderColor, const FString& /*FolderPath*/)
+
 /** Contains the delegates used to handle a custom drag-and-drop in the asset view */
 struct FAssetViewDragAndDropExtender
 {
diff --git a/Engine/Source/Editor/ContentBrowser/Public/ContentBrowserModule.h b/Engine/Source/Editor/ContentBrowser/Public/ContentBrowserModule.h
--- a/Engine/Source/Editor/ContentBrowser/Public/ContentBrowserModule.h
+++ b/Engine/Source/Editor/ContentBrowser/Public/ContentBrowserModule.h
@@ -214,6 +214,10 @@ public:
 		return AssetClassesRequiringDynamicTags.Contains(InName);
 	}
 	
+	FOnSetFolderColor& GetOnSetFolderColor()
+	{
+		return OnSetFolderColor;
+	}
 	
 
 private:
@@ -223,6 +227,9 @@ private:
 	/** List of asset classes whose tags are dynamic and therefore we should union all asset's tags rather than grabbing the first available. */
 	TArray<FName> AssetClassesRequiringDynamicTags;
 
+	/** Called when a custom folder color is set on a folder */
+	FOnSetFolderColor OnSetFolderColor;
+
 private:
 	IContentBrowserSingleton* ContentBrowserSingleton;
 	TSharedPtr<class FContentBrowserSpawner> ContentBrowserSpawner;
diff --git a/Engine/Source/Editor/ContentBrowser/Public/Experimental/ContentBrowserExtensionUtils.h b/Engine/Source/Editor/ContentBrowser/Public/Experimental/ContentBrowserExtensionUtils.h
new file mode 100644
--- /dev/null
+++ b/Engine/Source/Editor/ContentBrowser/Public/Experimental/ContentBrowserExtensionUtils.h
@@ -0,0 +1,15 @@
+// Copyright Epic Games, Inc. All Rights Reserved.
+
+#pragma once
+
+#include "Containers/UnrealString.h"
+#include "Math/Color.h"
+
+namespace UE::Editor::ContentBrowser::ExtensionUtils
+{
+	// Get the custom color of the given folder path (if any)
+	CONTENTBROWSER_API TOptional<FLinearColor> GetFolderColor(const FName& FolderPath);
+
+	// Set a custom color for the given folder path
+	CONTENTBROWSER_API void SetFolderColor(const FName& FolderPath, const FLinearColor& FolderColor);
+}
diff --git a/Engine/Source/Runtime/TypedElementFramework/Public/Elements/Columns/TypedElementFolderColumns.h b/Engine/Source/Runtime/TypedElementFramework/Public/Elements/Columns/TypedElementFolderColumns.h
new file mode 100644
--- /dev/null
+++ b/Engine/Source/Runtime/TypedElementFramework/Public/Elements/Columns/TypedElementFolderColumns.h
@@ -0,0 +1,16 @@
+// Copyright Epic Games, Inc. All Rights Reserved.
+
+#pragma once
+
+#include "Elements/Interfaces/TypedElementDataStorageInterface.h"
+
+#include "TypedElementFolderColumns.generated.h"
+
+/**
+ * Column that signifies this row is for a folder
+ */
+USTRUCT(meta = (DisplayName = "Folder"))
+struct FFolderTag final : public FEditorDataStorageTag
+{
+	GENERATED_BODY()
+};
\ No newline at end of file
