TEDS-Outliner: Allow the label column to use different widget constructors for each row depending on the columns it has (and the purpose)

diff --git a/Engine/Plugins/Experimental/EditorDataStorageFeatures/Source/TedsOutliner/Private/TedsOutlinerImpl.cpp b/Engine/Plugins/Experimental/EditorDataStorageFeatures/Source/TedsOutliner/Private/TedsOutlinerImpl.cpp
--- a/Engine/Plugins/Experimental/EditorDataStorageFeatures/Source/TedsOutliner/Private/TedsOutlinerImpl.cpp
+++ b/Engine/Plugins/Experimental/EditorDataStorageFeatures/Source/TedsOutliner/Private/TedsOutlinerImpl.cpp
@@ -99,65 +99,6 @@ FTedsOutlinerImpl::FTedsOutlinerImpl(const FTedsOutlinerParams& InParams, IScene
 
 }
 
-void FTedsOutlinerImpl::CreateLabelWidgetConstructors()
-{
-	using namespace UE::Editor::DataStorage::Queries;
-
-	// Create and assign the widget constructors for the label column
-	
-	auto CreateWidgetConstructorForQuery = [this](DataStorage::FQueryDescription InQueryDescription) -> TSharedPtr<FTypedElementWidgetConstructor>
-	{
-		// Create the Widget Constructor for the item label column
-		using MatchApproach = ITypedElementDataStorageUiInterface::EMatchApproach;
-
-		// Make a copy of the columns because CreateWidgetConstructors can modify it
-		TArray<TWeakObjectPtr<const UScriptStruct>> ColumnTypes(InQueryDescription.SelectionTypes);
-
-		TSharedPtr<FTypedElementWidgetConstructor> OutWidgetConstructorPtr;
-
-		bool bFoundWidget = false;
-
-		// We also want to look at the RowLabel purpose for the label
-		TArray<FName> ItemLabelCellWidgetPurposes{TEXT("SceneOutliner.RowLabel")};
-		ItemLabelCellWidgetPurposes.Append(CellWidgetPurposes);
-		
-		for(const FName& WidgetPurpose : ItemLabelCellWidgetPurposes)
-		{
-			StorageUi->CreateWidgetConstructors(WidgetPurpose, MatchApproach::ExactMatch, ColumnTypes, {},
-				[&OutWidgetConstructorPtr, ColumnTypes, &bFoundWidget](
-				TUniquePtr<FTypedElementWidgetConstructor> CreatedConstructor, 
-				TConstArrayView<TWeakObjectPtr<const UScriptStruct>> MatchedColumnTypes)
-				{
-					if (ColumnTypes.Num() == MatchedColumnTypes.Num())
-					{
-						OutWidgetConstructorPtr = TSharedPtr<FTypedElementWidgetConstructor>(CreatedConstructor.Release());
-						bFoundWidget = true;
-					}
-					// Either this was the exact match so no need to search further or the longest possible chain didn't match so the next ones will 
-					// always be shorter in both cases just return.
-					return false;
-				});
-
-			if (bFoundWidget)
-			{
-				break;
-			}
-		}
-
-		return OutWidgetConstructorPtr;
-	};
-	
-	DataStorage::QueryHandle LabelColumnQueryHandle = Storage->RegisterQuery(
-		Select()
-			.ReadWrite<FTypedElementLabelColumn, FTypedElementClassTypeInfoColumn>()
-		.Compile());
-
-	if (TSharedPtr<FTypedElementWidgetConstructor> LabelColumnWidgetConstructor = CreateWidgetConstructorForQuery(Storage->GetQueryDescription(LabelColumnQueryHandle)))
-	{
-		QueryToWidgetConstructorMap.Emplace(LabelColumnQueryHandle, LabelColumnWidgetConstructor);
-	}
-}
-
 void FTedsOutlinerImpl::CreateFilterQueries()
 {
 	using namespace UE::Editor::DataStorage::Queries;
@@ -209,7 +150,6 @@ void FTedsOutlinerImpl::CreateFilterQueries()
 
 void FTedsOutlinerImpl::Init()
 {
-	CreateLabelWidgetConstructors();
 	CreateFilterQueries();
 	
 	if (SelectionSetName.IsSet())
@@ -227,14 +167,6 @@ void FTedsOutlinerImpl::Init()
 FTedsOutlinerImpl::~FTedsOutlinerImpl()
 {
 	UnregisterQueries();
-
-	// This is done outside of unregister queries because that is called when the internal query changes (i.e filter) but we don't want to unregister the
-	// label widget queries until shutdown
-	for(const TPair<DataStorage::QueryHandle, TSharedPtr<FTypedElementWidgetConstructor>>& QueryConstructorPair : QueryToWidgetConstructorMap)
-	{
-		Storage->UnregisterQuery(QueryConstructorPair.Key);
-	}
-	
 	FTSTicker::GetCoreTicker().RemoveTicker(TickerHandle);
 }
 
@@ -260,25 +192,49 @@ void FTedsOutlinerImpl::SetSelection(const TArray<DataStorage::RowHandle>& InSel
 
 TSharedRef<SWidget> FTedsOutlinerImpl::CreateLabelWidgetForItem(DataStorage::RowHandle InRowHandle, const STableRow<FSceneOutlinerTreeItemPtr>& InRow) const
 {
-	auto CreateWidgetForQuery = [InRowHandle, this, &InRow](const TPair<DataStorage::QueryHandle, TSharedPtr<FTypedElementWidgetConstructor>>& QueryConstructorPair) -> TSharedPtr<SWidget>
+	// Find the best matching label widget for a given list of columns
+	auto CreateWidgetConstructor = [this](TArray<TWeakObjectPtr<const UScriptStruct>>& ColumnTypes) -> TSharedPtr<FTypedElementWidgetConstructor>
 	{
-		DataStorage::FQueryDescription QueryDescription = Storage->GetQueryDescription(QueryConstructorPair.Key);
-		
-		// Create a generic metadata view for the Type Widget
-		DataStorage::FMetaData QueryWideMetaData;
-		QueryWideMetaData.AddImmutableData("TypeInfoWidget_bUseIcon", true);
-		DataStorage::FGenericMetaDataView GenericMetaDataView(QueryWideMetaData);
+		TSharedPtr<FTypedElementWidgetConstructor> OutWidgetConstructorPtr;
 
-		// Create metadata for the query itself
-		DataStorage::FQueryMetaDataView QueryMetaDataView(QueryDescription);
+		bool bFoundWidget = false;
 
-		// Combine the two metadata
-		DataStorage::FComboMetaDataView MetaDataArgs(GenericMetaDataView, QueryMetaDataView);
+		static TArray<FName> Purposes{TEXT("SceneOutliner.RowLabel"), TEXT("General.RowLabel")};
 
-		TArray<TWeakObjectPtr<const UScriptStruct>> ColumnTypes(QueryDescription.SelectionTypes);
-		TSharedPtr<FTypedElementWidgetConstructor> CellWidgetConstructor = QueryConstructorPair.Value;
+		for(const FName& Purpose : Purposes)
+		{
+			StorageUi->CreateWidgetConstructors(Purpose, ITypedElementDataStorageUiInterface::EMatchApproach::LongestMatch, ColumnTypes, {},
+					[&OutWidgetConstructorPtr, ColumnTypes, &bFoundWidget](
+					TUniquePtr<FTypedElementWidgetConstructor> CreatedConstructor, 
+					TConstArrayView<TWeakObjectPtr<const UScriptStruct>> MatchedColumnTypes)
+					{
+						OutWidgetConstructorPtr = TSharedPtr<FTypedElementWidgetConstructor>(CreatedConstructor.Release());
+						bFoundWidget = true;
+					
+						// Either this was the exact match so no need to search further or the longest possible chain didn't match so the next ones will 
+						// always be shorter in both cases just return.
+						return false;
+					});
 
-		DataStorage::RowHandle UiRowHandle = Storage->AddRow(Storage->FindTable(DataStorage::TableViewerUtils::GetWidgetTableName()));
+			if(bFoundWidget)
+			{
+				break;
+			}
+		}
+		return OutWidgetConstructorPtr;
+	};
+	
+	using namespace DataStorage::Queries;
+	// Query description to pass as metadata to allow the label column to be writable
+	static FQueryDescription MetaDataQueryDescription = Select().ReadWrite<FTypedElementLabelColumn>().Where().Compile();
+
+	// Create a widget for this row using the given widget constructor
+	auto CreateWidget = [InRowHandle, this, &InRow](const TSharedPtr<FTypedElementWidgetConstructor>& WidgetConstructor) -> TSharedPtr<SWidget>
+	{
+		// Create metadata for the query
+		FQueryMetaDataView QueryMetaDataView(MetaDataQueryDescription);
+
+		RowHandle UiRowHandle = Storage->AddRow(Storage->FindTable(DataStorage::TableViewerUtils::GetWidgetTableName()));
 
 		if (FTypedElementRowReferenceColumn* RowReference = Storage->GetColumn<FTypedElementRowReferenceColumn>(UiRowHandle))
 		{
@@ -287,7 +243,7 @@ TSharedRef<SWidget> FTedsOutlinerImpl::CreateLabelWidgetForItem(DataStorage::Row
 
 		Storage->AddColumn(UiRowHandle, FTedsOutlinerColumn{.Outliner = StaticCastSharedRef<ISceneOutliner>(SceneOutliner->AsShared())});
 		
-		TSharedPtr<SWidget> Widget = StorageUi->ConstructWidget(UiRowHandle, *CellWidgetConstructor, MetaDataArgs);
+		TSharedPtr<SWidget> Widget = StorageUi->ConstructWidget(UiRowHandle, *WidgetConstructor, QueryMetaDataView);
 		
 		if (FExternalWidgetSelectionColumn* ExternalWidgetSelectionColumn = Storage->GetColumn<FExternalWidgetSelectionColumn>(UiRowHandle))
 		{
@@ -296,25 +252,29 @@ TSharedRef<SWidget> FTedsOutlinerImpl::CreateLabelWidgetForItem(DataStorage::Row
 		return Widget;
 		
 	};
-	
-	TSharedRef<SHorizontalBox> CombinedWidget = SNew(SHorizontalBox);
 
-	for(const TPair<DataStorage::QueryHandle, TSharedPtr<FTypedElementWidgetConstructor>>& QueryConstructorPair : QueryToWidgetConstructorMap)
+	// Get all the columns on the given row
+	TArray<TWeakObjectPtr<const UScriptStruct>> Columns;
+	Storage->ListColumns(InRowHandle, [&Columns](const UScriptStruct& ColumnType)
 	{
-		if (TSharedPtr<SWidget> WidgetForQuery = CreateWidgetForQuery(QueryConstructorPair))
+		Columns.Add(&ColumnType);
+	});
+
+	// Create the label widget
+	if(const TSharedPtr<FTypedElementWidgetConstructor> WidgetConstructor = CreateWidgetConstructor(Columns))
+	{
+		if(const TSharedPtr<SWidget> Widget = CreateWidget(WidgetConstructor))
 		{
-			CombinedWidget->AddSlot()
-					.AutoWidth()
+			return SNew(SBox)
 					.HAlign(HAlign_Left)
 					.VAlign(VAlign_Center)
-					.Padding(2.0f, 0.0f, 4.0f, 0.0f)
 					[
-						WidgetForQuery.ToSharedRef()
+						Widget.ToSharedRef()
 					];
 		}
 	}
-	
-	return CombinedWidget;
+
+	return SNullWidget::NullWidget;
 }
 
 void FTedsOutlinerImpl::AppendQuery(DataStorage::FQueryDescription& Query1, const DataStorage::FQueryDescription& Query2)
diff --git a/Engine/Plugins/Experimental/EditorDataStorageFeatures/Source/TedsOutliner/Public/TedsOutlinerImpl.h b/Engine/Plugins/Experimental/EditorDataStorageFeatures/Source/TedsOutliner/Public/TedsOutlinerImpl.h
--- a/Engine/Plugins/Experimental/EditorDataStorageFeatures/Source/TedsOutliner/Public/TedsOutlinerImpl.h
+++ b/Engine/Plugins/Experimental/EditorDataStorageFeatures/Source/TedsOutliner/Public/TedsOutlinerImpl.h
@@ -172,7 +172,6 @@ protected:
 	void ClearSelection() const;
 	void Tick();
 
-	void CreateLabelWidgetConstructors();
 	void CreateFilterQueries();
 
 	// Check if this row can be displayed in this table viewer
@@ -186,9 +185,6 @@ protected:
 
 	FTedsOutlinerParams CreationParams;
 
-	// Widget constructor to create the label widget
-	TArray<TPair<DataStorage::QueryHandle, TSharedPtr<FTypedElementWidgetConstructor>>> QueryToWidgetConstructorMap;
-
 	// Widget purposes this table viewer supports
 	TArray<FName> CellWidgetPurposes;
 	
