[Teds] Rename interface types and base factory class

diff --git a/Engine/Plugins/Experimental/EditorDataStorage/Source/TedsUI/Public/Widgets/PackagePathWidget.h b/Engine/Plugins/Experimental/EditorDataStorage/Source/TedsUI/Public/Widgets/PackagePathWidget.h
--- a/Engine/Plugins/Experimental/EditorDataStorage/Source/TedsUI/Public/Widgets/PackagePathWidget.h
+++ b/Engine/Plugins/Experimental/EditorDataStorage/Source/TedsUI/Public/Widgets/PackagePathWidget.h
@@ -7,20 +7,20 @@
 
 #include "PackagePathWidget.generated.h"
 
-class ITypedElementDataStorageInterface;
+class IEditorDataStorageProvider;
 class SWidget;
 class UScriptStruct;
 
 UCLASS()
-class TEDSUI_API UPackagePathWidgetFactory : public UTypedElementDataStorageFactory
+class TEDSUI_API UPackagePathWidgetFactory : public UEditorDataStorageFactory
 {
 	GENERATED_BODY()
 
 public:
 	~UPackagePathWidgetFactory() override = default;
 
-	void RegisterWidgetConstructors(ITypedElementDataStorageInterface& DataStorage,
-		ITypedElementDataStorageUiInterface& DataStorageUi) const override;
+	void RegisterWidgetConstructors(IEditorDataStorageProvider& DataStorage,
+		IEditorDataStorageUiProvider& DataStorageUi) const override;
 };
 
 USTRUCT()
@@ -36,7 +36,7 @@ protected:
 	explicit FPackagePathWidgetConstructor(const UScriptStruct* InTypeInfo);
 
 	TSharedPtr<SWidget> CreateWidget(const UE::Editor::DataStorage::FMetaDataView& Arguments) override;
-	bool FinalizeWidget(ITypedElementDataStorageInterface* DataStorage, ITypedElementDataStorageUiInterface* DataStorageUi,
+	bool FinalizeWidget(IEditorDataStorageProvider* DataStorage, IEditorDataStorageUiProvider* DataStorageUi,
 		UE::Editor::DataStorage::RowHandle Row, const TSharedPtr<SWidget>& Widget) override;
 };
 
@@ -50,6 +50,6 @@ public:
 	~FLoadedPackagePathWidgetConstructor() override = default;
 
 protected:
-	bool FinalizeWidget(ITypedElementDataStorageInterface* DataStorage, ITypedElementDataStorageUiInterface* DataStorageUi,
+	bool FinalizeWidget(IEditorDataStorageProvider* DataStorage, IEditorDataStorageUiProvider* DataStorageUi,
 		UE::Editor::DataStorage::RowHandle Row, const TSharedPtr<SWidget>& Widget) override;
 };
diff --git a/Engine/Plugins/Experimental/EditorDataStorageFeatures/Source/TedsAssetData/Private/CB/TedsAssetDataCBDataSource.cpp b/Engine/Plugins/Experimental/EditorDataStorageFeatures/Source/TedsAssetData/Private/CB/TedsAssetDataCBDataSource.cpp
--- a/Engine/Plugins/Experimental/EditorDataStorageFeatures/Source/TedsAssetData/Private/CB/TedsAssetDataCBDataSource.cpp
+++ b/Engine/Plugins/Experimental/EditorDataStorageFeatures/Source/TedsAssetData/Private/CB/TedsAssetDataCBDataSource.cpp
@@ -47,7 +47,7 @@ TAutoConsoleVariable<bool> CVarTEDSAssetDataCBSourceIncludeTagsAndValues(TEXT("T
 		}
 	}));
 
-FTedsAssetDataCBDataSource::FTedsAssetDataCBDataSource(ITypedElementDataStorageInterface& InDatabase)
+FTedsAssetDataCBDataSource::FTedsAssetDataCBDataSource(IEditorDataStorageProvider& InDatabase)
 	: Database(InDatabase)
 {
 	using namespace UE::Editor::DataStorage;
diff --git a/Engine/Plugins/Experimental/EditorDataStorageFeatures/Source/TedsAssetData/Private/CB/TedsAssetDataCBDataSource.h b/Engine/Plugins/Experimental/EditorDataStorageFeatures/Source/TedsAssetData/Private/CB/TedsAssetDataCBDataSource.h
--- a/Engine/Plugins/Experimental/EditorDataStorageFeatures/Source/TedsAssetData/Private/CB/TedsAssetDataCBDataSource.h
+++ b/Engine/Plugins/Experimental/EditorDataStorageFeatures/Source/TedsAssetData/Private/CB/TedsAssetDataCBDataSource.h
@@ -13,7 +13,7 @@ class FAssetPackageData;
 class FAssetPropertyTagCache;
 class IAssetRegistry;
 class IPlugin;
-class ITypedElementDataStorageInterface;
+class IEditorDataStorageProvider;
 
 struct FAssetData;
 struct FAssetDataColumn_Experimental;
@@ -30,7 +30,7 @@ namespace UE::Editor::AssetData::Private
 class FTedsAssetDataCBDataSource
 {
 public:
-	explicit FTedsAssetDataCBDataSource(ITypedElementDataStorageInterface& InDatabase);
+	explicit FTedsAssetDataCBDataSource(IEditorDataStorageProvider& InDatabase);
 	~FTedsAssetDataCBDataSource();
 	
 	// Toggle the storage of the metadata from the asset registry into teds 
@@ -75,7 +75,7 @@ private:
 	void ProcessAssetDataAndPathUpdateQueryCallback(DataStorage::IQueryContext& Context, const DataStorage::RowHandle* Rows, const FAssetDataColumn_Experimental* AssetDataColumn) const;
 	void ProcessAssetDataUpdateQueryCallback(DataStorage::IQueryContext& Context, const DataStorage::RowHandle*Rows, const FAssetDataColumn_Experimental* AssetDataColumn) const;
 
-	ITypedElementDataStorageInterface& Database;
+	IEditorDataStorageProvider& Database;
 	DataStorage::QueryHandle ProcessPathQuery;
 	DataStorage::QueryHandle ProcessAssetDataPathUpdateQuery;
 	DataStorage::QueryHandle ProcessAssetDataAndPathUpdateQuery;
diff --git a/Engine/Plugins/Experimental/EditorDataStorageFeatures/Source/TedsAssetData/Private/Debug/TedsAssetDataDebugCommands.cpp b/Engine/Plugins/Experimental/EditorDataStorageFeatures/Source/TedsAssetData/Private/Debug/TedsAssetDataDebugCommands.cpp
--- a/Engine/Plugins/Experimental/EditorDataStorageFeatures/Source/TedsAssetData/Private/Debug/TedsAssetDataDebugCommands.cpp
+++ b/Engine/Plugins/Experimental/EditorDataStorageFeatures/Source/TedsAssetData/Private/Debug/TedsAssetDataDebugCommands.cpp
@@ -30,7 +30,7 @@ namespace UE::Editor::DataStorage::Debug::Private
 		TEXT("Print some debug information on the specified path."),
 		FConsoleCommandWithArgsDelegate::CreateLambda([](const TArray<FString>& InArgs)
 			{
-				const ITypedElementDataStorageInterface* Database = UTypedElementRegistry::GetInstance()->GetDataStorage();
+				const IEditorDataStorageProvider* Database = UTypedElementRegistry::GetInstance()->GetDataStorage();
 				IAssetRegistry& AssetRegistry = IAssetRegistry::GetChecked();
 
 				for (const FString& Path : InArgs)
@@ -160,7 +160,7 @@ namespace UE::Editor::DataStorage::Debug::Private
 		ULevel::StaticClass(), UBlueprint::StaticClass(), UTexture::StaticClass() };
 
 	// Populate an asset row with random information
-	void PopulateRowWithRandomInfo(RowHandle Row, ITypedElementDataStorageInterface* DataStorage)
+	void PopulateRowWithRandomInfo(RowHandle Row, IEditorDataStorageProvider* DataStorage)
 	{
 		// Don't modify any rows that aren't our placeholder assets
 		if (!DataStorage->HasColumns<FAssetTag>(Row))
@@ -220,7 +220,7 @@ namespace UE::Editor::DataStorage::Debug::Private
 				TRACE_CPUPROFILER_EVENT_SCOPE(TEDS.Debug.CreateDebugAssetRows);
 
 
-				ITypedElementDataStorageInterface* DataStorage = UTypedElementRegistry::GetInstance()->GetMutableDataStorage();
+				IEditorDataStorageProvider* DataStorage = UTypedElementRegistry::GetInstance()->GetMutableDataStorage();
 
 				if (!DataStorage)
 				{
@@ -252,7 +252,7 @@ namespace UE::Editor::DataStorage::Debug::Private
 			{
 				using namespace UE::Editor::DataStorage::Queries;
 
-				ITypedElementDataStorageInterface* DataStorage = UTypedElementRegistry::GetInstance()->GetMutableDataStorage();
+				IEditorDataStorageProvider* DataStorage = UTypedElementRegistry::GetInstance()->GetMutableDataStorage();
 
 				if (!DataStorage)
 				{
@@ -265,7 +265,7 @@ namespace UE::Editor::DataStorage::Debug::Private
 				TArray<RowHandle> Rows;
 				
 				DataStorage->RunQuery(AssetQueryHandle,
-					CreateDirectQueryCallbackBinding([&Rows](const ITypedElementDataStorageInterface::IDirectQueryContext& Context)
+					CreateDirectQueryCallbackBinding([&Rows](const IEditorDataStorageProvider::IDirectQueryContext& Context)
 					{
 						Rows.Append(Context.GetRowHandles());
 					}));
diff --git a/Engine/Plugins/Experimental/EditorDataStorageFeatures/Source/TedsAssetData/Private/Processors/AssetProcessors.cpp b/Engine/Plugins/Experimental/EditorDataStorageFeatures/Source/TedsAssetData/Private/Processors/AssetProcessors.cpp
--- a/Engine/Plugins/Experimental/EditorDataStorageFeatures/Source/TedsAssetData/Private/Processors/AssetProcessors.cpp
+++ b/Engine/Plugins/Experimental/EditorDataStorageFeatures/Source/TedsAssetData/Private/Processors/AssetProcessors.cpp
@@ -17,7 +17,7 @@
 #include "Elements/Columns/TypedElementMiscColumns.h"
 #include "Elements/Framework/TypedElementIndexHasher.h"
 
-void UTedsAssetDataFactory::RegisterQueries(ITypedElementDataStorageInterface& DataStorage)
+void UTedsAssetDataFactory::RegisterQueries(IEditorDataStorageProvider& DataStorage)
 {
 	using namespace UE::Editor::DataStorage::Queries;
 
@@ -65,7 +65,7 @@ void UTedsAssetDataFactory::RegisterQueries(ITypedElementDataStorageInterface& D
 		);
 }
 
-void UTedsAssetDataFactory::PreRegister(ITypedElementDataStorageInterface& DataStorage)
+void UTedsAssetDataFactory::PreRegister(IEditorDataStorageProvider& DataStorage)
 {
 	if (FContentBrowserModule* ContentBrowserModule = FModuleManager::Get().GetModulePtr<FContentBrowserModule>("ContentBrowser"))
 	{
@@ -73,7 +73,7 @@ void UTedsAssetDataFactory::PreRegister(ITypedElementDataStorageInterface& DataS
 	}
 }
 
-void UTedsAssetDataFactory::PreShutdown(ITypedElementDataStorageInterface& DataStorage)
+void UTedsAssetDataFactory::PreShutdown(IEditorDataStorageProvider& DataStorage)
 {
 	if (FContentBrowserModule* ContentBrowserModule = FModuleManager::Get().GetModulePtr<FContentBrowserModule>("ContentBrowser"))
 	{
@@ -81,7 +81,7 @@ void UTedsAssetDataFactory::PreShutdown(ITypedElementDataStorageInterface& DataS
 	}
 }
 
-void UTedsAssetDataFactory::OnSetFolderColor(const FString& Path, ITypedElementDataStorageInterface* DataStorage)
+void UTedsAssetDataFactory::OnSetFolderColor(const FString& Path, IEditorDataStorageProvider* DataStorage)
 {
 	const UE::Editor::DataStorage::IndexHash PathHash = UE::Editor::DataStorage::GenerateIndexHash(FName(Path));
 	const UE::Editor::DataStorage::RowHandle Row = DataStorage->FindIndexedRow(PathHash);
diff --git a/Engine/Plugins/Experimental/EditorDataStorageFeatures/Source/TedsAssetData/Private/Processors/AssetProcessors.h b/Engine/Plugins/Experimental/EditorDataStorageFeatures/Source/TedsAssetData/Private/Processors/AssetProcessors.h
--- a/Engine/Plugins/Experimental/EditorDataStorageFeatures/Source/TedsAssetData/Private/Processors/AssetProcessors.h
+++ b/Engine/Plugins/Experimental/EditorDataStorageFeatures/Source/TedsAssetData/Private/Processors/AssetProcessors.h
@@ -6,21 +6,21 @@
 
 #include "AssetProcessors.generated.h"
 
-class ITypedElementDataStorageInterface;
+class IEditorDataStorageProvider;
 
 UCLASS()
-class UTedsAssetDataFactory : public UTypedElementDataStorageFactory
+class UTedsAssetDataFactory : public UEditorDataStorageFactory
 {
 	GENERATED_BODY()
 
 public:
 	~UTedsAssetDataFactory() override = default;
 
-	virtual void RegisterQueries(ITypedElementDataStorageInterface& DataStorage) override;
-	virtual void PreRegister(ITypedElementDataStorageInterface& DataStorage) override;
-	virtual void PreShutdown(ITypedElementDataStorageInterface& DataStorage) override;
+	void RegisterQueries(IEditorDataStorageProvider& DataStorage) override;
+	virtual void PreRegister(IEditorDataStorageProvider& DataStorage) override;
+	virtual void PreShutdown(IEditorDataStorageProvider& DataStorage) override;
 
 protected:
 
-	void OnSetFolderColor(const FString& Path, ITypedElementDataStorageInterface* DataStorage);
+	void OnSetFolderColor(const FString& Path, IEditorDataStorageProvider* DataStorage);
 };
diff --git a/Engine/Plugins/Experimental/EditorDataStorageFeatures/Source/TedsAssetData/Private/TedsAssetData.cpp b/Engine/Plugins/Experimental/EditorDataStorageFeatures/Source/TedsAssetData/Private/TedsAssetData.cpp
--- a/Engine/Plugins/Experimental/EditorDataStorageFeatures/Source/TedsAssetData/Private/TedsAssetData.cpp
+++ b/Engine/Plugins/Experimental/EditorDataStorageFeatures/Source/TedsAssetData/Private/TedsAssetData.cpp
@@ -36,7 +36,7 @@ struct FPopulateAssetDataRowArgs
 
 // Only safe if the GT is blocked during the operation
 template<typename TAssetData>
-FPopulateAssetDataRowArgs ThreadSafe_PopulateAssetDataTableRow(TAssetData&& InAssetData, const ITypedElementDataStorageInterface& Database)
+FPopulateAssetDataRowArgs ThreadSafe_PopulateAssetDataTableRow(TAssetData&& InAssetData, const IEditorDataStorageProvider& Database)
 {
 	FPopulateAssetDataRowArgs Output;
 	Output.ObjectPathHash = GenerateIndexHash(InAssetData.GetSoftObjectPath());
@@ -58,7 +58,7 @@ FPopulateAssetDataRowArgs ThreadSafe_PopulateAssetDataTableRow(TAssetData&& InAs
 }
 
 
-void PopulateAssetDataTableRow(FPopulateAssetDataRowArgs&& InAssetDataRowArgs, ITypedElementDataStorageInterface& Database, RowHandle RowHandle)
+void PopulateAssetDataTableRow(FPopulateAssetDataRowArgs&& InAssetDataRowArgs, IEditorDataStorageProvider& Database, RowHandle RowHandle)
 {
 	if (InAssetDataRowArgs.PathRow !=  InvalidRowHandle)
 	{
@@ -141,7 +141,7 @@ FPopulatePathRowArgs ThreadSafe_PopulatePathRowArgs(IndexHash AssetRegistryPathH
 	return Args;
 }
 
-void PopulatePathDataTableRow(FPopulatePathRowArgs&& InPopulatePathRowArgs, ITypedElementDataStorageInterface& Database, RowHandle InRowHandle)
+void PopulatePathDataTableRow(FPopulatePathRowArgs&& InPopulatePathRowArgs, IEditorDataStorageProvider& Database, RowHandle InRowHandle)
 {
 	if (InPopulatePathRowArgs.ParentAssetRegistryPathHash != InvalidRowHandle)
 	{
@@ -165,7 +165,7 @@ void PopulatePathDataTableRow(FPopulatePathRowArgs&& InPopulatePathRowArgs, ITyp
 }
 
 
-FTedsAssetData::FTedsAssetData(ITypedElementDataStorageInterface& InDatabase)
+FTedsAssetData::FTedsAssetData(IEditorDataStorageProvider& InDatabase)
 	: Database(InDatabase)
 {
 	using namespace UE::Editor::DataStorage::Queries;
diff --git a/Engine/Plugins/Experimental/EditorDataStorageFeatures/Source/TedsAssetData/Private/TedsAssetData.h b/Engine/Plugins/Experimental/EditorDataStorageFeatures/Source/TedsAssetData/Private/TedsAssetData.h
--- a/Engine/Plugins/Experimental/EditorDataStorageFeatures/Source/TedsAssetData/Private/TedsAssetData.h
+++ b/Engine/Plugins/Experimental/EditorDataStorageFeatures/Source/TedsAssetData/Private/TedsAssetData.h
@@ -8,46 +8,43 @@
 
 struct FAssetData;
 
-class ITypedElementDataStorageInterface;
+class IEditorDataStorageProvider;
 
-namespace UE::Editor
+namespace UE::Editor::AssetData::Private
 {
-	namespace AssetData::Private
+	/**
+	 * Manage the registration and life cycle of the row related representing the data from the asset registry into TEDS.
+	 */
+	class FTedsAssetData
 	{
-		/**
-		 * Manage the registration and life cycle of the row related representing the data from the asset registry into TEDS.
-		 */
-		class FTedsAssetData
-		{
-		public:
-			FTedsAssetData(ITypedElementDataStorageInterface& InDatabase);
-
-			~FTedsAssetData();
-
-			void ProcessAllEvents();
-
-		private:
-			void OnAssetsAdded(TConstArrayView<FAssetData> InAssetsAdded);
-			void OnAssetsRemoved(TConstArrayView<FAssetData> InAssetsRemoved);
-			void OnAssetsUpdated(TConstArrayView<FAssetData> InAssetsUpdated);
-			void OnAssetsUpdatedOnDisk(TConstArrayView<FAssetData> InAssetsUpdated);
-
-			void OnAssetRenamed(const FAssetData& InAsset, const FString& InOldObjectPath);
-
-			void OnPathsAdded(TConstArrayView<FStringView> InPathsAdded);
-			void OnPathsRemoved(TConstArrayView<FStringView> InPathsRemoved);
-
-			ITypedElementDataStorageInterface& Database;
-			DataStorage::TableHandle PathsTable = DataStorage::InvalidTableHandle;
-			DataStorage::TableHandle AssetsDataTable = DataStorage::InvalidTableHandle;
-
-			DataStorage::QueryHandle UpdateAssetsInPathQuery = DataStorage::InvalidQueryHandle;
-			DataStorage::QueryHandle UpdateAssetsInPath = DataStorage::InvalidQueryHandle;
-			DataStorage::QueryHandle ResolveMissingAssetInPathQuery = DataStorage::InvalidQueryHandle;
-			DataStorage::QueryHandle UpdateParentToChildrenAssetPathQuery = DataStorage::InvalidQueryHandle;
-			DataStorage::QueryHandle ResolveMissingParentPathQuery = DataStorage::InvalidQueryHandle;
-			DataStorage::QueryHandle RemoveUpdatedPathTagQuery = DataStorage::InvalidQueryHandle;
-			DataStorage::QueryHandle RemoveUpdatedAssetDataTagQuery = DataStorage::InvalidQueryHandle;
-		};
-	} // namespace AssetData::Private
-} // namespace UE::Editor
+	public:
+		FTedsAssetData(IEditorDataStorageProvider& InDatabase);
+
+		~FTedsAssetData();
+
+		void ProcessAllEvents();
+
+	private:
+		void OnAssetsAdded(TConstArrayView<FAssetData> InAssetsAdded);
+		void OnAssetsRemoved(TConstArrayView<FAssetData> InAssetsRemoved);
+		void OnAssetsUpdated(TConstArrayView<FAssetData> InAssetsUpdated);
+		void OnAssetsUpdatedOnDisk(TConstArrayView<FAssetData> InAssetsUpdated);
+
+		void OnAssetRenamed(const FAssetData& InAsset, const FString& InOldObjectPath);
+
+		void OnPathsAdded(TConstArrayView<FStringView> InPathsAdded);
+		void OnPathsRemoved(TConstArrayView<FStringView> InPathsRemoved);
+
+		IEditorDataStorageProvider& Database;
+		DataStorage::TableHandle PathsTable = DataStorage::InvalidTableHandle;
+		DataStorage::TableHandle AssetsDataTable = DataStorage::InvalidTableHandle;
+
+		DataStorage::QueryHandle UpdateAssetsInPathQuery = DataStorage::InvalidQueryHandle;
+		DataStorage::QueryHandle UpdateAssetsInPath = DataStorage::InvalidQueryHandle;
+		DataStorage::QueryHandle ResolveMissingAssetInPathQuery = DataStorage::InvalidQueryHandle;
+		DataStorage::QueryHandle UpdateParentToChildrenAssetPathQuery = DataStorage::InvalidQueryHandle;
+		DataStorage::QueryHandle ResolveMissingParentPathQuery = DataStorage::InvalidQueryHandle;
+		DataStorage::QueryHandle RemoveUpdatedPathTagQuery = DataStorage::InvalidQueryHandle;
+		DataStorage::QueryHandle RemoveUpdatedAssetDataTagQuery = DataStorage::InvalidQueryHandle;
+	};
+} // namespace UE::Editor::AssetData::Private
diff --git a/Engine/Plugins/Experimental/EditorDataStorageFeatures/Source/TedsAssetData/Private/TedsAssetDataModule.cpp b/Engine/Plugins/Experimental/EditorDataStorageFeatures/Source/TedsAssetData/Private/TedsAssetDataModule.cpp
--- a/Engine/Plugins/Experimental/EditorDataStorageFeatures/Source/TedsAssetData/Private/TedsAssetDataModule.cpp
+++ b/Engine/Plugins/Experimental/EditorDataStorageFeatures/Source/TedsAssetData/Private/TedsAssetDataModule.cpp
@@ -108,7 +108,7 @@ void FTedsAssetDataModule::ProcessDependentEvents()
 
 void FTedsAssetDataModule::InitAssetRegistryStorage()
 {
-	ITypedElementDataStorageInterface& MutableDataStorage = *UTypedElementRegistry::GetInstance()->GetMutableDataStorage();
+	IEditorDataStorageProvider& MutableDataStorage = *UTypedElementRegistry::GetInstance()->GetMutableDataStorage();
 
 	AssetDataCBDataSource = MakeUnique<Private::FTedsAssetDataCBDataSource>(MutableDataStorage);
 	AssetRegistryStorage = MakeUnique<Private::FTedsAssetData>(MutableDataStorage);
diff --git a/Engine/Plugins/Experimental/EditorDataStorageFeatures/Source/TedsAssetData/Private/Tests/TedsAssetDataTest.cpp b/Engine/Plugins/Experimental/EditorDataStorageFeatures/Source/TedsAssetData/Private/Tests/TedsAssetDataTest.cpp
--- a/Engine/Plugins/Experimental/EditorDataStorageFeatures/Source/TedsAssetData/Private/Tests/TedsAssetDataTest.cpp
+++ b/Engine/Plugins/Experimental/EditorDataStorageFeatures/Source/TedsAssetData/Private/Tests/TedsAssetDataTest.cpp
@@ -65,7 +65,7 @@ bool FTedsAssetDataTest::RunTest(const FString& Parameters)
 
 
 	// Do a sanity check that the data from the asset registry exist in Teds
-	const ITypedElementDataStorageInterface* Database = UTypedElementRegistry::GetInstance()->GetDataStorage();
+	const IEditorDataStorageProvider* Database = UTypedElementRegistry::GetInstance()->GetDataStorage();
 	
 	for (const TPair<FName, int32>& PathAndAssetCount : AssetRegistryPathsAndAssetCount)
 	{
diff --git a/Engine/Plugins/Experimental/EditorDataStorageFeatures/Source/TedsAssetData/Private/Widgets/AssetDataDiskSizeWidget.cpp b/Engine/Plugins/Experimental/EditorDataStorageFeatures/Source/TedsAssetData/Private/Widgets/AssetDataDiskSizeWidget.cpp
--- a/Engine/Plugins/Experimental/EditorDataStorageFeatures/Source/TedsAssetData/Private/Widgets/AssetDataDiskSizeWidget.cpp
+++ b/Engine/Plugins/Experimental/EditorDataStorageFeatures/Source/TedsAssetData/Private/Widgets/AssetDataDiskSizeWidget.cpp
@@ -8,8 +8,8 @@
 #include "UObject/ObjectMacros.h"
 #include "Widgets/Text/STextBlock.h"
 
-void UDiskSizeWidgetFactory::RegisterWidgetConstructors(ITypedElementDataStorageInterface& DataStorage,
-                                                        ITypedElementDataStorageUiInterface& DataStorageUi) const
+void UDiskSizeWidgetFactory::RegisterWidgetConstructors(IEditorDataStorageProvider& DataStorage,
+                                                        IEditorDataStorageUiProvider& DataStorageUi) const
 {
 	using namespace UE::Editor::DataStorage::Queries;
 	DataStorageUi.RegisterWidgetFactory<FDiskSizeWidgetConstructor>(TEXT("General.Cell"), TColumn<FDiskSizeColumn>());
@@ -20,8 +20,8 @@ FDiskSizeWidgetConstructor::FDiskSizeWidgetConstructor()
 {
 }
 
-TSharedPtr<SWidget> FDiskSizeWidgetConstructor::CreateWidget(ITypedElementDataStorageInterface* DataStorage,
-	ITypedElementDataStorageUiInterface* DataStorageUi, UE::Editor::DataStorage::RowHandle TargetRow, UE::Editor::DataStorage::RowHandle WidgetRow,
+TSharedPtr<SWidget> FDiskSizeWidgetConstructor::CreateWidget(IEditorDataStorageProvider* DataStorage,
+	IEditorDataStorageUiProvider* DataStorageUi, UE::Editor::DataStorage::RowHandle TargetRow, UE::Editor::DataStorage::RowHandle WidgetRow,
 	const UE::Editor::DataStorage::FMetaDataView& Arguments)
 {
 	UE::Editor::DataStorage::FAttributeBinder Binder(TargetRow, DataStorage);
diff --git a/Engine/Plugins/Experimental/EditorDataStorageFeatures/Source/TedsAssetData/Private/Widgets/AssetDataDiskSizeWidget.h b/Engine/Plugins/Experimental/EditorDataStorageFeatures/Source/TedsAssetData/Private/Widgets/AssetDataDiskSizeWidget.h
--- a/Engine/Plugins/Experimental/EditorDataStorageFeatures/Source/TedsAssetData/Private/Widgets/AssetDataDiskSizeWidget.h
+++ b/Engine/Plugins/Experimental/EditorDataStorageFeatures/Source/TedsAssetData/Private/Widgets/AssetDataDiskSizeWidget.h
@@ -7,19 +7,19 @@
 
 #include "AssetDataDiskSizeWidget.generated.h"
 
-class ITypedElementDataStorageInterface;
+class IEditorDataStorageProvider;
 class UScriptStruct;
 
 UCLASS()
-class UDiskSizeWidgetFactory : public UTypedElementDataStorageFactory
+class UDiskSizeWidgetFactory : public UEditorDataStorageFactory
 {
 	GENERATED_BODY()
 
 public:
 	~UDiskSizeWidgetFactory() override = default;
 
-	TEDSASSETDATA_API void RegisterWidgetConstructors(ITypedElementDataStorageInterface& DataStorage,
-		ITypedElementDataStorageUiInterface& DataStorageUi) const override;
+	TEDSASSETDATA_API void RegisterWidgetConstructors(IEditorDataStorageProvider& DataStorage,
+		IEditorDataStorageUiProvider& DataStorageUi) const override;
 };
 
 // Widget to show disk size in bytes
@@ -33,8 +33,8 @@ public:
 	~FDiskSizeWidgetConstructor() override = default;
 
 	TEDSASSETDATA_API virtual TSharedPtr<SWidget> CreateWidget(
-		ITypedElementDataStorageInterface* DataStorage,
-		ITypedElementDataStorageUiInterface* DataStorageUi,
+		IEditorDataStorageProvider* DataStorage,
+		IEditorDataStorageUiProvider* DataStorageUi,
 		UE::Editor::DataStorage::RowHandle TargetRow,
 		UE::Editor::DataStorage::RowHandle WidgetRow,
 		const UE::Editor::DataStorage::FMetaDataView& Arguments) override;
diff --git a/Engine/Plugins/Experimental/EditorDataStorageFeatures/Source/TedsAssetData/Private/Widgets/AssetDataLabelWidget.cpp b/Engine/Plugins/Experimental/EditorDataStorageFeatures/Source/TedsAssetData/Private/Widgets/AssetDataLabelWidget.cpp
--- a/Engine/Plugins/Experimental/EditorDataStorageFeatures/Source/TedsAssetData/Private/Widgets/AssetDataLabelWidget.cpp
+++ b/Engine/Plugins/Experimental/EditorDataStorageFeatures/Source/TedsAssetData/Private/Widgets/AssetDataLabelWidget.cpp
@@ -15,8 +15,8 @@
 
 #define LOCTEXT_NAMESPACE "FAssetDataLabelWidgetConstructor"
 
-void UAssetDataLabelWidgetFactory::RegisterWidgetConstructors(ITypedElementDataStorageInterface& DataStorage,
-                                                              ITypedElementDataStorageUiInterface& DataStorageUi) const
+void UAssetDataLabelWidgetFactory::RegisterWidgetConstructors(IEditorDataStorageProvider& DataStorage,
+                                                              IEditorDataStorageUiProvider& DataStorageUi) const
 {
 	using namespace UE::Editor::DataStorage::Queries;
 	DataStorageUi.RegisterWidgetFactory<FAssetDataLabelWidgetConstructor>(
@@ -34,8 +34,8 @@ FAssetDataLabelWidgetConstructor::FAssetDataLabelWidgetConstructor(const UScript
 {
 }
 
-TSharedPtr<SWidget> FAssetDataLabelWidgetConstructor::CreateWidget(ITypedElementDataStorageInterface* DataStorage,
-	ITypedElementDataStorageUiInterface* DataStorageUi, UE::Editor::DataStorage::RowHandle TargetRow, UE::Editor::DataStorage::RowHandle WidgetRow,
+TSharedPtr<SWidget> FAssetDataLabelWidgetConstructor::CreateWidget(IEditorDataStorageProvider* DataStorage,
+	IEditorDataStorageUiProvider* DataStorageUi, UE::Editor::DataStorage::RowHandle TargetRow, UE::Editor::DataStorage::RowHandle WidgetRow,
 	const UE::Editor::DataStorage::FMetaDataView& Arguments)
 {
 	if (DataStorage->IsRowAvailable(TargetRow))
@@ -97,7 +97,7 @@ TSharedPtr<SWidget> FAssetDataLabelWidgetConstructor::CreateWidget(ITypedElement
 }
 
 FText FAssetDataLabelWidgetConstructor::ConstructToolTip(
-	ITypedElementDataStorageInterface* DataStorage, UE::Editor::DataStorage::RowHandle DataRow)
+	IEditorDataStorageProvider* DataStorage, UE::Editor::DataStorage::RowHandle DataRow)
 {
 	TStringBuilder<1024> ToolTip;
 	if (FNameColumn* ItemName = DataStorage->GetColumn<FNameColumn>(DataRow))
diff --git a/Engine/Plugins/Experimental/EditorDataStorageFeatures/Source/TedsAssetData/Public/Widgets/AssetDataLabelWidget.h b/Engine/Plugins/Experimental/EditorDataStorageFeatures/Source/TedsAssetData/Public/Widgets/AssetDataLabelWidget.h
--- a/Engine/Plugins/Experimental/EditorDataStorageFeatures/Source/TedsAssetData/Public/Widgets/AssetDataLabelWidget.h
+++ b/Engine/Plugins/Experimental/EditorDataStorageFeatures/Source/TedsAssetData/Public/Widgets/AssetDataLabelWidget.h
@@ -9,19 +9,19 @@
 
 #include "AssetDataLabelWidget.generated.h"
 
-class ITypedElementDataStorageInterface;
+class IEditorDataStorageProvider;
 class UScriptStruct;
 
 UCLASS()
-class UAssetDataLabelWidgetFactory : public UTypedElementDataStorageFactory
+class UAssetDataLabelWidgetFactory : public UEditorDataStorageFactory
 {
 	GENERATED_BODY()
 
 public:
 	~UAssetDataLabelWidgetFactory() override = default;
 
-	TEDSASSETDATA_API void RegisterWidgetConstructors(ITypedElementDataStorageInterface& DataStorage,
-		ITypedElementDataStorageUiInterface& DataStorageUi) const override;
+	TEDSASSETDATA_API void RegisterWidgetConstructors(IEditorDataStorageProvider& DataStorage,
+		IEditorDataStorageUiProvider& DataStorageUi) const override;
 };
 
 // Label widget for assets in TEDS
@@ -36,13 +36,13 @@ public:
 	~FAssetDataLabelWidgetConstructor() override = default;
 
 	TEDSASSETDATA_API virtual TSharedPtr<SWidget> CreateWidget(
-		ITypedElementDataStorageInterface* DataStorage,
-		ITypedElementDataStorageUiInterface* DataStorageUi,
+		IEditorDataStorageProvider* DataStorage,
+		IEditorDataStorageUiProvider* DataStorageUi,
 		UE::Editor::DataStorage::RowHandle TargetRow,
 		UE::Editor::DataStorage::RowHandle WidgetRow, 
 		const UE::Editor::DataStorage::FMetaDataView& Arguments) override;
 
 protected:
 	
-	static FText ConstructToolTip(ITypedElementDataStorageInterface* DataStorage, UE::Editor::DataStorage::RowHandle DataRow);
+	static FText ConstructToolTip(IEditorDataStorageProvider* DataStorage, UE::Editor::DataStorage::RowHandle DataRow);
 };
\ No newline at end of file
diff --git a/Engine/Plugins/Experimental/EditorDataStorageFeatures/Source/TedsContentBrowser/Private/TedsContentBrowserModule.h b/Engine/Plugins/Experimental/EditorDataStorageFeatures/Source/TedsContentBrowser/Private/TedsContentBrowserModule.h
--- a/Engine/Plugins/Experimental/EditorDataStorageFeatures/Source/TedsContentBrowser/Private/TedsContentBrowserModule.h
+++ b/Engine/Plugins/Experimental/EditorDataStorageFeatures/Source/TedsContentBrowser/Private/TedsContentBrowserModule.h
@@ -7,7 +7,7 @@
 #include "Modules/ModuleInterface.h"
 #include "Templates/SharedPointer.h"
 
-class ITypedElementDataStorageInterface;
+class IEditorDataStorageProvider;
 
 namespace UE::Editor::DataStorage
 {
@@ -57,7 +57,7 @@ namespace UE::Editor::ContentBrowser
 	private:
 
 		// Ptr to the data storage interface
-		ITypedElementDataStorageInterface* DataStorage;
+		IEditorDataStorageProvider* DataStorage;
 		
 		// The actual table viewer widget
 		TSharedPtr<DataStorage::STedsTableViewer> TableViewer;
diff --git a/Engine/Plugins/Experimental/EditorDataStorageFeatures/Source/TedsDebugger/Private/QueryEditor/TedsQueryEditorModel.cpp b/Engine/Plugins/Experimental/EditorDataStorageFeatures/Source/TedsDebugger/Private/QueryEditor/TedsQueryEditorModel.cpp
--- a/Engine/Plugins/Experimental/EditorDataStorageFeatures/Source/TedsDebugger/Private/QueryEditor/TedsQueryEditorModel.cpp
+++ b/Engine/Plugins/Experimental/EditorDataStorageFeatures/Source/TedsDebugger/Private/QueryEditor/TedsQueryEditorModel.cpp
@@ -23,7 +23,7 @@ namespace UE::Editor::DataStorage::Debug::QueryEditor
 		Id = INDEX_NONE;
 	}
 
-	FTedsQueryEditorModel::FTedsQueryEditorModel(ITypedElementDataStorageInterface& InTypedElementDataStorage)
+	FTedsQueryEditorModel::FTedsQueryEditorModel(IEditorDataStorageProvider& InTypedElementDataStorage)
 		: TypedElementDataStorage(InTypedElementDataStorage)
 	{
 	}
@@ -91,12 +91,12 @@ namespace UE::Editor::DataStorage::Debug::QueryEditor
 		ModelChangedDelegate.Broadcast();
 	}
 
-	ITypedElementDataStorageInterface& FTedsQueryEditorModel::GetTedsInterface()
+	IEditorDataStorageProvider& FTedsQueryEditorModel::GetTedsInterface()
 	{
 		return TypedElementDataStorage;
 	}
 	
-	const ITypedElementDataStorageInterface& FTedsQueryEditorModel::GetTedsInterface() const
+	const IEditorDataStorageProvider& FTedsQueryEditorModel::GetTedsInterface() const
 	{
 		return TypedElementDataStorage;
 	}
@@ -141,21 +141,21 @@ namespace UE::Editor::DataStorage::Debug::QueryEditor
 				}
 				else
 				{
-					ITypedElementDataStorageInterface::FQueryDescription::EOperatorType OperatorType = [](const FConditionEntryInternal& Entry)
+					IEditorDataStorageProvider::FQueryDescription::EOperatorType OperatorType = [](const FConditionEntryInternal& Entry)
 					{
 						switch(Entry.OperatorType)
 						{
 						case EOperatorType::All:
-							return ITypedElementDataStorageInterface::FQueryDescription::EOperatorType::SimpleAll;
+							return IEditorDataStorageProvider::FQueryDescription::EOperatorType::SimpleAll;
 						case EOperatorType::Any:
-							return ITypedElementDataStorageInterface::FQueryDescription::EOperatorType::SimpleAny;
+							return IEditorDataStorageProvider::FQueryDescription::EOperatorType::SimpleAny;
 						case EOperatorType::None:
-							return ITypedElementDataStorageInterface::FQueryDescription::EOperatorType::SimpleNone;
+							return IEditorDataStorageProvider::FQueryDescription::EOperatorType::SimpleNone;
 						case EOperatorType::Invalid:
 						case EOperatorType::Unset:
 						default:
 							check(false);
-							return ITypedElementDataStorageInterface::FQueryDescription::EOperatorType::Max;
+							return IEditorDataStorageProvider::FQueryDescription::EOperatorType::Max;
 						}
 					}(Entry);
 				
@@ -165,7 +165,7 @@ namespace UE::Editor::DataStorage::Debug::QueryEditor
 			}
 		}
 
-		Description.Action = ITypedElementDataStorageInterface::FQueryDescription::EActionType::Select;
+		Description.Action = IEditorDataStorageProvider::FQueryDescription::EActionType::Select;
 
 		return Description;
 	}
@@ -180,14 +180,14 @@ namespace UE::Editor::DataStorage::Debug::QueryEditor
 		for (int32 Index = 0; Index < SelectionTypeCount; ++Index)
 		{
 			Description.ConditionOperators.AddZeroed_GetRef().Type = Description.SelectionTypes[Index];
-			Description.ConditionTypes.Add(ITypedElementDataStorageInterface::FQueryDescription::EOperatorType::SimpleAll);
+			Description.ConditionTypes.Add(IEditorDataStorageProvider::FQueryDescription::EOperatorType::SimpleAll);
 		}
 
 		Description.SelectionTypes.Empty();
 		Description.SelectionMetaData.Empty();
 		Description.SelectionAccessTypes.Empty();
 
-		Description.Action = ITypedElementDataStorageInterface::FQueryDescription::EActionType::Count;
+		Description.Action = IEditorDataStorageProvider::FQueryDescription::EActionType::Count;
 
 		return Description;
 	}
diff --git a/Engine/Plugins/Experimental/EditorDataStorageFeatures/Source/TedsDebugger/Private/QueryEditor/TedsQueryEditorModel.h b/Engine/Plugins/Experimental/EditorDataStorageFeatures/Source/TedsDebugger/Private/QueryEditor/TedsQueryEditorModel.h
--- a/Engine/Plugins/Experimental/EditorDataStorageFeatures/Source/TedsDebugger/Private/QueryEditor/TedsQueryEditorModel.h
+++ b/Engine/Plugins/Experimental/EditorDataStorageFeatures/Source/TedsDebugger/Private/QueryEditor/TedsQueryEditorModel.h
@@ -4,7 +4,7 @@
 #include "Elements/Framework/TypedElementMetaData.h"
 #include "UObject/Class.h"
 
-class ITypedElementDataStorageInterface;
+class IEditorDataStorageProvider;
 
 namespace UE::Editor::DataStorage::Debug::QueryEditor
 {
@@ -48,11 +48,11 @@ namespace UE::Editor::DataStorage::Debug::QueryEditor
 	class FTedsQueryEditorModel
 	{
 	public:
-		explicit FTedsQueryEditorModel(ITypedElementDataStorageInterface& InTypedElementDataStorage);
+		explicit FTedsQueryEditorModel(IEditorDataStorageProvider& InTypedElementDataStorage);
 
 		void Reset();
-		ITypedElementDataStorageInterface& GetTedsInterface();
-		const ITypedElementDataStorageInterface& GetTedsInterface() const;
+		IEditorDataStorageProvider& GetTedsInterface();
+		const IEditorDataStorageProvider& GetTedsInterface() const;
 		
 		FQueryDescription GenerateQueryDescription();
 		// Special function to generate a description that puts the Select elements as All Conditions
@@ -105,7 +105,7 @@ namespace UE::Editor::DataStorage::Debug::QueryEditor
 		FConditionEntryInternal* FindEntryByHandle(FConditionEntryHandle Handle);
 		const FConditionEntryInternal* FindEntryByHandle(FConditionEntryHandle Handle) const;
 
-		ITypedElementDataStorageInterface& TypedElementDataStorage;
+		IEditorDataStorageProvider& TypedElementDataStorage;
 
 		int32 IdGenerator = 0;
 
diff --git a/Engine/Plugins/Experimental/EditorDataStorageFeatures/Source/TedsDebugger/Private/STedsDebugger.cpp b/Engine/Plugins/Experimental/EditorDataStorageFeatures/Source/TedsDebugger/Private/STedsDebugger.cpp
--- a/Engine/Plugins/Experimental/EditorDataStorageFeatures/Source/TedsDebugger/Private/STedsDebugger.cpp
+++ b/Engine/Plugins/Experimental/EditorDataStorageFeatures/Source/TedsDebugger/Private/STedsDebugger.cpp
@@ -132,7 +132,7 @@ TSharedRef<SDockTab> STedsDebugger::SpawnQueryEditorTab(const FSpawnTabArgs& Arg
 
 		if (Registry && Registry->AreDataStorageInterfacesSet())
 		{
-			ITypedElementDataStorageInterface* DataStorageInterface = Registry->GetMutableDataStorage();
+			IEditorDataStorageProvider* DataStorageInterface = Registry->GetMutableDataStorage();
 			QueryEditorModel = MakeUnique<QueryEditor::FTedsQueryEditorModel>(*DataStorageInterface);
 		}
 	}
diff --git a/Engine/Plugins/Experimental/EditorDataStorageFeatures/Source/TedsDebugger/Private/Widgets/QueryEditor/TedsQueryEditorResultsView.cpp b/Engine/Plugins/Experimental/EditorDataStorageFeatures/Source/TedsDebugger/Private/Widgets/QueryEditor/TedsQueryEditorResultsView.cpp
--- a/Engine/Plugins/Experimental/EditorDataStorageFeatures/Source/TedsDebugger/Private/Widgets/QueryEditor/TedsQueryEditorResultsView.cpp
+++ b/Engine/Plugins/Experimental/EditorDataStorageFeatures/Source/TedsDebugger/Private/Widgets/QueryEditor/TedsQueryEditorResultsView.cpp
@@ -102,7 +102,7 @@ namespace UE::Editor::DataStorage::Debug::QueryEditor
 	{
 		using namespace UE::Editor::DataStorage::Queries;
 
-		ITypedElementDataStorageInterface& TedsInterface = Model->GetTedsInterface();
+		IEditorDataStorageProvider& TedsInterface = Model->GetTedsInterface();
 
 		if(bModelDirty)
 		{
@@ -153,7 +153,7 @@ namespace UE::Editor::DataStorage::Debug::QueryEditor
 			NewTableViewerRows_Set.Reserve(TableViewerRows_Set.Num());
 
 			FQueryResult QueryResult = Model->GetTedsInterface().RunQuery(TableViewerQueryHandle,
-				CreateDirectQueryCallbackBinding([&NewTableViewerRows_Set](const ITypedElementDataStorageInterface::IDirectQueryContext& Context, const RowHandle*)
+				CreateDirectQueryCallbackBinding([&NewTableViewerRows_Set](const IEditorDataStorageProvider::IDirectQueryContext& Context, const RowHandle*)
 				{
 					NewTableViewerRows_Set.Append(Context.GetRowHandles());
 				}));
@@ -187,7 +187,7 @@ namespace UE::Editor::DataStorage::Debug::QueryEditor
 	
 		UTypedElementRegistry* Registry = UTypedElementRegistry::GetInstance();
 		checkf(Registry, TEXT("SResultsView created before UTypedElementRegistry is available."));
-		ITypedElementDataStorageUiInterface* StorageUi = Registry->GetMutableDataStorageUi();
+		IEditorDataStorageUiProvider* StorageUi = Registry->GetMutableDataStorageUi();
 		checkf(StorageUi, TEXT("SResultsView created before data storage interfaces were initialized."))
 
 		StorageUi->CreateWidgetConstructors(TEXT("General.Cell.RowHandle"), FMetaDataView(), AssignWidgetToColumn);
diff --git a/Engine/Plugins/Experimental/EditorDataStorageFeatures/Source/TedsDebugger/Private/Widgets/RowReferenceWidget.cpp b/Engine/Plugins/Experimental/EditorDataStorageFeatures/Source/TedsDebugger/Private/Widgets/RowReferenceWidget.cpp
--- a/Engine/Plugins/Experimental/EditorDataStorageFeatures/Source/TedsDebugger/Private/Widgets/RowReferenceWidget.cpp
+++ b/Engine/Plugins/Experimental/EditorDataStorageFeatures/Source/TedsDebugger/Private/Widgets/RowReferenceWidget.cpp
@@ -19,7 +19,7 @@
 
 namespace UE::Editor::DataStorage::Debug::Private
 {
-	void OnNavigateHyperlink(const ITypedElementDataStorageInterface* DataStorage,  RowHandle TargetRowHandle, RowHandle UiRowHandle)
+	void OnNavigateHyperlink(const IEditorDataStorageProvider* DataStorage,  RowHandle TargetRowHandle, RowHandle UiRowHandle)
 	{
 		const FTedsOutlinerColumn* TedsOutlinerColumn = DataStorage->GetColumn<FTedsOutlinerColumn>(UiRowHandle);
 
@@ -66,7 +66,7 @@ namespace UE::Editor::DataStorage::Debug::Private
 		SBox* WidgetInstance = static_cast<SBox*>(Widget.Get());
 		WidgetInstance->SetContent(SNullWidget::NullWidget);
 
-		const ITypedElementDataStorageInterface* DataStorage = UTypedElementRegistry::GetInstance()->GetDataStorage();
+		const IEditorDataStorageProvider* DataStorage = UTypedElementRegistry::GetInstance()->GetDataStorage();
 
 		// We only navigate to row references that have a label column
 		if (const FTypedElementLabelColumn* LabelColumn = DataStorage->GetColumn<FTypedElementLabelColumn>(TargetRow))
@@ -95,7 +95,7 @@ URowReferenceWidgetFactory::~URowReferenceWidgetFactory()
 {
 }
 
-void URowReferenceWidgetFactory::RegisterWidgetConstructors(ITypedElementDataStorageInterface& DataStorage, ITypedElementDataStorageUiInterface& DataStorageUi) const
+void URowReferenceWidgetFactory::RegisterWidgetConstructors(IEditorDataStorageProvider& DataStorage, IEditorDataStorageUiProvider& DataStorageUi) const
 {
 	using namespace UE::Editor::DataStorage::Queries;
 
@@ -104,7 +104,7 @@ void URowReferenceWidgetFactory::RegisterWidgetConstructors(ITypedElementDataSto
 		TColumn<FTypedElementRowReferenceColumn>());
 }
 
-void URowReferenceWidgetFactory::RegisterQueries(ITypedElementDataStorageInterface& DataStorage)
+void URowReferenceWidgetFactory::RegisterQueries(IEditorDataStorageProvider& DataStorage)
 {
 	using namespace UE::Editor::DataStorage::Queries;
 	
@@ -149,7 +149,7 @@ TSharedPtr<SWidget> FRowReferenceWidgetConstructor::CreateWidget(const UE::Edito
 			.VAlign(VAlign_Center);
 }
 
-bool FRowReferenceWidgetConstructor::FinalizeWidget(ITypedElementDataStorageInterface* DataStorage, ITypedElementDataStorageUiInterface* DataStorageUi, UE::Editor::DataStorage::RowHandle Row, const TSharedPtr<SWidget>& Widget)
+bool FRowReferenceWidgetConstructor::FinalizeWidget(IEditorDataStorageProvider* DataStorage, IEditorDataStorageUiProvider* DataStorageUi, UE::Editor::DataStorage::RowHandle Row, const TSharedPtr<SWidget>& Widget)
 {
 	using namespace UE::Editor::DataStorage;
 	checkf(Widget, TEXT("Referenced widget is not valid. A constructed widget may not have been cleaned up. This can "
diff --git a/Engine/Plugins/Experimental/EditorDataStorageFeatures/Source/TedsDebugger/Private/Widgets/RowReferenceWidget.h b/Engine/Plugins/Experimental/EditorDataStorageFeatures/Source/TedsDebugger/Private/Widgets/RowReferenceWidget.h
--- a/Engine/Plugins/Experimental/EditorDataStorageFeatures/Source/TedsDebugger/Private/Widgets/RowReferenceWidget.h
+++ b/Engine/Plugins/Experimental/EditorDataStorageFeatures/Source/TedsDebugger/Private/Widgets/RowReferenceWidget.h
@@ -7,24 +7,24 @@
 
 #include "RowReferenceWidget.generated.h"
 
-class ITypedElementDataStorageInterface;
+class IEditorDataStorageProvider;
 class SWidget;
 
 /*
  * Widget for the TEDS Debugger that visualizes a reference to another row
  */
 UCLASS()
-class URowReferenceWidgetFactory : public UTypedElementDataStorageFactory
+class URowReferenceWidgetFactory : public UEditorDataStorageFactory
 {
 	GENERATED_BODY()
 
 public:
 	~URowReferenceWidgetFactory() override;
 
-	void RegisterWidgetConstructors(ITypedElementDataStorageInterface& DataStorage,
-		ITypedElementDataStorageUiInterface& DataStorageUi) const override;
+	void RegisterWidgetConstructors(IEditorDataStorageProvider& DataStorage,
+		IEditorDataStorageUiProvider& DataStorageUi) const override;
 	
-	void RegisterQueries(ITypedElementDataStorageInterface& DataStorage) override;
+	void RegisterQueries(IEditorDataStorageProvider& DataStorage) override;
 };
 
 USTRUCT()
@@ -38,6 +38,6 @@ public:
 
 protected:
 	TSharedPtr<SWidget> CreateWidget(const UE::Editor::DataStorage::FMetaDataView& Arguments) override;
-	bool FinalizeWidget(ITypedElementDataStorageInterface* DataStorage, ITypedElementDataStorageUiInterface* DataStorageUi,
+	bool FinalizeWidget(IEditorDataStorageProvider* DataStorage, IEditorDataStorageUiProvider* DataStorageUi,
 		UE::Editor::DataStorage::RowHandle Row, const TSharedPtr<SWidget>& Widget) override;
 };
\ No newline at end of file
diff --git a/Engine/Plugins/Experimental/EditorDataStorageFeatures/Source/TedsDebugger/Private/Widgets/SlateVisualizationWidget.cpp b/Engine/Plugins/Experimental/EditorDataStorageFeatures/Source/TedsDebugger/Private/Widgets/SlateVisualizationWidget.cpp
--- a/Engine/Plugins/Experimental/EditorDataStorageFeatures/Source/TedsDebugger/Private/Widgets/SlateVisualizationWidget.cpp
+++ b/Engine/Plugins/Experimental/EditorDataStorageFeatures/Source/TedsDebugger/Private/Widgets/SlateVisualizationWidget.cpp
@@ -9,7 +9,7 @@
 #include "Widgets/SBoxPanel.h"
 
 
-void USlateVisualizationWidgetFactory::RegisterWidgetConstructors(ITypedElementDataStorageInterface& DataStorage, ITypedElementDataStorageUiInterface& DataStorageUi) const
+void USlateVisualizationWidgetFactory::RegisterWidgetConstructors(IEditorDataStorageProvider& DataStorage, IEditorDataStorageUiProvider& DataStorageUi) const
 {
 	using namespace UE::Editor::DataStorage::Queries;
 
@@ -28,7 +28,7 @@ TSharedPtr<SWidget> FSlateVisualizationWidgetConstructor::CreateWidget(const UE:
 	return SNew(SHorizontalBox);
 }
 
-bool FSlateVisualizationWidgetConstructor::FinalizeWidget(ITypedElementDataStorageInterface* DataStorage, ITypedElementDataStorageUiInterface* DataStorageUi, UE::Editor::DataStorage::RowHandle Row, const TSharedPtr<SWidget>& Widget)
+bool FSlateVisualizationWidgetConstructor::FinalizeWidget(IEditorDataStorageProvider* DataStorage, IEditorDataStorageUiProvider* DataStorageUi, UE::Editor::DataStorage::RowHandle Row, const TSharedPtr<SWidget>& Widget)
 {
 	checkf(Widget, TEXT("Referenced widget is not valid. A constructed widget may not have been cleaned up. This can "
 	"also happen if this processor is running in the same phase as the processors responsible for cleaning up old "
diff --git a/Engine/Plugins/Experimental/EditorDataStorageFeatures/Source/TedsDebugger/Private/Widgets/SlateVisualizationWidget.h b/Engine/Plugins/Experimental/EditorDataStorageFeatures/Source/TedsDebugger/Private/Widgets/SlateVisualizationWidget.h
--- a/Engine/Plugins/Experimental/EditorDataStorageFeatures/Source/TedsDebugger/Private/Widgets/SlateVisualizationWidget.h
+++ b/Engine/Plugins/Experimental/EditorDataStorageFeatures/Source/TedsDebugger/Private/Widgets/SlateVisualizationWidget.h
@@ -7,22 +7,22 @@
 
 #include "SlateVisualizationWidget.generated.h"
 
-class ITypedElementDataStorageInterface;
+class IEditorDataStorageProvider;
 class SWidget;
 
 /*
  * Widget for the TEDS Debugger that shows a slate widget reference
  */
 UCLASS()
-class USlateVisualizationWidgetFactory : public UTypedElementDataStorageFactory
+class USlateVisualizationWidgetFactory : public UEditorDataStorageFactory
 {
 	GENERATED_BODY()
 
 public:
 	~USlateVisualizationWidgetFactory() override = default;
 
-	void RegisterWidgetConstructors(ITypedElementDataStorageInterface& DataStorage,
-		ITypedElementDataStorageUiInterface& DataStorageUi) const override;
+	void RegisterWidgetConstructors(IEditorDataStorageProvider& DataStorage,
+		IEditorDataStorageUiProvider& DataStorageUi) const override;
 };
 
 USTRUCT()
@@ -36,6 +36,6 @@ public:
 
 protected:
 	TSharedPtr<SWidget> CreateWidget(const UE::Editor::DataStorage::FMetaDataView& Arguments) override;
-	bool FinalizeWidget(ITypedElementDataStorageInterface* DataStorage, ITypedElementDataStorageUiInterface* DataStorageUi,
+	bool FinalizeWidget(IEditorDataStorageProvider* DataStorage, IEditorDataStorageUiProvider* DataStorageUi,
 		UE::Editor::DataStorage::RowHandle Row, const TSharedPtr<SWidget>& Widget) override;
 };
\ No newline at end of file
diff --git a/Engine/Plugins/Experimental/EditorDataStorageFeatures/Source/TedsOutliner/Private/Compatibility/SceneOutlinerRowHandleColumn.cpp b/Engine/Plugins/Experimental/EditorDataStorageFeatures/Source/TedsOutliner/Private/Compatibility/SceneOutlinerRowHandleColumn.cpp
--- a/Engine/Plugins/Experimental/EditorDataStorageFeatures/Source/TedsOutliner/Private/Compatibility/SceneOutlinerRowHandleColumn.cpp
+++ b/Engine/Plugins/Experimental/EditorDataStorageFeatures/Source/TedsOutliner/Private/Compatibility/SceneOutlinerRowHandleColumn.cpp
@@ -21,7 +21,7 @@ FSceneOutlinerRowHandleColumn::FSceneOutlinerRowHandleColumn(ISceneOutliner& Sce
 	
 	UTypedElementRegistry* Registry = UTypedElementRegistry::GetInstance();
 	checkf(Registry, TEXT("FSceneOutlinerRowHandleColumn created before UTypedElementRegistry is available."));
-	ITypedElementDataStorageUiInterface* StorageUi = Registry->GetMutableDataStorageUi();
+	IEditorDataStorageUiProvider* StorageUi = Registry->GetMutableDataStorageUi();
 	checkf(StorageUi, TEXT("FSceneOutlinerRowHandleColumn created before data storage interfaces were initialized."))
 
 	StorageUi->CreateWidgetConstructors(TEXT("General.Cell.RowHandle"), UE::Editor::DataStorage::FMetaDataView(), AssignWidgetToColumn);
diff --git a/Engine/Plugins/Experimental/EditorDataStorageFeatures/Source/TedsOutliner/Private/Compatibility/SceneOutlinerTedsBridge.cpp b/Engine/Plugins/Experimental/EditorDataStorageFeatures/Source/TedsOutliner/Private/Compatibility/SceneOutlinerTedsBridge.cpp
--- a/Engine/Plugins/Experimental/EditorDataStorageFeatures/Source/TedsOutliner/Private/Compatibility/SceneOutlinerTedsBridge.cpp
+++ b/Engine/Plugins/Experimental/EditorDataStorageFeatures/Source/TedsOutliner/Private/Compatibility/SceneOutlinerTedsBridge.cpp
@@ -39,7 +39,7 @@ FAutoConsoleCommand BindColumnsToSceneOutlinerConsoleCommand(
 		    const FName WidgetPurposes[] = {TEXT("SceneOutliner.Cell"), TEXT("General.Cell")};
 
 			UTypedElementRegistry* Registry = UTypedElementRegistry::GetInstance();
-			if (ITypedElementDataStorageInterface* DataStorage = Registry->GetMutableDataStorage())
+			if (IEditorDataStorageProvider* DataStorage = Registry->GetMutableDataStorage())
 			{
 				static UE::Editor::DataStorage::QueryHandle Queries[] =
 				{
@@ -117,9 +117,9 @@ public:
 	~FSceneOutlinerTedsBridge();
 
 	void Initialize(
-		ITypedElementDataStorageInterface& InStorage,
-		ITypedElementDataStorageUiInterface& InStorageUi,
-		ITypedElementDataStorageCompatibilityInterface& InStorageCompatibility,
+		IEditorDataStorageProvider& InStorage,
+		IEditorDataStorageUiProvider& InStorageUi,
+		IEditorDataStorageCompatibilityProvider& InStorageCompatibility,
 		const TSharedPtr<ISceneOutliner>& InOutliner);
 
 	void AssignQuery(UE::Editor::DataStorage::QueryHandle Query, const TConstArrayView<FName> CellWidgetPurposes);
@@ -131,9 +131,9 @@ private:
 
 	TArray<FName> AddedColumns;
 	TWeakPtr<ISceneOutliner> Outliner;
-	ITypedElementDataStorageInterface* Storage{ nullptr };
-	ITypedElementDataStorageUiInterface* StorageUi{ nullptr };
-	ITypedElementDataStorageCompatibilityInterface* StorageCompatibility{ nullptr };
+	IEditorDataStorageProvider* Storage{ nullptr };
+	IEditorDataStorageUiProvider* StorageUi{ nullptr };
+	IEditorDataStorageCompatibilityProvider* StorageCompatibility{ nullptr };
 	FTreeItemIDDealiaser Dealiaser;
 	TArray<FName> CellWidgetPurposes;
 };
@@ -143,9 +143,9 @@ class FOutlinerColumn : public ISceneOutlinerColumn
 public:
 	FOutlinerColumn(
 		UE::Editor::DataStorage::QueryHandle InQuery,
-		ITypedElementDataStorageInterface& InStorage, 
-		ITypedElementDataStorageUiInterface& InStorageUi, 
-		ITypedElementDataStorageCompatibilityInterface& InStorageCompatibility,
+		IEditorDataStorageProvider& InStorage, 
+		IEditorDataStorageUiProvider& InStorageUi, 
+		IEditorDataStorageCompatibilityProvider& InStorageCompatibility,
 		FName InNameId, 
 		TArray<TWeakObjectPtr<const UScriptStruct>> InColumnTypes,
 		TSharedPtr<FTypedElementWidgetConstructor> InHeaderWidgetConstructor,
@@ -311,9 +311,9 @@ public:
 	// The table viewer implementation that we internally use to create our widgets
 	TUniquePtr<UE::Editor::DataStorage::FTedsTableViewerColumn> TableViewerColumnImpl;
 	
-	ITypedElementDataStorageInterface& Storage;
-	ITypedElementDataStorageUiInterface& StorageUi;
-	ITypedElementDataStorageCompatibilityInterface& StorageCompatibility;
+	IEditorDataStorageProvider& Storage;
+	IEditorDataStorageUiProvider& StorageUi;
+	IEditorDataStorageCompatibilityProvider& StorageCompatibility;
 	UE::Editor::DataStorage::QueryHandle QueryHandle;
 	UE::Editor::DataStorage::FMetaData MetaData;
 	FName NameId;
@@ -329,9 +329,9 @@ public:
 // USceneOutlinerTedsBridgeFactory
 // 
 
-void USceneOutlinerTedsBridgeFactory::RegisterWidgetPurposes(ITypedElementDataStorageUiInterface& DataStorageUi) const
+void USceneOutlinerTedsBridgeFactory::RegisterWidgetPurposes(IEditorDataStorageUiProvider& DataStorageUi) const
 {
-	using PurposeType = ITypedElementDataStorageUiInterface::EPurposeType;
+	using PurposeType = IEditorDataStorageUiProvider::EPurposeType;
 
 	DataStorageUi.RegisterWidgetPurpose(FSceneOutlinerTedsQueryBinder::HeaderWidgetPurpose, PurposeType::UniqueByNameAndColumn,
 		LOCTEXT("HeaderWidgetPurpose", "Widgets for headers in any Scene Outliner for specific columns or column combinations."));
@@ -474,9 +474,9 @@ FSceneOutlinerTedsBridge::~FSceneOutlinerTedsBridge()
 }
 
 void FSceneOutlinerTedsBridge::Initialize(
-	ITypedElementDataStorageInterface& InStorage,
-	ITypedElementDataStorageUiInterface& InStorageUi,
-	ITypedElementDataStorageCompatibilityInterface& InStorageCompatibility,
+	IEditorDataStorageProvider& InStorage,
+	IEditorDataStorageUiProvider& InStorageUi,
+	IEditorDataStorageCompatibilityProvider& InStorageCompatibility,
 	const TSharedPtr<ISceneOutliner>& InOutliner)
 {
 	Storage = &InStorage;
@@ -498,7 +498,7 @@ FTreeItemIDDealiaser FSceneOutlinerTedsBridge::GetDealiaser()
 
 void FSceneOutlinerTedsBridge::AssignQuery(UE::Editor::DataStorage::QueryHandle Query, const TConstArrayView<FName> InCellWidgetPurposes)
 {
-	using MatchApproach = ITypedElementDataStorageUiInterface::EMatchApproach;
+	using MatchApproach = IEditorDataStorageUiProvider::EMatchApproach;
 	constexpr int32 DefaultPriorityIndex = 100;
 	FSceneOutlinerTedsQueryBinder& Binder = FSceneOutlinerTedsQueryBinder::GetInstance();
 	CellWidgetPurposes = InCellWidgetPurposes;
@@ -506,12 +506,12 @@ void FSceneOutlinerTedsBridge::AssignQuery(UE::Editor::DataStorage::QueryHandle
 
 	if (TSharedPtr<ISceneOutliner> OutlinerPinned = Outliner.Pin())
 	{
-		const ITypedElementDataStorageInterface::FQueryDescription& Description = Storage->GetQueryDescription(Query);
+		const IEditorDataStorageProvider::FQueryDescription& Description = Storage->GetQueryDescription(Query);
 		UE::Editor::DataStorage::FQueryMetaDataView MetaDataView(Description);
 
 		ClearColumns(*OutlinerPinned);
 
-		if (Description.Action == ITypedElementDataStorageInterface::FQueryDescription::EActionType::Select)
+		if (Description.Action == IEditorDataStorageProvider::FQueryDescription::EActionType::Select)
 		{
 			int32 SelectionCount = Description.SelectionTypes.Num();
 			AddedColumns.Reset(SelectionCount);
diff --git a/Engine/Plugins/Experimental/EditorDataStorageFeatures/Source/TedsOutliner/Private/TedsOutlinerHierarchy.cpp b/Engine/Plugins/Experimental/EditorDataStorageFeatures/Source/TedsOutliner/Private/TedsOutlinerHierarchy.cpp
--- a/Engine/Plugins/Experimental/EditorDataStorageFeatures/Source/TedsOutliner/Private/TedsOutlinerHierarchy.cpp
+++ b/Engine/Plugins/Experimental/EditorDataStorageFeatures/Source/TedsOutliner/Private/TedsOutlinerHierarchy.cpp
@@ -49,7 +49,7 @@ FSceneOutlinerTreeItemPtr FTedsOutlinerHierarchy::FindOrCreateParentItem(const I
 	const TMap<FSceneOutlinerTreeItemID, FSceneOutlinerTreeItemPtr>& Items, bool bCreate)
 {
 	const FTedsOutlinerTreeItem* TedsTreeItem = Item.CastTo<FTedsOutlinerTreeItem>();
-	const ITypedElementDataStorageInterface* Storage = TedsOutlinerImpl->GetStorage();
+	const IEditorDataStorageProvider* Storage = TedsOutlinerImpl->GetStorage();
 	
 	// If this item is not a TEDS item, we are not handling it
 	if(!TedsTreeItem)
diff --git a/Engine/Plugins/Experimental/EditorDataStorageFeatures/Source/TedsOutliner/Private/TedsOutlinerImpl.cpp b/Engine/Plugins/Experimental/EditorDataStorageFeatures/Source/TedsOutliner/Private/TedsOutlinerImpl.cpp
--- a/Engine/Plugins/Experimental/EditorDataStorageFeatures/Source/TedsOutliner/Private/TedsOutlinerImpl.cpp
+++ b/Engine/Plugins/Experimental/EditorDataStorageFeatures/Source/TedsOutliner/Private/TedsOutlinerImpl.cpp
@@ -203,7 +203,7 @@ TSharedRef<SWidget> FTedsOutlinerImpl::CreateLabelWidgetForItem(DataStorage::Row
 
 		for(const FName& Purpose : Purposes)
 		{
-			StorageUi->CreateWidgetConstructors(Purpose, ITypedElementDataStorageUiInterface::EMatchApproach::LongestMatch, ColumnTypes, {},
+			StorageUi->CreateWidgetConstructors(Purpose, IEditorDataStorageUiProvider::EMatchApproach::LongestMatch, ColumnTypes, {},
 					[&OutWidgetConstructorPtr, ColumnTypes, &bFoundWidget](
 					TUniquePtr<FTypedElementWidgetConstructor> CreatedConstructor, 
 					TConstArrayView<TWeakObjectPtr<const UScriptStruct>> MatchedColumnTypes)
@@ -734,17 +734,17 @@ void FTedsOutlinerImpl::Tick()
 	}
 }
 
-ITypedElementDataStorageInterface* FTedsOutlinerImpl::GetStorage() const
+IEditorDataStorageProvider* FTedsOutlinerImpl::GetStorage() const
 {
 	return Storage;
 }
 
-ITypedElementDataStorageUiInterface* FTedsOutlinerImpl::GetStorageUI() const
+IEditorDataStorageUiProvider* FTedsOutlinerImpl::GetStorageUI() const
 {
 	return StorageUi;
 }
 
-ITypedElementDataStorageCompatibilityInterface* FTedsOutlinerImpl::GetStorageCompatibility() const
+IEditorDataStorageCompatibilityProvider* FTedsOutlinerImpl::GetStorageCompatibility() const
 {
 	return StorageCompatibility;
 }
diff --git a/Engine/Plugins/Experimental/EditorDataStorageFeatures/Source/TedsOutliner/Private/TedsOutlinerMode.cpp b/Engine/Plugins/Experimental/EditorDataStorageFeatures/Source/TedsOutliner/Private/TedsOutlinerMode.cpp
--- a/Engine/Plugins/Experimental/EditorDataStorageFeatures/Source/TedsOutliner/Private/TedsOutlinerMode.cpp
+++ b/Engine/Plugins/Experimental/EditorDataStorageFeatures/Source/TedsOutliner/Private/TedsOutlinerMode.cpp
@@ -53,7 +53,7 @@ void FTedsOutlinerMode::Rebuild()
 void FTedsOutlinerMode::OnSelectionChanged()
 {
 	TOptional<FName> SelectionSetName = TedsOutlinerImpl->GetSelectionSetName();
-	ITypedElementDataStorageInterface* Storage = TedsOutlinerImpl->GetStorage();
+	IEditorDataStorageProvider* Storage = TedsOutlinerImpl->GetStorage();
 	
 	// The selection in TEDS was changed, update the outliner to respond
 	SceneOutliner->SetSelection([SelectionSetName, Storage](ISceneOutlinerTreeItem& InItem) -> bool
@@ -139,7 +139,7 @@ bool FTedsOutlinerMode::ParseDragDrop(FSceneOutlinerDragDropPayload& OutPayload,
 FSceneOutlinerDragValidationInfo FTedsOutlinerMode::ValidateDrop(const ISceneOutlinerTreeItem& DropTarget, const FSceneOutlinerDragDropPayload& Payload) const
 {
 	const TOptional<FTedsOutlinerHierarchyData>& HierarchyData = TedsOutlinerImpl->GetHierarchyData();
-	ITypedElementDataStorageInterface* Storage = TedsOutlinerImpl->GetStorage();
+	IEditorDataStorageProvider* Storage = TedsOutlinerImpl->GetStorage();
 
 	// We don't want drag/drop if this TEDS Outliner isn't showing any hierarchy data
 	if(!HierarchyData.IsSet())
@@ -220,7 +220,7 @@ FSceneOutlinerDragValidationInfo FTedsOutlinerMode::ValidateDrop(const ISceneOut
 void FTedsOutlinerMode::OnDrop(ISceneOutlinerTreeItem& DropTarget, const FSceneOutlinerDragDropPayload& Payload, const FSceneOutlinerDragValidationInfo& ValidationInfo) const
 {
 	const TOptional<FTedsOutlinerHierarchyData>& HierarchyData = TedsOutlinerImpl->GetHierarchyData();
-	ITypedElementDataStorageInterface* Storage = TedsOutlinerImpl->GetStorage();
+	IEditorDataStorageProvider* Storage = TedsOutlinerImpl->GetStorage();
 
 	if(!Private::TedsOutlinerDragDropEnabledCvar->GetBool() || !HierarchyData.IsSet())
 	{
diff --git a/Engine/Plugins/Experimental/EditorDataStorageFeatures/Source/TedsOutliner/Private/TedsOutlinerModule.cpp b/Engine/Plugins/Experimental/EditorDataStorageFeatures/Source/TedsOutliner/Private/TedsOutlinerModule.cpp
--- a/Engine/Plugins/Experimental/EditorDataStorageFeatures/Source/TedsOutliner/Private/TedsOutlinerModule.cpp
+++ b/Engine/Plugins/Experimental/EditorDataStorageFeatures/Source/TedsOutliner/Private/TedsOutlinerModule.cpp
@@ -114,7 +114,7 @@ void FTedsOutlinerModule::ShutdownModule()
 DataStorage::QueryHandle FTedsOutlinerModule::GetLevelEditorTedsOutlinerColumnQuery()
 {
 	UTypedElementRegistry* Registry = UTypedElementRegistry::GetInstance();
-	ITypedElementDataStorageInterface* Storage = Registry->GetMutableDataStorage();
+	IEditorDataStorageProvider* Storage = Registry->GetMutableDataStorage();
 		
 	using namespace DataStorage::Queries;
 
diff --git a/Engine/Plugins/Experimental/EditorDataStorageFeatures/Source/TedsOutliner/Private/Widgets/AlertWidget.cpp b/Engine/Plugins/Experimental/EditorDataStorageFeatures/Source/TedsOutliner/Private/Widgets/AlertWidget.cpp
--- a/Engine/Plugins/Experimental/EditorDataStorageFeatures/Source/TedsOutliner/Private/Widgets/AlertWidget.cpp
+++ b/Engine/Plugins/Experimental/EditorDataStorageFeatures/Source/TedsOutliner/Private/Widgets/AlertWidget.cpp
@@ -124,7 +124,7 @@ namespace UE::Editor::DataStorage::Widgets::Private
 					ActionButton.SetOnClicked(FOnClicked::CreateLambda([RowWithAlertAction]()
 						{
 							UTypedElementRegistry* Registry = UTypedElementRegistry::GetInstance();
-							const ITypedElementDataStorageInterface* DataStorage = Registry->GetDataStorage();
+							const IEditorDataStorageProvider* DataStorage = Registry->GetDataStorage();
 							if (const FTypedElementAlertActionColumn* Action =
 								DataStorage->GetColumn<FTypedElementAlertActionColumn>(RowWithAlertAction))
 							{
@@ -148,8 +148,8 @@ namespace UE::Editor::DataStorage::Widgets::Private
 //
 
 void UAlertWidgetFactory::RegisterWidgetConstructors(
-	ITypedElementDataStorageInterface& DataStorage,
-	ITypedElementDataStorageUiInterface& DataStorageUi) const
+	IEditorDataStorageProvider& DataStorage,
+	IEditorDataStorageUiProvider& DataStorageUi) const
 {
 	using namespace UE::Editor::DataStorage::Queries;
 
@@ -160,13 +160,13 @@ void UAlertWidgetFactory::RegisterWidgetConstructors(
 		TColumn<FTypedElementAlertColumn>() || TColumn<FTypedElementChildAlertColumn>());
 }
 
-void UAlertWidgetFactory::RegisterQueries(ITypedElementDataStorageInterface& DataStorage)
+void UAlertWidgetFactory::RegisterQueries(IEditorDataStorageProvider& DataStorage)
 {
 	RegisterAlertQueries(DataStorage);
 	RegisterAlertHeaderQueries(DataStorage);
 }
 
-void UAlertWidgetFactory::RegisterAlertQueries(ITypedElementDataStorageInterface& DataStorage)
+void UAlertWidgetFactory::RegisterAlertQueries(IEditorDataStorageProvider& DataStorage)
 {
 	using namespace UE::Editor::DataStorage::Queries;
 	
@@ -247,7 +247,7 @@ void UAlertWidgetFactory::RegisterAlertQueries(ITypedElementDataStorageInterface
 		.Compile());
 }
 
-void UAlertWidgetFactory::RegisterAlertHeaderQueries(ITypedElementDataStorageInterface& DataStorage)
+void UAlertWidgetFactory::RegisterAlertHeaderQueries(IEditorDataStorageProvider& DataStorage)
 {
 	using namespace UE::Editor::DataStorage::Queries;
 	
@@ -362,8 +362,8 @@ TConstArrayView<const UScriptStruct*> FAlertWidgetConstructor::GetAdditionalColu
 	return Columns;
 }
 
-bool FAlertWidgetConstructor::FinalizeWidget(ITypedElementDataStorageInterface* DataStorage,
-	ITypedElementDataStorageUiInterface* DataStorageUi, UE::Editor::DataStorage::RowHandle Row, const TSharedPtr<SWidget>& Widget)
+bool FAlertWidgetConstructor::FinalizeWidget(IEditorDataStorageProvider* DataStorage,
+	IEditorDataStorageUiProvider* DataStorageUi, UE::Editor::DataStorage::RowHandle Row, const TSharedPtr<SWidget>& Widget)
 {
 	using namespace UE::Editor::DataStorage;
 
@@ -411,8 +411,8 @@ TConstArrayView<const UScriptStruct*> FAlertHeaderWidgetConstructor::GetAddition
 	return Columns;
 }
 
-bool FAlertHeaderWidgetConstructor::FinalizeWidget(ITypedElementDataStorageInterface* DataStorage, 
-	ITypedElementDataStorageUiInterface* DataStorageUi, UE::Editor::DataStorage::RowHandle Row, const TSharedPtr<SWidget>& Widget)
+bool FAlertHeaderWidgetConstructor::FinalizeWidget(IEditorDataStorageProvider* DataStorage, 
+	IEditorDataStorageUiProvider* DataStorageUi, UE::Editor::DataStorage::RowHandle Row, const TSharedPtr<SWidget>& Widget)
 {
 	DataStorage->AddColumn(Row, FUIHeaderPropertiesColumn
 		{
diff --git a/Engine/Plugins/Experimental/EditorDataStorageFeatures/Source/TedsOutliner/Private/Widgets/OutlinerLabelWidget.cpp b/Engine/Plugins/Experimental/EditorDataStorageFeatures/Source/TedsOutliner/Private/Widgets/OutlinerLabelWidget.cpp
--- a/Engine/Plugins/Experimental/EditorDataStorageFeatures/Source/TedsOutliner/Private/Widgets/OutlinerLabelWidget.cpp
+++ b/Engine/Plugins/Experimental/EditorDataStorageFeatures/Source/TedsOutliner/Private/Widgets/OutlinerLabelWidget.cpp
@@ -103,8 +103,8 @@ namespace UE::OutlinerLabelWidget::Local
 	}
 }
 
-void UOutlinerLabelWidgetFactory::RegisterWidgetConstructors(ITypedElementDataStorageInterface& DataStorage,
-	ITypedElementDataStorageUiInterface& DataStorageUi) const
+void UOutlinerLabelWidgetFactory::RegisterWidgetConstructors(IEditorDataStorageProvider& DataStorage,
+	IEditorDataStorageUiProvider& DataStorageUi) const
 {
 	using namespace UE::Editor::DataStorage::Queries;
 	DataStorageUi.RegisterWidgetFactory<FOutlinerLabelWidgetConstructor>(
@@ -117,8 +117,8 @@ FOutlinerLabelWidgetConstructor::FOutlinerLabelWidgetConstructor()
 {
 }
 
-TSharedPtr<SWidget> FOutlinerLabelWidgetConstructor::CreateWidget(ITypedElementDataStorageInterface* DataStorage,
-	ITypedElementDataStorageUiInterface* DataStorageUi, RowHandle TargetRow, RowHandle WidgetRow,
+TSharedPtr<SWidget> FOutlinerLabelWidgetConstructor::CreateWidget(IEditorDataStorageProvider* DataStorage,
+	IEditorDataStorageUiProvider* DataStorageUi, RowHandle TargetRow, RowHandle WidgetRow,
 	const UE::Editor::DataStorage::FMetaDataView& Arguments)
 {
 	using namespace UE::Editor::DataStorage;
@@ -171,8 +171,8 @@ TSharedPtr<SWidget> FOutlinerLabelWidgetConstructor::CreateWidget(ITypedElementD
 	
 }
 
-TSharedRef<SWidget> FOutlinerLabelWidgetConstructor::CreateLabel(ITypedElementDataStorageInterface* DataStorage,
-	ITypedElementDataStorageUiInterface* DataStorageUi, RowHandle TargetRow, RowHandle WidgetRow,
+TSharedRef<SWidget> FOutlinerLabelWidgetConstructor::CreateLabel(IEditorDataStorageProvider* DataStorage,
+	IEditorDataStorageUiProvider* DataStorageUi, RowHandle TargetRow, RowHandle WidgetRow,
 	const UE::Editor::DataStorage::FMetaDataView& Arguments)
 {
 	using namespace UE::Editor::DataStorage;
diff --git a/Engine/Plugins/Experimental/EditorDataStorageFeatures/Source/TedsOutliner/Private/Widgets/OutlinerLabelWidget.h b/Engine/Plugins/Experimental/EditorDataStorageFeatures/Source/TedsOutliner/Private/Widgets/OutlinerLabelWidget.h
--- a/Engine/Plugins/Experimental/EditorDataStorageFeatures/Source/TedsOutliner/Private/Widgets/OutlinerLabelWidget.h
+++ b/Engine/Plugins/Experimental/EditorDataStorageFeatures/Source/TedsOutliner/Private/Widgets/OutlinerLabelWidget.h
@@ -9,19 +9,19 @@
 
 #include "OutlinerLabelWidget.generated.h"
 
-class ITypedElementDataStorageInterface;
+class IEditorDataStorageProvider;
 class UScriptStruct;
 
 UCLASS()
-class UOutlinerLabelWidgetFactory : public UTypedElementDataStorageFactory
+class UOutlinerLabelWidgetFactory : public UEditorDataStorageFactory
 {
 	GENERATED_BODY()
 
 public:
 	~UOutlinerLabelWidgetFactory() override = default;
 
-	TEDSOUTLINER_API void RegisterWidgetConstructors(ITypedElementDataStorageInterface& DataStorage,
-		ITypedElementDataStorageUiInterface& DataStorageUi) const override;
+	TEDSOUTLINER_API void RegisterWidgetConstructors(IEditorDataStorageProvider& DataStorage,
+		IEditorDataStorageUiProvider& DataStorageUi) const override;
 };
 
 // Label widget for the Scene Outliner that shows an icon (with optional override information) + a text label
@@ -35,15 +35,15 @@ public:
 	~FOutlinerLabelWidgetConstructor() override = default;
 
 	TEDSOUTLINER_API virtual TSharedPtr<SWidget> CreateWidget(
-		ITypedElementDataStorageInterface* DataStorage,
-		ITypedElementDataStorageUiInterface* DataStorageUi,
+		IEditorDataStorageProvider* DataStorage,
+		IEditorDataStorageUiProvider* DataStorageUi,
 		UE::Editor::DataStorage::RowHandle TargetRow,
 		UE::Editor::DataStorage::RowHandle WidgetRow, 
 		const UE::Editor::DataStorage::FMetaDataView& Arguments) override;
 
 protected:
 
-	TSharedRef<SWidget> CreateLabel(ITypedElementDataStorageInterface* DataStorage,
-	ITypedElementDataStorageUiInterface* DataStorageUi, RowHandle TargetRow, RowHandle WidgetRow,
+	TSharedRef<SWidget> CreateLabel(IEditorDataStorageProvider* DataStorage,
+	IEditorDataStorageUiProvider* DataStorageUi, RowHandle TargetRow, RowHandle WidgetRow,
 	const UE::Editor::DataStorage::FMetaDataView& Arguments);
 };
\ No newline at end of file
diff --git a/Engine/Plugins/Experimental/EditorDataStorageFeatures/Source/TedsOutliner/Private/Widgets/TransformHeadsUpWidget.cpp b/Engine/Plugins/Experimental/EditorDataStorageFeatures/Source/TedsOutliner/Private/Widgets/TransformHeadsUpWidget.cpp
--- a/Engine/Plugins/Experimental/EditorDataStorageFeatures/Source/TedsOutliner/Private/Widgets/TransformHeadsUpWidget.cpp
+++ b/Engine/Plugins/Experimental/EditorDataStorageFeatures/Source/TedsOutliner/Private/Widgets/TransformHeadsUpWidget.cpp
@@ -152,7 +152,7 @@ static void UpdateTransformHeadsUpDisplay(FTypedElementSlateWidgetReferenceColum
 // UTransformHeadsUpWidgetFactory
 //
 
-void UTransformHeadsUpWidgetFactory::RegisterQueries(ITypedElementDataStorageInterface& DataStorage)
+void UTransformHeadsUpWidgetFactory::RegisterQueries(IEditorDataStorageProvider& DataStorage)
 {
 	using namespace UE::Editor::DataStorage::Queries;
 		
@@ -184,8 +184,8 @@ void UTransformHeadsUpWidgetFactory::RegisterQueries(ITypedElementDataStorageInt
 
 }
 
-void UTransformHeadsUpWidgetFactory::RegisterWidgetConstructors(ITypedElementDataStorageInterface& DataStorage,
-	ITypedElementDataStorageUiInterface& DataStorageUi) const
+void UTransformHeadsUpWidgetFactory::RegisterWidgetConstructors(IEditorDataStorageProvider& DataStorage,
+	IEditorDataStorageUiProvider& DataStorageUi) const
 {
 	using namespace UE::Editor::DataStorage::Queries;
 
@@ -217,8 +217,8 @@ TSharedPtr<SWidget> FTransformHeadsUpWidgetConstructor::CreateWidget(const UE::E
 }
 
 bool FTransformHeadsUpWidgetConstructor::FinalizeWidget(
-	ITypedElementDataStorageInterface* DataStorage,
-	ITypedElementDataStorageUiInterface* DataStorageUi,
+	IEditorDataStorageProvider* DataStorage,
+	IEditorDataStorageUiProvider* DataStorageUi,
 	UE::Editor::DataStorage::RowHandle Row,
 	const TSharedPtr<SWidget>& Widget)
 {
diff --git a/Engine/Plugins/Experimental/EditorDataStorageFeatures/Source/TedsOutliner/Private/Widgets/TypeInfoWidget.cpp b/Engine/Plugins/Experimental/EditorDataStorageFeatures/Source/TedsOutliner/Private/Widgets/TypeInfoWidget.cpp
--- a/Engine/Plugins/Experimental/EditorDataStorageFeatures/Source/TedsOutliner/Private/Widgets/TypeInfoWidget.cpp
+++ b/Engine/Plugins/Experimental/EditorDataStorageFeatures/Source/TedsOutliner/Private/Widgets/TypeInfoWidget.cpp
@@ -13,8 +13,8 @@
 #include "Widgets/Text/STextBlock.h"
 #include "Widgets/Images/SImage.h"
 
-void UTypeInfoWidgetFactory::RegisterWidgetConstructors(ITypedElementDataStorageInterface& DataStorage,
-	ITypedElementDataStorageUiInterface& DataStorageUi) const
+void UTypeInfoWidgetFactory::RegisterWidgetConstructors(IEditorDataStorageProvider& DataStorage,
+	IEditorDataStorageUiProvider& DataStorageUi) const
 {
 	using namespace UE::Editor::DataStorage::Queries;
 
@@ -62,8 +62,8 @@ TSharedPtr<SWidget> FTypeInfoWidgetConstructor::CreateWidget(
 	}
 }
 
-bool FTypeInfoWidgetConstructor::FinalizeWidget(ITypedElementDataStorageInterface* DataStorage,
-	ITypedElementDataStorageUiInterface* DataStorageUi, UE::Editor::DataStorage::RowHandle Row, const TSharedPtr<SWidget>& Widget)
+bool FTypeInfoWidgetConstructor::FinalizeWidget(IEditorDataStorageProvider* DataStorage,
+	IEditorDataStorageUiProvider* DataStorageUi, UE::Editor::DataStorage::RowHandle Row, const TSharedPtr<SWidget>& Widget)
 {
 	checkf(Widget, TEXT("Referenced widget is not valid. A constructed widget may not have been cleaned up. This can "
 		"also happen if this processor is running in the same phase as the processors responsible for cleaning up old "
diff --git a/Engine/Plugins/Experimental/EditorDataStorageFeatures/Source/TedsOutliner/Public/Compatibility/SceneOutlinerTedsBridge.h b/Engine/Plugins/Experimental/EditorDataStorageFeatures/Source/TedsOutliner/Public/Compatibility/SceneOutlinerTedsBridge.h
--- a/Engine/Plugins/Experimental/EditorDataStorageFeatures/Source/TedsOutliner/Public/Compatibility/SceneOutlinerTedsBridge.h
+++ b/Engine/Plugins/Experimental/EditorDataStorageFeatures/Source/TedsOutliner/Public/Compatibility/SceneOutlinerTedsBridge.h
@@ -13,9 +13,9 @@
 #include "SceneOutlinerTedsBridge.generated.h"
 
 class ISceneOutliner;
-class ITypedElementDataStorageInterface;
-class ITypedElementDataStorageUiInterface;
-class ITypedElementDataStorageCompatibilityInterface;
+class IEditorDataStorageProvider;
+class IEditorDataStorageUiProvider;
+class IEditorDataStorageCompatibilityProvider;
 
 DECLARE_DELEGATE_RetVal_OneParam(FSceneOutlinerTreeItemID, FTreeItemIDDealiaser, UE::Editor::DataStorage::RowHandle);
 
@@ -61,20 +61,20 @@ private:
 	
 	TMap<TWeakPtr<ISceneOutliner>, TSharedPtr<FSceneOutlinerTedsBridge>> SceneOutliners;
 
-	ITypedElementDataStorageInterface* Storage{ nullptr };
-	ITypedElementDataStorageUiInterface* StorageUi{ nullptr };
-	ITypedElementDataStorageCompatibilityInterface* StorageCompatibility{ nullptr };
+	IEditorDataStorageProvider* Storage{ nullptr };
+	IEditorDataStorageUiProvider* StorageUi{ nullptr };
+	IEditorDataStorageCompatibilityProvider* StorageCompatibility{ nullptr };
 
 	TMap<TWeakObjectPtr<const UScriptStruct>, FName> TEDSToOutlinerDefaultColumnMapping;
 };
 
 UCLASS()
-class USceneOutlinerTedsBridgeFactory : public UTypedElementDataStorageFactory
+class USceneOutlinerTedsBridgeFactory : public UEditorDataStorageFactory
 {
 	GENERATED_BODY()
 
 public:
 	~USceneOutlinerTedsBridgeFactory() override = default;
 
-	void RegisterWidgetPurposes(ITypedElementDataStorageUiInterface& DataStorageUi) const override;
+	void RegisterWidgetPurposes(IEditorDataStorageUiProvider& DataStorageUi) const override;
 };
diff --git a/Engine/Plugins/Experimental/EditorDataStorageFeatures/Source/TedsOutliner/Public/TedsOutlinerImpl.h b/Engine/Plugins/Experimental/EditorDataStorageFeatures/Source/TedsOutliner/Public/TedsOutlinerImpl.h
--- a/Engine/Plugins/Experimental/EditorDataStorageFeatures/Source/TedsOutliner/Public/TedsOutlinerImpl.h
+++ b/Engine/Plugins/Experimental/EditorDataStorageFeatures/Source/TedsOutliner/Public/TedsOutlinerImpl.h
@@ -7,8 +7,8 @@
 #include "ISceneOutlinerMode.h"
 #include "Elements/Columns/TypedElementHiearchyColumns.h"
 
-class ITypedElementDataStorageUiInterface;
-class ITypedElementDataStorageCompatibilityInterface;
+class IEditorDataStorageUiProvider;
+class IEditorDataStorageCompatibilityProvider;
 struct FTypedElementWidgetConstructor;
 class SWidget;
 
@@ -116,9 +116,9 @@ public:
 	void Init();
 
 	// TEDS construct getters
-	ITypedElementDataStorageInterface* GetStorage() const;
-	ITypedElementDataStorageUiInterface* GetStorageUI() const;
-	ITypedElementDataStorageCompatibilityInterface* GetStorageCompatibility() const;
+	IEditorDataStorageProvider* GetStorage() const;
+	IEditorDataStorageUiProvider* GetStorageUI() const;
+	IEditorDataStorageCompatibilityProvider* GetStorageCompatibility() const;
 
 	TOptional<FName> GetSelectionSetName() const;
 
@@ -179,9 +179,9 @@ protected:
 	
 protected:
 	// TEDS Storage Constructs
-	ITypedElementDataStorageInterface* Storage{ nullptr };
-	ITypedElementDataStorageUiInterface* StorageUi{ nullptr };
-	ITypedElementDataStorageCompatibilityInterface* StorageCompatibility{ nullptr };
+	IEditorDataStorageProvider* Storage{ nullptr };
+	IEditorDataStorageUiProvider* StorageUi{ nullptr };
+	IEditorDataStorageCompatibilityProvider* StorageCompatibility{ nullptr };
 
 	FTedsOutlinerParams CreationParams;
 
diff --git a/Engine/Plugins/Experimental/EditorDataStorageFeatures/Source/TedsOutliner/Public/Widgets/AlertWidget.h b/Engine/Plugins/Experimental/EditorDataStorageFeatures/Source/TedsOutliner/Public/Widgets/AlertWidget.h
--- a/Engine/Plugins/Experimental/EditorDataStorageFeatures/Source/TedsOutliner/Public/Widgets/AlertWidget.h
+++ b/Engine/Plugins/Experimental/EditorDataStorageFeatures/Source/TedsOutliner/Public/Widgets/AlertWidget.h
@@ -9,20 +9,20 @@
 #include "AlertWidget.generated.h"
 
 UCLASS()
-class UAlertWidgetFactory : public UTypedElementDataStorageFactory
+class UAlertWidgetFactory : public UEditorDataStorageFactory
 {
 	GENERATED_BODY()
 
 public:
 	TEDSOUTLINER_API ~UAlertWidgetFactory() override = default;
 
-	TEDSOUTLINER_API void RegisterWidgetConstructors(ITypedElementDataStorageInterface& DataStorage,
-		ITypedElementDataStorageUiInterface& DataStorageUi) const override;
-	TEDSOUTLINER_API void RegisterQueries(ITypedElementDataStorageInterface& DataStorage) override;
+	TEDSOUTLINER_API void RegisterWidgetConstructors(IEditorDataStorageProvider& DataStorage,
+		IEditorDataStorageUiProvider& DataStorageUi) const override;
+	TEDSOUTLINER_API void RegisterQueries(IEditorDataStorageProvider& DataStorage) override;
 
 private:
-	void RegisterAlertQueries(ITypedElementDataStorageInterface& DataStorage);
-	void RegisterAlertHeaderQueries(ITypedElementDataStorageInterface& DataStorage);
+	void RegisterAlertQueries(IEditorDataStorageProvider& DataStorage);
+	void RegisterAlertHeaderQueries(IEditorDataStorageProvider& DataStorage);
 };
 
 USTRUCT()
@@ -46,7 +46,7 @@ public:
 protected:
 	TEDSOUTLINER_API TSharedPtr<SWidget> CreateWidget(const UE::Editor::DataStorage::FMetaDataView& Arguments) override;
 	TEDSOUTLINER_API TConstArrayView<const UScriptStruct*> GetAdditionalColumnsList() const override;
-	TEDSOUTLINER_API bool FinalizeWidget(ITypedElementDataStorageInterface* DataStorage, ITypedElementDataStorageUiInterface* DataStorageUi,
+	TEDSOUTLINER_API bool FinalizeWidget(IEditorDataStorageProvider* DataStorage, IEditorDataStorageUiProvider* DataStorageUi,
 		UE::Editor::DataStorage::RowHandle Row, const TSharedPtr<SWidget>& Widget) override;
 };
 
@@ -68,7 +68,7 @@ public:
 protected:
 	TEDSOUTLINER_API TSharedPtr<SWidget> CreateWidget(const UE::Editor::DataStorage::FMetaDataView& Arguments) override;
 	TEDSOUTLINER_API TConstArrayView<const UScriptStruct*> GetAdditionalColumnsList() const override;
-	TEDSOUTLINER_API bool FinalizeWidget(ITypedElementDataStorageInterface* DataStorage, ITypedElementDataStorageUiInterface* DataStorageUi,
+	TEDSOUTLINER_API bool FinalizeWidget(IEditorDataStorageProvider* DataStorage, IEditorDataStorageUiProvider* DataStorageUi,
 		UE::Editor::DataStorage::RowHandle Row, const TSharedPtr<SWidget>& Widget) override;
 };
 
diff --git a/Engine/Plugins/Experimental/EditorDataStorageFeatures/Source/TedsOutliner/Public/Widgets/TransformHeadsUpWidget.h b/Engine/Plugins/Experimental/EditorDataStorageFeatures/Source/TedsOutliner/Public/Widgets/TransformHeadsUpWidget.h
--- a/Engine/Plugins/Experimental/EditorDataStorageFeatures/Source/TedsOutliner/Public/Widgets/TransformHeadsUpWidget.h
+++ b/Engine/Plugins/Experimental/EditorDataStorageFeatures/Source/TedsOutliner/Public/Widgets/TransformHeadsUpWidget.h
@@ -19,16 +19,16 @@
  *		3. Unnormalized rotation
  */
 UCLASS()
-class UTransformHeadsUpWidgetFactory : public UTypedElementDataStorageFactory
+class UTransformHeadsUpWidgetFactory : public UEditorDataStorageFactory
 {
 	GENERATED_BODY()
 
 public:
 	~UTransformHeadsUpWidgetFactory() override = default;
 
-	void RegisterQueries(ITypedElementDataStorageInterface& DataStorage) override;
-	void RegisterWidgetConstructors(ITypedElementDataStorageInterface& DataStorage,
-		ITypedElementDataStorageUiInterface& DataStorageUi) const override;
+	void RegisterQueries(IEditorDataStorageProvider& DataStorage) override;
+	void RegisterWidgetConstructors(IEditorDataStorageProvider& DataStorage,
+		IEditorDataStorageUiProvider& DataStorageUi) const override;
 };
 
 USTRUCT()
@@ -45,8 +45,8 @@ public:
 protected:
 	TSharedPtr<SWidget> CreateWidget(const UE::Editor::DataStorage::FMetaDataView& Arguments) override;
 	bool FinalizeWidget(
-		ITypedElementDataStorageInterface* DataStorage,
-		ITypedElementDataStorageUiInterface* DataStorageUi,
+		IEditorDataStorageProvider* DataStorage,
+		IEditorDataStorageUiProvider* DataStorageUi,
 		UE::Editor::DataStorage::RowHandle Row,
 		const TSharedPtr<SWidget>& Widget) override;
 };
diff --git a/Engine/Plugins/Experimental/EditorDataStorageFeatures/Source/TedsOutliner/Public/Widgets/TypeInfoWidget.h b/Engine/Plugins/Experimental/EditorDataStorageFeatures/Source/TedsOutliner/Public/Widgets/TypeInfoWidget.h
--- a/Engine/Plugins/Experimental/EditorDataStorageFeatures/Source/TedsOutliner/Public/Widgets/TypeInfoWidget.h
+++ b/Engine/Plugins/Experimental/EditorDataStorageFeatures/Source/TedsOutliner/Public/Widgets/TypeInfoWidget.h
@@ -12,15 +12,15 @@ struct FSlateBrush;
 struct FTypedElementClassTypeInfoColumn;
 
 UCLASS()
-class UTypeInfoWidgetFactory : public UTypedElementDataStorageFactory
+class UTypeInfoWidgetFactory : public UEditorDataStorageFactory
 {
 	GENERATED_BODY()
 
 public:
 	~UTypeInfoWidgetFactory() override = default;
 
-	void RegisterWidgetConstructors(ITypedElementDataStorageInterface& DataStorage,
-		ITypedElementDataStorageUiInterface& DataStorageUi) const override;
+	void RegisterWidgetConstructors(IEditorDataStorageProvider& DataStorage,
+		IEditorDataStorageUiProvider& DataStorageUi) const override;
 };
 
 USTRUCT()
@@ -35,7 +35,7 @@ public:
 protected:
 	explicit FTypeInfoWidgetConstructor(const UScriptStruct* InTypeInfo);
 	TSharedPtr<SWidget> CreateWidget(const UE::Editor::DataStorage::FMetaDataView& Arguments) override;
-	bool FinalizeWidget(ITypedElementDataStorageInterface* DataStorage, ITypedElementDataStorageUiInterface* DataStorageUi,
+	bool FinalizeWidget(IEditorDataStorageProvider* DataStorage, IEditorDataStorageUiProvider* DataStorageUi,
 		UE::Editor::DataStorage::RowHandle Row, const TSharedPtr<SWidget>& Widget) override;
 
 protected:
diff --git a/Engine/Plugins/Experimental/EditorDataStorageFeatures/Source/TedsRevisionControl/Private/Processors/RevisionControlProcessors.cpp b/Engine/Plugins/Experimental/EditorDataStorageFeatures/Source/TedsRevisionControl/Private/Processors/RevisionControlProcessors.cpp
--- a/Engine/Plugins/Experimental/EditorDataStorageFeatures/Source/TedsRevisionControl/Private/Processors/RevisionControlProcessors.cpp
+++ b/Engine/Plugins/Experimental/EditorDataStorageFeatures/Source/TedsRevisionControl/Private/Processors/RevisionControlProcessors.cpp
@@ -162,7 +162,7 @@ namespace UE::Editor::RevisionControl::Private
 	}
 } // namespace UE::Editor::RevisionControl::Private
 
-void URevisionControlDataStorageFactory::RegisterTables(ITypedElementDataStorageInterface& DataStorage)
+void URevisionControlDataStorageFactory::RegisterTables(IEditorDataStorageProvider& DataStorage)
 {
 	DataStorage.RegisterTable(
 		TTypedElementColumnTypeList<
@@ -175,7 +175,7 @@ void URevisionControlDataStorageFactory::RegisterTables(ITypedElementDataStorage
 		FName("Editor_PackageUpdateTable"));
 }
 
-void URevisionControlDataStorageFactory::RegisterQueries(ITypedElementDataStorageInterface& DataStorage)
+void URevisionControlDataStorageFactory::RegisterQueries(IEditorDataStorageProvider& DataStorage)
 {
 	using namespace UE::Editor::DataStorage;
 	using namespace UE::Editor::RevisionControl::Private;
@@ -235,7 +235,7 @@ void URevisionControlDataStorageFactory::RegisterQueries(ITypedElementDataStorag
 	}
 }
 
-void URevisionControlDataStorageFactory::RegisterFetchUpdates(ITypedElementDataStorageInterface& DataStorage)
+void URevisionControlDataStorageFactory::RegisterFetchUpdates(IEditorDataStorageProvider& DataStorage)
 {
 	using namespace UE::Editor::DataStorage::Queries;
 	
@@ -263,7 +263,7 @@ void URevisionControlDataStorageFactory::RegisterFetchUpdates(ITypedElementDataS
 	}
 }
 
-void URevisionControlDataStorageFactory::RegisterApplyOverlays(ITypedElementDataStorageInterface& DataStorage)
+void URevisionControlDataStorageFactory::RegisterApplyOverlays(IEditorDataStorageProvider& DataStorage)
 {
 	using namespace UE::Editor::DataStorage::Queries;
 	using namespace UE::Editor::RevisionControl::Private;
@@ -377,7 +377,7 @@ void URevisionControlDataStorageFactory::RegisterApplyOverlays(ITypedElementData
 	}
 }
 
-void URevisionControlDataStorageFactory::RegisterRemoveOverlays(ITypedElementDataStorageInterface& DataStorage)
+void URevisionControlDataStorageFactory::RegisterRemoveOverlays(IEditorDataStorageProvider& DataStorage)
 {
 	using namespace UE::Editor::DataStorage::Queries;
 	
diff --git a/Engine/Plugins/Experimental/EditorDataStorageFeatures/Source/TedsRevisionControl/Private/Processors/RevisionControlProcessors.h b/Engine/Plugins/Experimental/EditorDataStorageFeatures/Source/TedsRevisionControl/Private/Processors/RevisionControlProcessors.h
--- a/Engine/Plugins/Experimental/EditorDataStorageFeatures/Source/TedsRevisionControl/Private/Processors/RevisionControlProcessors.h
+++ b/Engine/Plugins/Experimental/EditorDataStorageFeatures/Source/TedsRevisionControl/Private/Processors/RevisionControlProcessors.h
@@ -7,23 +7,23 @@
 
 #include "RevisionControlProcessors.generated.h"
 
-class ITypedElementDataStorageInterface;
+class IEditorDataStorageProvider;
 
 UCLASS()
-class URevisionControlDataStorageFactory : public UTypedElementDataStorageFactory
+class URevisionControlDataStorageFactory : public UEditorDataStorageFactory
 {
 	GENERATED_BODY()
 
 public:
 	~URevisionControlDataStorageFactory() override = default;
 
-	void RegisterTables(ITypedElementDataStorageInterface& DataStorage) override;
-	void RegisterQueries(ITypedElementDataStorageInterface& DataStorage) override;
+	void RegisterTables(IEditorDataStorageProvider& DataStorage) override;
+	void RegisterQueries(IEditorDataStorageProvider& DataStorage) override;
 
 private:
-	void RegisterFetchUpdates(ITypedElementDataStorageInterface& DataStorage);
-	void RegisterApplyOverlays(ITypedElementDataStorageInterface& DataStorage);
-	void RegisterRemoveOverlays(ITypedElementDataStorageInterface& DataStorage);
+	void RegisterFetchUpdates(IEditorDataStorageProvider& DataStorage);
+	void RegisterApplyOverlays(IEditorDataStorageProvider& DataStorage);
+	void RegisterRemoveOverlays(IEditorDataStorageProvider& DataStorage);
 	UE::Editor::DataStorage::QueryHandle FetchUpdates = UE::Editor::DataStorage::InvalidQueryHandle;
 	UE::Editor::DataStorage::QueryHandle ApplyNewOverlays = UE::Editor::DataStorage::InvalidQueryHandle;
 	UE::Editor::DataStorage::QueryHandle ChangeOverlay = UE::Editor::DataStorage::InvalidQueryHandle;
diff --git a/Engine/Plugins/Experimental/EditorDataStorageFeatures/Source/TedsRevisionControl/Private/Queries/ObjectPackagePathToColumnQueries.cpp b/Engine/Plugins/Experimental/EditorDataStorageFeatures/Source/TedsRevisionControl/Private/Queries/ObjectPackagePathToColumnQueries.cpp
--- a/Engine/Plugins/Experimental/EditorDataStorageFeatures/Source/TedsRevisionControl/Private/Queries/ObjectPackagePathToColumnQueries.cpp
+++ b/Engine/Plugins/Experimental/EditorDataStorageFeatures/Source/TedsRevisionControl/Private/Queries/ObjectPackagePathToColumnQueries.cpp
@@ -22,7 +22,7 @@ namespace UE::Editor::RevisionControl::Private
 		TEXT("Automatically query revision control provider and fill information into TEDS")
 	);
 
-	static void ResolvePackageReference(ITypedElementDataStorageInterface::IQueryContext& Context, const UPackage* Package, RowHandle Row, RowHandle PackageRow)
+	static void ResolvePackageReference(IEditorDataStorageProvider::IQueryContext& Context, const UPackage* Package, RowHandle Row, RowHandle PackageRow)
 	{
 		FTypedElementPackageReference PackageReference;
 		PackageReference.Row = PackageRow;
@@ -43,7 +43,7 @@ namespace UE::Editor::RevisionControl::Private
 	};
 } // namespace UE::Editor::DataStorage
 
-void UTypedElementUObjectPackagePathFactory::RegisterQueries(ITypedElementDataStorageInterface& DataStorage)
+void UTypedElementUObjectPackagePathFactory::RegisterQueries(IEditorDataStorageProvider& DataStorage)
 {
 	using namespace UE::Editor::DataStorage::Queries;
 	
@@ -89,7 +89,7 @@ void UTypedElementUObjectPackagePathFactory::RegisterQueries(ITypedElementDataSt
 	}
 }
 
-void UTypedElementUObjectPackagePathFactory::RegisterTryAddPackageRef(ITypedElementDataStorageInterface& DataStorage)
+void UTypedElementUObjectPackagePathFactory::RegisterTryAddPackageRef(IEditorDataStorageProvider& DataStorage)
 {
 	using namespace UE::Editor::DataStorage::Queries;
 	
diff --git a/Engine/Plugins/Experimental/EditorDataStorageFeatures/Source/TedsRevisionControl/Private/Queries/ObjectPackagePathToColumnQueries.h b/Engine/Plugins/Experimental/EditorDataStorageFeatures/Source/TedsRevisionControl/Private/Queries/ObjectPackagePathToColumnQueries.h
--- a/Engine/Plugins/Experimental/EditorDataStorageFeatures/Source/TedsRevisionControl/Private/Queries/ObjectPackagePathToColumnQueries.h
+++ b/Engine/Plugins/Experimental/EditorDataStorageFeatures/Source/TedsRevisionControl/Private/Queries/ObjectPackagePathToColumnQueries.h
@@ -7,19 +7,19 @@
 
 #include "ObjectPackagePathToColumnQueries.generated.h"
 
-class ITypedElementDataStorageInterface;
+class IEditorDataStorageProvider;
 
 UCLASS()
-class UTypedElementUObjectPackagePathFactory : public UTypedElementDataStorageFactory
+class UTypedElementUObjectPackagePathFactory : public UEditorDataStorageFactory
 {
 	GENERATED_BODY()
 
 public:
 	~UTypedElementUObjectPackagePathFactory() override = default;
 
-	void RegisterQueries(ITypedElementDataStorageInterface& DataStorage) override;
+	void RegisterQueries(IEditorDataStorageProvider& DataStorage) override;
 
 private:
-	void RegisterTryAddPackageRef(ITypedElementDataStorageInterface& DataStorage);
+	void RegisterTryAddPackageRef(IEditorDataStorageProvider& DataStorage);
 	UE::Editor::DataStorage::QueryHandle TryAddPackageRef = UE::Editor::DataStorage::InvalidQueryHandle;
 };
diff --git a/Engine/Plugins/Experimental/EditorDataStorageFeatures/Source/TedsSettings/Private/TedsSettingsManager.cpp b/Engine/Plugins/Experimental/EditorDataStorageFeatures/Source/TedsSettings/Private/TedsSettingsManager.cpp
--- a/Engine/Plugins/Experimental/EditorDataStorageFeatures/Source/TedsSettings/Private/TedsSettingsManager.cpp
+++ b/Engine/Plugins/Experimental/EditorDataStorageFeatures/Source/TedsSettings/Private/TedsSettingsManager.cpp
@@ -51,7 +51,7 @@ void FTedsSettingsManager::Initialize()
 
 		auto OnDataStorage = [this, TypedElementRegistry]
 			{
-				ITypedElementDataStorageInterface* DataStorage = TypedElementRegistry->GetMutableDataStorage();
+				IEditorDataStorageProvider* DataStorage = TypedElementRegistry->GetMutableDataStorage();
 				check(DataStorage);
 
 				RegisterTables(*DataStorage);
@@ -83,7 +83,7 @@ void FTedsSettingsManager::Shutdown()
 
 		if (TypedElementRegistry->AreDataStorageInterfacesSet())
 		{
-			ITypedElementDataStorageInterface* DataStorage = TypedElementRegistry->GetMutableDataStorage();
+			IEditorDataStorageProvider* DataStorage = TypedElementRegistry->GetMutableDataStorage();
 			check(DataStorage);
 
 			UnregisterSettings();
@@ -94,7 +94,7 @@ void FTedsSettingsManager::Shutdown()
 	}
 }
 
-void FTedsSettingsManager::RegisterTables(ITypedElementDataStorageInterface& DataStorage)
+void FTedsSettingsManager::RegisterTables(IEditorDataStorageProvider& DataStorage)
 {
 	if (SettingsContainerTable == UE::Editor::DataStorage::InvalidTableHandle)
 	{
@@ -111,7 +111,7 @@ void FTedsSettingsManager::RegisterTables(ITypedElementDataStorageInterface& Dat
 	}
 }
 
-void FTedsSettingsManager::RegisterQueries(ITypedElementDataStorageInterface& DataStorage)
+void FTedsSettingsManager::RegisterQueries(IEditorDataStorageProvider& DataStorage)
 {
 	using namespace UE::Editor::DataStorage::Queries;
 
@@ -126,7 +126,7 @@ void FTedsSettingsManager::RegisterQueries(ITypedElementDataStorageInterface& Da
 	}
 }
 
-void FTedsSettingsManager::UnregisterQueries(ITypedElementDataStorageInterface& DataStorage)
+void FTedsSettingsManager::UnregisterQueries(IEditorDataStorageProvider& DataStorage)
 {
 	DataStorage.UnregisterQuery(SelectAllSettingsQuery);
 	SelectAllSettingsQuery = UE::Editor::DataStorage::InvalidQueryHandle;
@@ -162,7 +162,7 @@ void FTedsSettingsManager::RegisterSettingsContainer(const FName& ContainerName)
 	UTypedElementRegistry* TypedElementRegistry = UTypedElementRegistry::GetInstance();
 	check(TypedElementRegistry);
 
-	ITypedElementDataStorageInterface* DataStorage = TypedElementRegistry->GetMutableDataStorage();
+	IEditorDataStorageProvider* DataStorage = TypedElementRegistry->GetMutableDataStorage();
 	check(DataStorage);
 
 	UE_LOG(LogTedsSettings, Log, TEXT("Register Settings Container : '%s'"), *ContainerName.ToString());
@@ -215,10 +215,10 @@ void FTedsSettingsManager::UnregisterSettings()
 	UTypedElementRegistry* TypedElementRegistry = UTypedElementRegistry::GetInstance();
 	check(TypedElementRegistry);
 
-	ITypedElementDataStorageInterface* DataStorage = TypedElementRegistry->GetMutableDataStorage();
+	IEditorDataStorageProvider* DataStorage = TypedElementRegistry->GetMutableDataStorage();
 	check(DataStorage);
 
-	ITypedElementDataStorageCompatibilityInterface* DataStorageCompatibility = TypedElementRegistry->GetMutableDataStorageCompatibility();
+	IEditorDataStorageCompatibilityProvider* DataStorageCompatibility = TypedElementRegistry->GetMutableDataStorageCompatibility();
 	check(DataStorageCompatibility);
 
 	SettingsModule->OnContainerAdded().RemoveAll(this);
@@ -286,10 +286,10 @@ void FTedsSettingsManager::UpdateSettingsCategory(TSharedPtr<ISettingsCategory>
 	UTypedElementRegistry* TypedElementRegistry = UTypedElementRegistry::GetInstance();
 	check(TypedElementRegistry);
 
-	ITypedElementDataStorageInterface* DataStorage = TypedElementRegistry->GetMutableDataStorage();
+	IEditorDataStorageProvider* DataStorage = TypedElementRegistry->GetMutableDataStorage();
 	check(DataStorage);
 
-	ITypedElementDataStorageCompatibilityInterface* DataStorageCompatibility = TypedElementRegistry->GetMutableDataStorageCompatibility();
+	IEditorDataStorageCompatibilityProvider* DataStorageCompatibility = TypedElementRegistry->GetMutableDataStorageCompatibility();
 	check(DataStorageCompatibility);
 
 	const FName& ContainerName = DataStorage->GetColumn<FNameColumn>(ContainerRow)->Name;
diff --git a/Engine/Plugins/Experimental/EditorDataStorageFeatures/Source/TedsSettings/Private/TedsSettingsManager.h b/Engine/Plugins/Experimental/EditorDataStorageFeatures/Source/TedsSettings/Private/TedsSettingsManager.h
--- a/Engine/Plugins/Experimental/EditorDataStorageFeatures/Source/TedsSettings/Private/TedsSettingsManager.h
+++ b/Engine/Plugins/Experimental/EditorDataStorageFeatures/Source/TedsSettings/Private/TedsSettingsManager.h
@@ -23,10 +23,10 @@ public:
 
 private:
 
-	void RegisterTables(ITypedElementDataStorageInterface& DataStorage);
+	void RegisterTables(IEditorDataStorageProvider& DataStorage);
 
-	void RegisterQueries(ITypedElementDataStorageInterface& DataStorage);
-	void UnregisterQueries(ITypedElementDataStorageInterface& DataStorage);
+	void RegisterQueries(IEditorDataStorageProvider& DataStorage);
+	void UnregisterQueries(IEditorDataStorageProvider& DataStorage);
 
 	void RegisterSettings();
 	void UnregisterSettings();
diff --git a/Engine/Plugins/Experimental/EditorDataStorageFeatures/Source/TedsSettings/Tests/TedsSettingsTests.cpp b/Engine/Plugins/Experimental/EditorDataStorageFeatures/Source/TedsSettings/Tests/TedsSettingsTests.cpp
--- a/Engine/Plugins/Experimental/EditorDataStorageFeatures/Source/TedsSettings/Tests/TedsSettingsTests.cpp
+++ b/Engine/Plugins/Experimental/EditorDataStorageFeatures/Source/TedsSettings/Tests/TedsSettingsTests.cpp
@@ -19,13 +19,11 @@
 
 namespace UE::Editor::Settings::Tests
 {
-	using namespace UE::Editor;
-
 	BEGIN_DEFINE_SPEC(FTedsSettingsTestFixture, "Editor.DataStorage.Settings", EAutomationTestFlags::EditorContext | EAutomationTestFlags::EngineFilter)
 	ISettingsModule* SettingsModule = nullptr;
 	UTypedElementRegistry* TypedElementRegistry = nullptr;
-	ITypedElementDataStorageInterface* DataStorage = nullptr;
-	ITypedElementDataStorageCompatibilityInterface* DataStorageCompatibility = nullptr;
+	IEditorDataStorageProvider* DataStorage = nullptr;
+	IEditorDataStorageCompatibilityProvider* DataStorageCompatibility = nullptr;
 	DataStorage::QueryHandle CountAllSettingsQuery = DataStorage::InvalidQueryHandle;
 
 	uint32 BeforeRowCount = 0;
diff --git a/Engine/Plugins/Experimental/EditorDataStorageFeatures/Source/TedsStyling/Private/TedsStylingFactory.cpp b/Engine/Plugins/Experimental/EditorDataStorageFeatures/Source/TedsStyling/Private/TedsStylingFactory.cpp
--- a/Engine/Plugins/Experimental/EditorDataStorageFeatures/Source/TedsStyling/Private/TedsStylingFactory.cpp
+++ b/Engine/Plugins/Experimental/EditorDataStorageFeatures/Source/TedsStyling/Private/TedsStylingFactory.cpp
@@ -26,19 +26,19 @@ namespace UE::TedsStylingFactory::Local
 
 }
 
-void UTedsStylingFactory::RegisterTables(ITypedElementDataStorageInterface& DataStorage)
+void UTedsStylingFactory::RegisterTables(IEditorDataStorageProvider& DataStorage)
 {
 	DataStorage.RegisterTable(TTypedElementColumnTypeList<FNameColumn, FSlateStyleSetColumn, FSlateStyleTag>(), UE::TedsStylingFactory::Local::TableName);
 }
 
-void UTedsStylingFactory::RegisterQueries(ITypedElementDataStorageInterface& DataStorage)
+void UTedsStylingFactory::RegisterQueries(IEditorDataStorageProvider& DataStorage)
 {
 	Super::RegisterQueries(DataStorage);
 }
 
 void UTedsStylingFactory::RegisterAllKnownStyles()
 {
-	ITypedElementDataStorageInterface* Interface = UTypedElementRegistry::GetInstance()->GetMutableDataStorage();
+	IEditorDataStorageProvider* Interface = UTypedElementRegistry::GetInstance()->GetMutableDataStorage();
 
 	auto RegisterBrushesForStyle = [Interface](const ISlateStyle& Style)
 	{
@@ -85,7 +85,7 @@ void UTedsStylingFactory::RegisterAllKnownStyles()
 	);
 }
 
-void UTedsStylingFactory::RegisterBrush(ITypedElementDataStorageInterface* DataStorage, const FName& StyleName, const FSlateBrush* Brush, const ISlateStyle& OwnerStyle)
+void UTedsStylingFactory::RegisterBrush(IEditorDataStorageProvider* DataStorage, const FName& StyleName, const FSlateBrush* Brush, const ISlateStyle& OwnerStyle)
 {
 	// We currently don't store the FSlateBrush* itself to avoid storing a raw pointer that we can query on demand anyways in case we want to support
 	// swapping brushes or unloading styles on demand in the future
@@ -104,7 +104,7 @@ void UTedsStylingFactory::RegisterBrush(ITypedElementDataStorageInterface* DataS
 	}
 }
 
-void UTedsStylingFactory::RegisterColor(ITypedElementDataStorageInterface* DataStorage, const FName& StyleName, const FSlateColor& Color,
+void UTedsStylingFactory::RegisterColor(IEditorDataStorageProvider* DataStorage, const FName& StyleName, const FSlateColor& Color,
 	const ISlateStyle& OwnerStyle)
 {
 	const UE::Editor::DataStorage::RowHandle Row = AddOrGetStyleRow(DataStorage, StyleName, OwnerStyle);
@@ -116,7 +116,7 @@ void UTedsStylingFactory::RegisterColor(ITypedElementDataStorageInterface* DataS
 
 }
 
-UE::Editor::DataStorage::RowHandle UTedsStylingFactory::AddOrGetStyleRow(ITypedElementDataStorageInterface* DataStorage,
+UE::Editor::DataStorage::RowHandle UTedsStylingFactory::AddOrGetStyleRow(IEditorDataStorageProvider* DataStorage,
 	const FName& StyleName, const ISlateStyle& OwnerStyle)
 {
 	const FName StyleSetName = OwnerStyle.GetStyleSetName();
