Add option to hide row handle column in Teds table viewer

diff --git a/Engine/Plugins/Experimental/Animation/HierarchyTable/Source/Editor/Private/HierarchyTableEditorToolkit.cpp b/Engine/Plugins/Experimental/Animation/HierarchyTable/Source/Editor/Private/HierarchyTableEditorToolkit.cpp
--- a/Engine/Plugins/Experimental/Animation/HierarchyTable/Source/Editor/Private/HierarchyTableEditorToolkit.cpp
+++ b/Engine/Plugins/Experimental/Animation/HierarchyTable/Source/Editor/Private/HierarchyTableEditorToolkit.cpp
@@ -155,6 +155,7 @@ TSharedRef<SWidget> FHierarchyTableEditorToolkit::CreateTedsOutliner()
 		Params.QueryDescription = RowQueryDescription;
 		Params.CellWidgetPurposes = TArray<FName>{ TEXT("General.Cell") };
 		Params.HierarchyData = FTedsOutlinerHierarchyData::GetDefaultHierarchyData();
+		Params.bShowRowHandleColumn = false;
 	}
 
 	FTedsOutlinerModule& TedsOutlinerModule = FModuleManager::GetModuleChecked<FTedsOutlinerModule>("TedsOutliner");
diff --git a/Engine/Plugins/Experimental/EditorDataStorageFeatures/Source/TedsOutliner/Private/TedsOutlinerModule.cpp b/Engine/Plugins/Experimental/EditorDataStorageFeatures/Source/TedsOutliner/Private/TedsOutlinerModule.cpp
--- a/Engine/Plugins/Experimental/EditorDataStorageFeatures/Source/TedsOutliner/Private/TedsOutlinerModule.cpp
+++ b/Engine/Plugins/Experimental/EditorDataStorageFeatures/Source/TedsOutliner/Private/TedsOutlinerModule.cpp
@@ -80,12 +80,15 @@ TSharedRef<ISceneOutliner> FTedsOutlinerModule::CreateTedsOutliner(const FSceneO
 	});
 
 	// Add the custom column that displays row handles
-	InitOptions.ColumnMap.Add(FSceneOutlinerRowHandleColumn::GetID(),
-		FSceneOutlinerColumnInfo(ESceneOutlinerColumnVisibility::Visible, 2,
-			FCreateSceneOutlinerColumn::CreateLambda([](ISceneOutliner& InSceneOutliner)
-			{
-				return MakeShareable(new FSceneOutlinerRowHandleColumn(InSceneOutliner));
-			})));
+	if (InInitTedsOptions.bShowRowHandleColumn)
+	{
+		InitOptions.ColumnMap.Add(FSceneOutlinerRowHandleColumn::GetID(),
+			FSceneOutlinerColumnInfo(ESceneOutlinerColumnVisibility::Visible, 2,
+				FCreateSceneOutlinerColumn::CreateLambda([](ISceneOutliner& InSceneOutliner)
+				{
+					return MakeShareable(new FSceneOutlinerRowHandleColumn(InSceneOutliner));
+				})));
+	}
 	
 	InitOptions.ColumnMap.Add(FSceneOutlinerBuiltInColumnTypes::Label(), FSceneOutlinerColumnInfo(ESceneOutlinerColumnVisibility::Visible, 10));
 
diff --git a/Engine/Plugins/Experimental/EditorDataStorageFeatures/Source/TedsOutliner/Public/TedsOutlinerImpl.h b/Engine/Plugins/Experimental/EditorDataStorageFeatures/Source/TedsOutliner/Public/TedsOutlinerImpl.h
--- a/Engine/Plugins/Experimental/EditorDataStorageFeatures/Source/TedsOutliner/Public/TedsOutlinerImpl.h
+++ b/Engine/Plugins/Experimental/EditorDataStorageFeatures/Source/TedsOutliner/Public/TedsOutlinerImpl.h
@@ -73,6 +73,7 @@ struct FTedsOutlinerParams
 	: SceneOutliner(InSceneOutliner)
 	, QueryDescription()
 	, bUseDefaultTedsFilters(false)
+	, bShowRowHandleColumn(true)
 	, HierarchyData(FTedsOutlinerHierarchyData::GetDefaultHierarchyData())
 	, CellWidgetPurposes{TEXT("SceneOutliner.Cell"), TEXT("General.Cell")}
 	{}
@@ -90,6 +91,9 @@ struct FTedsOutlinerParams
 	// If true, this Outliner will automatically add all TEDS tags and columns as filters
 	bool bUseDefaultTedsFilters;
 
+	// If true, this Outliner will include a column for row handle
+	bool bShowRowHandleColumn;
+
 	// If specified, this is how the TEDS Outliner will handle hierarchies. If not specified - there will be no hierarchies shown as a
 	// parent-child relation in the tree view
 	TOptional<FTedsOutlinerHierarchyData> HierarchyData;
