TedsDebugger : Make the Query Editor tab a PanelTab

This allows it to be floated out to its own window, but not docked into the root window ( e.g. next to the details panel ).  This fixes a crash where docking it in the root window and then closing TEDS Debugger would crash.

This changelist also fixes a bug in the OnPersistLayout lambda where we need to respect IsSavingLayoutOnClosedAllowed().  This is false when resetting to the default layout ( Window -> Load Layout -> Default Editor Layout )

diff --git a/Engine/Plugins/Experimental/EditorDataStorageFeatures/Source/TedsDebugger/Private/STedsDebugger.cpp b/Engine/Plugins/Experimental/EditorDataStorageFeatures/Source/TedsDebugger/Private/STedsDebugger.cpp
--- a/Engine/Plugins/Experimental/EditorDataStorageFeatures/Source/TedsDebugger/Private/STedsDebugger.cpp
+++ b/Engine/Plugins/Experimental/EditorDataStorageFeatures/Source/TedsDebugger/Private/STedsDebugger.cpp
@@ -78,9 +78,12 @@ void STedsDebugger::Construct(const FArguments& InArgs, const TSharedRef<SDockTa
 		FTabManager::FOnPersistLayout::CreateStatic(
 			[](const TSharedRef<FTabManager::FLayout>& InLayout)
 			{
-				if (InLayout->GetPrimaryArea().Pin().IsValid())
+				if (FGlobalTabmanager::Get()->CanSavePersistentLayouts())
 				{
-					FLayoutSaveRestore::SaveToConfig(GEditorLayoutIni, InLayout);
+					if (InLayout->GetPrimaryArea().Pin().IsValid())
+					{
+						FLayoutSaveRestore::SaveToConfig(GEditorLayoutIni, InLayout);
+					}
 				}
 			}
 		)
@@ -125,7 +128,7 @@ TSharedRef<SDockTab> STedsDebugger::SpawnToolbar(const FSpawnTabArgs& Args)
 
 TSharedRef<SDockTab> STedsDebugger::SpawnQueryEditorTab(const FSpawnTabArgs& Args)
 {
-	TSharedRef<SDockTab> DockTab = SNew(SDockTab).TabRole(ETabRole::NomadTab);
+	TSharedRef<SDockTab> DockTab = SNew(SDockTab).TabRole(ETabRole::PanelTab);
 	if (!QueryEditorModel)
 	{
 		UTypedElementRegistry* Registry = UTypedElementRegistry::GetInstance();
