Revision Control Overlays:
Fix a bug with the revision control processors using the wrong context causing the viewport overlay color column to never be added
Add observers to add/remove the viewport color column on selectoin change
Add queries to update the relevant SCC rows when a viewport overlay setting is changed

diff --git a/Engine/Plugins/Experimental/EditorDataStorageFeatures/Source/TedsRevisionControl/Private/Processors/RevisionControlProcessors.cpp b/Engine/Plugins/Experimental/EditorDataStorageFeatures/Source/TedsRevisionControl/Private/Processors/RevisionControlProcessors.cpp
--- a/Engine/Plugins/Experimental/EditorDataStorageFeatures/Source/TedsRevisionControl/Private/Processors/RevisionControlProcessors.cpp
+++ b/Engine/Plugins/Experimental/EditorDataStorageFeatures/Source/TedsRevisionControl/Private/Processors/RevisionControlProcessors.cpp
@@ -9,7 +9,6 @@
 
 #include "Elements/Columns/TypedElementCompatibilityColumns.h"
 #include "Elements/Columns/TypedElementMiscColumns.h"
-#include "Elements/Columns/TypedElementPackageColumns.h"
 #include "Elements/Columns/TypedElementRevisionControlColumns.h"
 #include "Elements/Columns/TypedElementSelectionColumns.h"
 #include "Elements/Columns/TypedElementViewportColumns.h"
@@ -22,6 +21,28 @@ namespace UE::Editor::RevisionControl::Private
 {
 	extern FAutoConsoleVariableRef CVarAutoPopulateState;
 
+	// Update the overlay color for all rows with InColumn
+	void UpdateSCCOverlayStates(IConsoleVariable*, UScriptStruct* InColumn)
+	{
+		IEditorDataStorageProvider* DataStorage = UTypedElementRegistry::GetInstance()->GetMutableDataStorage();
+
+		if(const URevisionControlDataStorageFactory* Factory = DataStorage->FindFactory<URevisionControlDataStorageFactory>())
+		{
+			Factory->UpdateOverlaysForSCCState(DataStorage, InColumn);
+		}
+	}
+
+	// Update all currently existing overlay colors
+	void UpdateOverlayColors(IConsoleVariable*)
+	{
+		IEditorDataStorageProvider* DataStorage = UTypedElementRegistry::GetInstance()->GetMutableDataStorage();
+
+		if(const URevisionControlDataStorageFactory* Factory = DataStorage->FindFactory<URevisionControlDataStorageFactory>())
+		{
+			Factory->UpdateOverlayColors(DataStorage);
+		}
+	}
+
 	static bool gEnableOverlays = false;
 	TAutoConsoleVariable<bool> CVarEnableOverlays(
 		TEXT("RevisionControl.Overlays.Enable"),
@@ -34,6 +55,7 @@ namespace UE::Editor::RevisionControl::Private
 		TEXT("RevisionControl.Overlays.CheckedOutByOtherUser.Enable"),
 		gEnableOverlayCheckedOutByOtherUser,
 		TEXT("Enables overlays for files that are checked out by another user."),
+		FConsoleVariableDelegate::CreateStatic(UpdateSCCOverlayStates, FSCCExternallyLockedColumn::StaticStruct()),
 		ECVF_Default);
 
 	static bool gEnableOverlayNotAtHeadRevision = true;
@@ -41,6 +63,7 @@ namespace UE::Editor::RevisionControl::Private
 		TEXT("RevisionControl.Overlays.NotAtHeadRevision.Enable"),
 		gEnableOverlayNotAtHeadRevision,
 		TEXT("Enables overlays for files that are not at the latest revision."),
+		FConsoleVariableDelegate::CreateStatic(UpdateSCCOverlayStates, FSCCNotCurrentTag::StaticStruct()),
 		ECVF_Default);
 
 	static bool gEnableOverlayCheckedOut = false;
@@ -48,6 +71,7 @@ namespace UE::Editor::RevisionControl::Private
 		TEXT("RevisionControl.Overlays.CheckedOut.Enable"),
 		gEnableOverlayCheckedOut,
 		TEXT("Enables overlays for files that are checked out by user."),
+		FConsoleVariableDelegate::CreateStatic(UpdateSCCOverlayStates, FSCCLockedTag::StaticStruct()),
 		ECVF_Default);
 
 	static bool gEnableOverlayOpenForAdd = false;
@@ -55,6 +79,7 @@ namespace UE::Editor::RevisionControl::Private
 		TEXT("RevisionControl.Overlays.OpenForAdd.Enable"),
 		gEnableOverlayOpenForAdd,
 		TEXT("Enables overlays for files that are newly added."),
+		FConsoleVariableDelegate::CreateStatic(UpdateSCCOverlayStates, FSCCStatusColumn::StaticStruct()),
 		ECVF_Default);
 
 	static int32 gOverlayAlpha = 20; // [0..100]
@@ -62,6 +87,7 @@ namespace UE::Editor::RevisionControl::Private
 		TEXT("RevisionControl.Overlays.Alpha"),
 		gOverlayAlpha,
 		TEXT("Configures overlay opacity."),
+		FConsoleVariableDelegate::CreateStatic(UpdateOverlayColors),
 		ECVF_Default);
 
 	#if UE_BUILD_SHIPPING
@@ -80,7 +106,7 @@ namespace UE::Editor::RevisionControl::Private
 	#endif
 
 	using namespace UE::Editor::DataStorage;
-	static FColor DetermineOverlayColor(const IQueryContext& ObjectContext, const ICommonQueryContext& SCCContext, const FTypedElementUObjectColumn& Actor)
+	static FColor DetermineOverlayColor(const ICommonQueryContext& SCCContext, const FTypedElementUObjectColumn& Actor, bool bSelected)
 	{
 		check(IsInGameThread());
 
@@ -105,7 +131,6 @@ namespace UE::Editor::RevisionControl::Private
 
 		bool bExternal = Actor.Object.IsValid() ? Cast<AActor>(Actor.Object)->IsPackageExternal() : false;
 		bool bIgnored = !bExternal;
-		bool bSelected = ObjectContext.HasColumn<FTypedElementSelectionColumn>();
 		if (!bIgnored && !bSelected)
 		{
 			// Convert CVar value from [0..100] to [0..255] range.
@@ -215,6 +240,12 @@ void URevisionControlDataStorageFactory::RegisterQueries(IEditorDataStorageProvi
 				DataStorage.UnregisterQuery(ApplyOverlaysObjectToSCC);
 				ApplyOverlaysObjectToSCC = InvalidQueryHandle;
 
+				DataStorage.UnregisterQuery(SelectionAdded);
+				SelectionAdded = InvalidQueryHandle;
+
+				DataStorage.UnregisterQuery(SelectionRemoved);
+				SelectionRemoved = InvalidQueryHandle;
+
 				RegisterRemoveOverlays(DataStorage);
 			}
 		}
@@ -233,6 +264,71 @@ void URevisionControlDataStorageFactory::RegisterQueries(IEditorDataStorageProvi
 	{
 		RegisterRemoveOverlays(DataStorage);
 	}
+	
+	RegisterGeneralQueries(DataStorage);
+}
+
+void URevisionControlDataStorageFactory::UpdateOverlaysForSCCState(IEditorDataStorageProvider* DataStorage, const UScriptStruct* Column) const
+{
+	using namespace UE::Editor::DataStorage::Queries;
+
+	// The column provided must be present in our map
+	const QueryHandle* QueryForColumn = GeneralQueriesMap.Find(Column);
+	if(!QueryForColumn)
+	{
+		return;
+	}
+
+	// Collect all SCC and Object rows to request a package update
+	TArray<FTypedElementPackageUpdateColumn> PackagesToUpdate;
+
+	DirectQueryCallback RowCollector = CreateDirectQueryCallbackBinding(
+		[&PackagesToUpdate, DataStorage] (const IDirectQueryContext& Context, const FTypedElementPackageReference* PackageReference)
+		{
+			uint32 RowCount = Context.GetRowCount();
+			TConstArrayView<RowHandle> RowHandles = Context.GetRowHandles();
+
+			for(uint32 RowIndex = 0; RowIndex < RowCount; ++RowIndex)
+			{
+				if(DataStorage->IsRowAvailable(PackageReference[RowIndex].Row))
+				{
+					// We are querying the SCC rows so the package row is the one present in the callback and the object
+					// row is the one present in FTypedElementPackageReference
+ 					PackagesToUpdate.Add(FTypedElementPackageUpdateColumn{.ObjectRow = PackageReference[RowIndex].Row, .PackageRow = RowHandles[RowIndex]});
+				}
+			}
+		});
+
+	DataStorage->RunQuery(*QueryForColumn, RowCollector);
+	RequestPackageUpdates(DataStorage, PackagesToUpdate);
+}
+
+void URevisionControlDataStorageFactory::UpdateOverlayColors(IEditorDataStorageProvider* DataStorage) const
+{
+	using namespace UE::Editor::DataStorage::Queries;
+
+	TArray<FTypedElementPackageUpdateColumn> PackagesToUpdate;
+
+	DirectQueryCallback RowCollector = CreateDirectQueryCallbackBinding(
+		[&PackagesToUpdate] (const IDirectQueryContext& Context, const FTypedElementPackageReference* PackageReference)
+		{
+			uint32 RowCount = Context.GetRowCount();
+			TConstArrayView<RowHandle> RowHandles = Context.GetRowHandles();
+
+			for(uint32 RowIndex = 0; RowIndex < RowCount; ++RowIndex)
+			{
+				if(PackageReference[RowIndex].Row != InvalidRowHandle)
+				{
+					// The viewport color column is on the actor rows, so the row handle in the callback is the Object Row and the row handle
+					// in FTypedElementPackageReference is the SCC row
+					PackagesToUpdate.Add(FTypedElementPackageUpdateColumn{.ObjectRow = RowHandles[RowIndex], .PackageRow = PackageReference[RowIndex].Row});
+				}
+			}
+		});
+
+	DataStorage->RunQuery(FetchOverlayColors, RowCollector);
+	RequestPackageUpdates(DataStorage, PackagesToUpdate);
+
 }
 
 void URevisionControlDataStorageFactory::RegisterFetchUpdates(IEditorDataStorageProvider& DataStorage)
@@ -298,22 +394,23 @@ void URevisionControlDataStorageFactory::RegisterApplyOverlays(IEditorDataStorag
 			.Compile());
 	}
 
+	enum EFlushPackageUpdatesSubqueries
+	{
+		EApplyOverlaysObjectToSCC,
+		EApplyNewOverlays,
+		EChangeOverlay,
+			
+		Num
+	};
+	TArray<RowHandle> Subqueries;
+	Subqueries.AddUninitialized(EFlushPackageUpdatesSubqueries::Num);
+	Subqueries[EApplyOverlaysObjectToSCC] = ApplyOverlaysObjectToSCC;
+	Subqueries[EApplyNewOverlays] = ApplyNewOverlays;
+	Subqueries[EChangeOverlay] = ChangeOverlay;
+
 	if (FlushPackageUpdates == InvalidQueryHandle)
 	{
 		check(ApplyOverlaysObjectToSCC != InvalidQueryHandle && ApplyNewOverlays != InvalidQueryHandle && ChangeOverlay!= InvalidQueryHandle);
-		enum EFlushPackageUpdatesSubqueries
-		{
-			EApplyOverlaysObjectToSCC,
-			EApplyNewOverlays,
-			EChangeOverlay,
-			
-			Num
-		};
-		TArray<RowHandle> Subqueries;
-		Subqueries.AddUninitialized(EFlushPackageUpdatesSubqueries::Num);
-		Subqueries[EApplyOverlaysObjectToSCC] = ApplyOverlaysObjectToSCC;
-		Subqueries[EApplyNewOverlays] = ApplyNewOverlays;
-		Subqueries[EChangeOverlay] = ChangeOverlay;
 		
 		FlushPackageUpdates = DataStorage.RegisterQuery(
 			Select(
@@ -327,15 +424,15 @@ void URevisionControlDataStorageFactory::RegisterApplyOverlays(IEditorDataStorag
 					//   Determine if a color should be applied based on SCC status tags
 					//   If so, add the OverlayColorColumn to the actor row
 					Context.RunSubquery(EApplyNewOverlays, Update.ObjectRow, CreateSubqueryCallbackBinding(
-						[&Context](RowHandle ObjectRow, const FTypedElementUObjectColumn& Actor, const FTypedElementPackageReference& PackageReference)
+						[&Context](ISubqueryContext& ActorQueryContext, RowHandle ObjectRow, const FTypedElementUObjectColumn& Actor, const FTypedElementPackageReference& PackageReference)
 						{
 							Context.RunSubquery(EApplyOverlaysObjectToSCC, PackageReference.Row, CreateSubqueryCallbackBinding(
-								[&Context, &ObjectRow, &Actor](ISubqueryContext& SubQueryContext)
+								[&ActorQueryContext, &ObjectRow, &Actor](ISubqueryContext& SubQueryContext)
 								{
-									FColor Color = DetermineOverlayColor(Context, SubQueryContext, Actor);
+									FColor Color = DetermineOverlayColor(SubQueryContext, Actor, ActorQueryContext.HasColumn<FTypedElementSelectionColumn>());
 									if (Color.Bits != 0)
 									{
-										Context.AddColumn<FTypedElementViewportOverlayColorColumn>(ObjectRow, { .OverlayColor = Color });
+										ActorQueryContext.AddColumn<FTypedElementViewportOverlayColorColumn>(ObjectRow, { .OverlayColor = Color });
 									}
 								})
 							);
@@ -349,21 +446,21 @@ void URevisionControlDataStorageFactory::RegisterApplyOverlays(IEditorDataStorag
 					//
 					// Note: Remove and re-add will trigger observer in TypedElementActorViewportProcessors to SetOverlayColor on the primitive components
 					Context.RunSubquery(EChangeOverlay, Update.ObjectRow, CreateSubqueryCallbackBinding(
-						[&Context](RowHandle ObjectRow, const FTypedElementUObjectColumn& Actor, const FTypedElementPackageReference& PackageReference, const FTypedElementViewportOverlayColorColumn& OverlayColorColumn)
+						[&Context](ISubqueryContext& ActorQueryContext, RowHandle ObjectRow, const FTypedElementUObjectColumn& Actor, const FTypedElementPackageReference& PackageReference, const FTypedElementViewportOverlayColorColumn& OverlayColorColumn)
 						{
 							Context.RunSubquery(EApplyOverlaysObjectToSCC, PackageReference.Row, CreateSubqueryCallbackBinding(
-								[&Context, &ObjectRow, &Actor, &OverlayColorColumn](ISubqueryContext& SubQueryContext)
+								[&ActorQueryContext, &ObjectRow, &Actor, &OverlayColorColumn](ISubqueryContext& SubQueryContext)
 								{
-									FColor Color = DetermineOverlayColor(Context, SubQueryContext, Actor);
+									FColor Color = DetermineOverlayColor(SubQueryContext, Actor, ActorQueryContext.HasColumn<FTypedElementSelectionColumn>());
 									if (Color.Bits == 0)
 									{
-										Context.RemoveColumns<FTypedElementViewportOverlayColorColumn>(ObjectRow);
+										ActorQueryContext.RemoveColumns<FTypedElementViewportOverlayColorColumn>(ObjectRow);
 									}
 									else if (Color != OverlayColorColumn.OverlayColor)
 									{
 										// Remove and re-add to trigger the observer
-										Context.RemoveColumns<FTypedElementViewportOverlayColorColumn>(ObjectRow);
-										Context.AddColumn<FTypedElementViewportOverlayColorColumn>(ObjectRow, { .OverlayColor = Color });
+										ActorQueryContext.RemoveColumns<FTypedElementViewportOverlayColorColumn>(ObjectRow);
+										ActorQueryContext.AddColumn<FTypedElementViewportOverlayColorColumn>(ObjectRow, { .OverlayColor = Color });
 									}
 								})
 							);
@@ -375,6 +472,59 @@ void URevisionControlDataStorageFactory::RegisterApplyOverlays(IEditorDataStorag
 				.SubQuery(Subqueries)
 			.Compile());
 	}
+
+	SelectionAdded = DataStorage.RegisterQuery(
+					Select(
+					TEXT("Update Overlay on Selection"),
+					FObserver::OnAdd<FTypedElementSelectionColumn>(),
+					[this](IQueryContext& Context, RowHandle RowHandle, const FTypedElementSelectionColumn& SelectionColumn)
+					{
+						// We only care about the level editor's selection set for now. When the selection column is made dynamic
+						// we can directly query for it
+						if(SelectionColumn.SelectionSet.IsNone())
+						{
+							// Since we know DetermineOverlayColor() removes the overlay when the row is selected we can directly
+							// remove the column here to skip the need to do all the checks in there. But if that logic ever
+							// changes we will have to update it here too
+							Context.RemoveColumns<FTypedElementViewportOverlayColorColumn>(RowHandle);
+						}
+					})
+					.Where()
+						.All<FTypedElementActorTag, FTypedElementViewportOverlayColorColumn, FTypedElementPackageReference>()
+					.Compile());
+	
+	SelectionRemoved = DataStorage.RegisterQuery(
+					Select(
+					TEXT("Update Overlay on Deselection"),
+					FObserver::OnRemove<FTypedElementSelectionColumn>(),
+					[this](IQueryContext& ActorQueryContext, RowHandle ObjectRow, const FTypedElementSelectionColumn& SelectionColumn,
+						const FTypedElementPackageReference& PackageReference, const FTypedElementUObjectColumn& Actor)
+					{
+						// We only care about the level editor's selection set for now. When the selection column is made dynamic
+						// we can directly query for it
+						if(SelectionColumn.SelectionSet.IsNone())
+						{
+							// When an item is deselected, add the viewport overlay color column to it if applicable
+							ActorQueryContext.RunSubquery(EApplyOverlaysObjectToSCC, PackageReference.Row, CreateSubqueryCallbackBinding(
+								[&ActorQueryContext, &ObjectRow, &Actor](ISubqueryContext& SubQueryContext)
+								{
+									// We manually set the selection as false because the item is being deslected, but the row still has the selection column
+									// when the observer is fired
+									bool bSelected = false;
+									FColor Color = DetermineOverlayColor(SubQueryContext, Actor, bSelected);
+									if (Color.Bits != 0)
+									{
+										ActorQueryContext.AddColumn<FTypedElementViewportOverlayColorColumn>(ObjectRow, { .OverlayColor = Color });
+									}
+								})
+							);
+						}
+					})
+					.Where()
+						.All<FTypedElementActorTag>()
+					.DependsOn()
+						.SubQuery(Subqueries)
+					.Compile());
 }
 
 void URevisionControlDataStorageFactory::RegisterRemoveOverlays(IEditorDataStorageProvider& DataStorage)
@@ -402,4 +552,46 @@ void URevisionControlDataStorageFactory::RegisterRemoveOverlays(IEditorDataStora
 				.All<FTypedElementActorTag>()
 			.Compile());
 	}
-}
\ No newline at end of file
+}
+
+void URevisionControlDataStorageFactory::RegisterGeneralQueries(IEditorDataStorageProvider& DataStorage)
+{
+	using namespace UE::Editor::DataStorage::Queries;
+
+	// SCC Columns we want to update the overlay state for
+	const TArray<const UScriptStruct*> Columns{ FSCCNotCurrentTag::StaticStruct(), FSCCExternallyLockedColumn::StaticStruct(),
+		FSCCStatusColumn::StaticStruct(), FSCCLockedTag::StaticStruct() };
+
+	for(const UScriptStruct* Column : Columns)
+	{
+		QueryHandle Query = DataStorage.RegisterQuery(
+			Select()
+				.ReadOnly<FTypedElementPackageReference>()
+			.Where()
+				.All(Column)
+			.Compile());
+
+		GeneralQueriesMap.Add(Column, Query);
+	}
+
+	if(FetchOverlayColors == InvalidQueryHandle)
+	{
+		FetchOverlayColors = DataStorage.RegisterQuery(
+		Select()
+			.ReadOnly<FTypedElementPackageReference>()
+		.Where()
+			.All<FTypedElementViewportOverlayColorColumn, FTypedElementActorTag>()
+		.Compile());
+	}
+}
+
+void URevisionControlDataStorageFactory::RequestPackageUpdates(IEditorDataStorageProvider* DataStorage, TArray<FTypedElementPackageUpdateColumn>& Packages) const
+{
+	UE::Editor::DataStorage::TableHandle Table = DataStorage->FindTable(FName("Editor_PackageUpdateTable"));
+
+	for(FTypedElementPackageUpdateColumn& PackageUpdate : Packages)
+	{
+		UE::Editor::DataStorage::RowHandle UpdateRow = DataStorage->AddRow(Table);
+		DataStorage->AddColumn(UpdateRow, MoveTemp(PackageUpdate));
+	}
+}
diff --git a/Engine/Plugins/Experimental/EditorDataStorageFeatures/Source/TedsRevisionControl/Private/Processors/RevisionControlProcessors.h b/Engine/Plugins/Experimental/EditorDataStorageFeatures/Source/TedsRevisionControl/Private/Processors/RevisionControlProcessors.h
--- a/Engine/Plugins/Experimental/EditorDataStorageFeatures/Source/TedsRevisionControl/Private/Processors/RevisionControlProcessors.h
+++ b/Engine/Plugins/Experimental/EditorDataStorageFeatures/Source/TedsRevisionControl/Private/Processors/RevisionControlProcessors.h
@@ -2,6 +2,7 @@
 
 #pragma once
 
+#include "Elements/Columns/TypedElementPackageColumns.h"
 #include "Elements/Common/TypedElementHandles.h"
 #include "Elements/Interfaces/TypedElementDataStorageFactory.h"
 
@@ -20,14 +21,34 @@ public:
 	void RegisterTables(IEditorDataStorageProvider& DataStorage) override;
 	void RegisterQueries(IEditorDataStorageProvider& DataStorage) override;
 
+	// Update the overlays for all SCC rows that contain the Column specified
+	void UpdateOverlaysForSCCState(IEditorDataStorageProvider* DataStorage, const UScriptStruct* Column) const;
+
+	// Update the color for all actors that currently have an overlay
+	void UpdateOverlayColors(IEditorDataStorageProvider* DataStorage) const;
+
+
 private:
 	void RegisterFetchUpdates(IEditorDataStorageProvider& DataStorage);
 	void RegisterApplyOverlays(IEditorDataStorageProvider& DataStorage);
 	void RegisterRemoveOverlays(IEditorDataStorageProvider& DataStorage);
+	void RegisterGeneralQueries(IEditorDataStorageProvider& DataStorage);
+
+	void RequestPackageUpdates(IEditorDataStorageProvider* DataStorage, TArray<FTypedElementPackageUpdateColumn>& Packages) const;
+	
 	UE::Editor::DataStorage::QueryHandle FetchUpdates = UE::Editor::DataStorage::InvalidQueryHandle;
 	UE::Editor::DataStorage::QueryHandle ApplyNewOverlays = UE::Editor::DataStorage::InvalidQueryHandle;
 	UE::Editor::DataStorage::QueryHandle ChangeOverlay = UE::Editor::DataStorage::InvalidQueryHandle;
 	UE::Editor::DataStorage::QueryHandle ApplyOverlaysObjectToSCC = UE::Editor::DataStorage::InvalidQueryHandle;
 	UE::Editor::DataStorage::QueryHandle RemoveOverlays = UE::Editor::DataStorage::InvalidQueryHandle;
 	UE::Editor::DataStorage::QueryHandle FlushPackageUpdates = UE::Editor::DataStorage::InvalidQueryHandle;
+	UE::Editor::DataStorage::QueryHandle SelectionAdded = UE::Editor::DataStorage::InvalidQueryHandle;
+	UE::Editor::DataStorage::QueryHandle SelectionRemoved = UE::Editor::DataStorage::InvalidQueryHandle;
+
+	// General queries to gather SCC rows
+	TMap<const UScriptStruct*, UE::Editor::DataStorage::QueryHandle> GeneralQueriesMap;
+
+	// Query to fetch all rows with overlays
+	UE::Editor::DataStorage::QueryHandle FetchOverlayColors = UE::Editor::DataStorage::InvalidQueryHandle;
+
 };
