TEDS-CB Integration:

Modify the CB extension to use a delegate that creates the custom view so that multiple CB's can have independant views
Add a new alert column to the table viewer, and add alerts to any textures that are non-square
Create a new widget for Skeleton assets to link directly to them

diff --git a/Engine/Plugins/Experimental/EditorDataStorageFeatures/Source/TedsAssetData/Private/Processors/AssetProcessors.cpp b/Engine/Plugins/Experimental/EditorDataStorageFeatures/Source/TedsAssetData/Private/Processors/AssetProcessors.cpp
--- a/Engine/Plugins/Experimental/EditorDataStorageFeatures/Source/TedsAssetData/Private/Processors/AssetProcessors.cpp
+++ b/Engine/Plugins/Experimental/EditorDataStorageFeatures/Source/TedsAssetData/Private/Processors/AssetProcessors.cpp
@@ -14,8 +14,75 @@
 #include "Elements/Interfaces/TypedElementQueryStorageInterfaces.h"
 #include "Experimental/ContentBrowserExtensionUtils.h"
 #include "TedsAssetDataColumns.h"
+#include "Elements/Columns/TypedElementAlertColumns.h"
 #include "Elements/Columns/TypedElementMiscColumns.h"
 #include "Elements/Framework/TypedElementIndexHasher.h"
+#include "String/LexFromString.h"
+
+#define LOCTEXT_NAMESPACE "UTedsAssetDataFactory"
+
+namespace UE::TedsAssetDataFactory::Private
+{
+	// @return true if any of the texture dimensions are not a power of 2
+	bool IsTextureNonSquare(FStringView Dimensions)
+	{
+		// Texture dimensions are of the form XxY (2D) or XxYxZ (3D)
+
+		// Extract the X Dimension first (everything to the left of the first x)
+		int32 FirstXPos;
+		if(!Dimensions.FindChar('x', FirstXPos))
+		{
+			return false; // Failsafe, in case the dimension isn't in the correct format somehow.
+		}
+		
+		FStringView XDimension(Dimensions.Left(FirstXPos));
+		int32 X;
+
+		// Convert the X Dimension to an int
+		LexFromString(X, XDimension);
+
+		// If it isn't a power of 2, early out since we already know this is a non square without needing to check the rest
+		if(!FMath::IsPowerOfTwo(X))
+		{
+			return true;
+		}
+
+		// Extract out the YDimension, or YxZ in case of a 3D texture
+		FStringView YDimension(Dimensions.RightChop(FirstXPos + 1));
+		int32 SecondXPos;
+
+		// If we find another x, this is a 3D texture
+		if(YDimension.FindChar('x', SecondXPos))
+		{
+			// Extract out the Z Dimension first, which is everything to the right of the second X. And then convert it to an int
+			FStringView ZDimension = YDimension.RightChop(SecondXPos + 1);
+			int32 Z;
+			LexFromString(Z, ZDimension);
+
+			// Early out if the Z Dimension is not a power of 2 so we don't need to check Y Dimension
+			if(!FMath::IsPowerOfTwo(Z))
+			{
+				return true;
+			}
+
+			// The Y Dimension is everything between the two x's for a 3D texture
+			YDimension = YDimension.Left(SecondXPos);
+		}
+
+		// Regardless of 2D/3D, once we are here YDimension should be the correct value so check it finally
+		int32 Y;
+		LexFromString(Y, YDimension);
+
+		if(!FMath::IsPowerOfTwo(Y))
+		{
+			return true;
+		}
+
+		// If we are here, all dimensions are a power of 2 and the texture is square.
+		return false;
+
+	}
+}
 
 void UTedsAssetDataFactory::RegisterQueries(IEditorDataStorageProvider& DataStorage)
 {
@@ -40,7 +107,7 @@ void UTedsAssetDataFactory::RegisterQueries(IEditorDataStorageProvider& DataStor
 			}
 		)
 		.Where()
-			.All<FFolderTag, FUpdatedPathTag>()
+			.All<FFolderTag, FUpdatedPathTag, FVirtualPathColumn_Experimental>()
 		.Compile()
 		);
 
@@ -60,7 +127,41 @@ void UTedsAssetDataFactory::RegisterQueries(IEditorDataStorageProvider& DataStor
 			}
 		)
 		.Where()
-			.All<FFolderTag, FTypedElementSyncBackToWorldTag>()
+			.All<FFolderTag, FTypedElementSyncBackToWorldTag, FVirtualPathColumn_Experimental>()
+		.Compile()
+		);
+
+	DataStorage.RegisterQuery(
+		Select(
+			TEXT("TedsAssetDataFactory: Add/Remove non-square texture warning"),
+			FProcessor(EQueryTickPhase::PostPhysics, DataStorage.GetQueryTickGroupName(EQueryTickGroups::SyncExternalToDataStorage)),
+			[](IQueryContext& Context, const RowHandle* Rows)
+			{
+				const int32 NumRows = Context.GetRowCount();
+
+				TConstArrayView<FItemStringAttributeColumn_Experimental> ColumnView = MakeConstArrayView(
+								Context.GetColumn<FItemStringAttributeColumn_Experimental>("Dimensions"),
+								NumRows);
+				
+				for (int32 Index = 0; Index < NumRows; ++Index)
+				{
+					if(UE::TedsAssetDataFactory::Private::IsTextureNonSquare(ColumnView[Index].Value))
+					{
+						FTypedElementAlertColumn Alert;
+						Alert.AlertType = FTypedElementAlertColumnType::Error;
+						Alert.Message = LOCTEXT("NonSquareTextureAlert", "Texture has a non-square aspect ratio");
+						Context.AddColumn(Rows[Index], MoveTemp(Alert));
+					}
+					else
+					{
+						Context.RemoveColumns<FTypedElementAlertColumn>(Rows[Index]);
+					}
+				}
+			}
+		)
+		.ReadOnly<FItemStringAttributeColumn_Experimental>("Dimensions")
+		.Where()
+			.All<FAssetTag, FUpdatedAssetDataTag>()
 		.Compile()
 		);
 }
@@ -91,3 +192,5 @@ void UTedsAssetDataFactory::OnSetFolderColor(const FString& Path, IEditorDataSto
 		DataStorage->AddColumn<FUpdatedPathTag>(Row);
 	}
 }
+
+#undef LOCTEXT_NAMESPACE
diff --git a/Engine/Plugins/Experimental/EditorDataStorageFeatures/Source/TedsAssetData/Private/Widgets/SkeletalMeshSkeletonWidget.cpp b/Engine/Plugins/Experimental/EditorDataStorageFeatures/Source/TedsAssetData/Private/Widgets/SkeletalMeshSkeletonWidget.cpp
new file mode 100644
--- /dev/null
+++ b/Engine/Plugins/Experimental/EditorDataStorageFeatures/Source/TedsAssetData/Private/Widgets/SkeletalMeshSkeletonWidget.cpp
@@ -0,0 +1,65 @@
+// Copyright Epic Games, Inc. All Rights Reserved.
+
+#include "Widgets/SkeletalMeshSkeletonWidget.h"
+
+#include "IContentBrowserSingleton.h"
+#include "Elements/Framework/TypedElementAttributeBinding.h"
+#include "Internationalization/Text.h"
+#include "TedsAssetDataColumns.h"
+#include "AssetRegistry/AssetRegistryModule.h"
+#include "UObject/ObjectMacros.h"
+#include "Widgets/Input/SHyperlink.h"
+
+namespace UE::SkeletalMeshSkeletonWidget::Private
+{
+	
+	FString GetAssetName(const FString& RawPath)
+	{
+		return FSoftObjectPath(RawPath).GetAssetName();
+	}
+	
+	void NavigateToAsset(IEditorDataStorageProvider* DataStorage, UE::Editor::DataStorage::RowHandle Row)
+	{
+		if (FItemStringAttributeColumn_Experimental* SkeletonAttributeColumn =
+				DataStorage->GetColumn<FItemStringAttributeColumn_Experimental>(Row, FName("Skeleton")))
+		{
+			FAssetRegistryModule& AssetRegistryModule = FModuleManager::LoadModuleChecked<FAssetRegistryModule>("AssetRegistry");
+			TArray<FAssetData> Assets;
+			Assets.Add(AssetRegistryModule.Get().GetAssetByObjectPath(FSoftObjectPath(SkeletonAttributeColumn->Value)));
+			IContentBrowserSingleton::Get().SyncBrowserToAssets(Assets);
+		}
+	}
+}
+
+void USkeletalMeshSkeletonWidgetFactory::RegisterWidgetConstructors(IEditorDataStorageProvider& DataStorage,
+														IEditorDataStorageUiProvider& DataStorageUi) const
+{
+	using namespace UE::Editor::DataStorage::Queries;
+	DataStorageUi.RegisterWidgetFactory<FSkeletalMeshSkeletonWidgetConstructor>(TEXT("General.Cell"), TColumn<FItemStringAttributeColumn_Experimental>("Skeleton"));
+}
+
+FSkeletalMeshSkeletonWidgetConstructor::FSkeletalMeshSkeletonWidgetConstructor()
+	: FSimpleWidgetConstructor(StaticStruct())
+{
+}
+
+TSharedPtr<SWidget> FSkeletalMeshSkeletonWidgetConstructor::CreateWidget(IEditorDataStorageProvider* DataStorage,
+	IEditorDataStorageUiProvider* DataStorageUi, UE::Editor::DataStorage::RowHandle TargetRow, UE::Editor::DataStorage::RowHandle WidgetRow,
+	const UE::Editor::DataStorage::FMetaDataView& Arguments)
+{
+	UE::Editor::DataStorage::FAttributeBinder Binder(TargetRow, DataStorage);
+
+	return SNew(SHorizontalBox)
+			+SHorizontalBox::Slot()
+			.AutoWidth()
+			[
+				SNew(SHyperlink)
+				.Text(Binder.BindData(&FItemStringAttributeColumn_Experimental::Value, [](const FString& SkeletonPath)
+				{
+					return FText::FromString(UE::SkeletalMeshSkeletonWidget::Private::GetAssetName(SkeletonPath));
+				}, FString(), FName("Skeleton")))
+				.ToolTipText(Binder.BindText(&FItemStringAttributeColumn_Experimental::Value, FName("Skeleton")))
+				.OnNavigate_Static(&UE::SkeletalMeshSkeletonWidget::Private::NavigateToAsset, DataStorage, TargetRow)
+				.Style(FAppStyle::Get(), "Common.GotoBlueprintHyperlink")
+			];
+}
diff --git a/Engine/Plugins/Experimental/EditorDataStorageFeatures/Source/TedsAssetData/Private/Widgets/SkeletalMeshSkeletonWidget.h b/Engine/Plugins/Experimental/EditorDataStorageFeatures/Source/TedsAssetData/Private/Widgets/SkeletalMeshSkeletonWidget.h
new file mode 100644
--- /dev/null
+++ b/Engine/Plugins/Experimental/EditorDataStorageFeatures/Source/TedsAssetData/Private/Widgets/SkeletalMeshSkeletonWidget.h
@@ -0,0 +1,41 @@
+// Copyright Epic Games, Inc. All Rights Reserved.
+
+#pragma once
+
+#include "Elements/Interfaces/TypedElementDataStorageFactory.h"
+#include "Elements/Interfaces/TypedElementDataStorageUiInterface.h"
+
+#include "SkeletalMeshSkeletonWidget.generated.h"
+
+class IEditorDataStorageProvider;
+class UScriptStruct;
+
+UCLASS()
+class USkeletalMeshSkeletonWidgetFactory : public UEditorDataStorageFactory
+{
+	GENERATED_BODY()
+
+public:
+	~USkeletalMeshSkeletonWidgetFactory() override = default;
+
+	TEDSASSETDATA_API void RegisterWidgetConstructors(IEditorDataStorageProvider& DataStorage,
+		IEditorDataStorageUiProvider& DataStorageUi) const override;
+};
+
+// Specialized widget to display the "Skeleton" metadata on static mesh assets
+USTRUCT()
+struct FSkeletalMeshSkeletonWidgetConstructor : public FSimpleWidgetConstructor
+{
+	GENERATED_BODY()
+
+public:
+	TEDSASSETDATA_API FSkeletalMeshSkeletonWidgetConstructor();
+	~FSkeletalMeshSkeletonWidgetConstructor() override = default;
+
+	TEDSASSETDATA_API virtual TSharedPtr<SWidget> CreateWidget(
+		IEditorDataStorageProvider* DataStorage,
+		IEditorDataStorageUiProvider* DataStorageUi,
+		UE::Editor::DataStorage::RowHandle TargetRow,
+		UE::Editor::DataStorage::RowHandle WidgetRow,
+		const UE::Editor::DataStorage::FMetaDataView& Arguments) override;
+};
\ No newline at end of file
diff --git a/Engine/Plugins/Experimental/EditorDataStorageFeatures/Source/TedsContentBrowser/Private/TedsContentBrowserModule.cpp b/Engine/Plugins/Experimental/EditorDataStorageFeatures/Source/TedsContentBrowser/Private/TedsContentBrowserModule.cpp
--- a/Engine/Plugins/Experimental/EditorDataStorageFeatures/Source/TedsContentBrowser/Private/TedsContentBrowserModule.cpp
+++ b/Engine/Plugins/Experimental/EditorDataStorageFeatures/Source/TedsContentBrowser/Private/TedsContentBrowserModule.cpp
@@ -13,6 +13,7 @@
 #include "Modules/ModuleManager.h"
 #include "QueryStack/FQueryStackNode_RowView.h"
 #include "TedsAssetDataColumns.h"
+#include "Elements/Columns/TypedElementAlertColumns.h"
 #include "Elements/Columns/TypedElementMiscColumns.h"
 #include "Widgets/STedsTableViewer.h"
 #include "Widgets/Views/SListView.h"
@@ -36,7 +37,8 @@ namespace UE::Editor::ContentBrowser
 			
 			if(bEnableTedsContentBrowser)
 			{
-				ContentBrowserModule.SetContentBrowserViewExtender(TedsContentBrowserModule.GetContentBrowserViewExtender());
+				ContentBrowserModule.SetContentBrowserViewExtender(
+					FContentBrowserModule::FCreateViewExtender::CreateStatic(&FTedsContentBrowserModule::CreateContentBrowserViewExtender));
 			}
 			else
 			{
@@ -44,13 +46,9 @@ namespace UE::Editor::ContentBrowser
 			}
 		}));
 
-	TSharedPtr<IContentBrowserViewExtender> FTedsContentBrowserModule::GetContentBrowserViewExtender()
+	TSharedPtr<IContentBrowserViewExtender> FTedsContentBrowserModule::CreateContentBrowserViewExtender()
 	{
-		if(!ViewExtender)
-		{
-			ViewExtender = MakeShared<FTedsContentBrowserViewExtender>();
-		}
-		return ViewExtender;
+		return MakeShared<FTedsContentBrowserViewExtender>();
 	}
 
 	void FTedsContentBrowserModule::StartupModule()
@@ -88,21 +86,7 @@ namespace UE::Editor::ContentBrowser
 
 	void FTedsContentBrowserViewExtender::AddRow(const TSharedPtr<FAssetViewItem>& Item)
 	{
-		FAssetData ItemAssetData;
-		FName PackagePath;
-		
-		RowHandle RowHandle = InvalidRowHandle;
-		
-		if (Item->GetItem().Legacy_TryGetAssetData(ItemAssetData))
-		{
-			IndexHash IndexHash = GenerateIndexHash(ItemAssetData.GetSoftObjectPath());
-			RowHandle = DataStorage->FindIndexedRow(IndexHash);
-		}
-		else if(Item->GetItem().Legacy_TryGetPackagePath(PackagePath))
-		{
-			IndexHash IndexHash = GenerateIndexHash(PackagePath);
-			RowHandle = DataStorage->FindIndexedRow(IndexHash);
-		}
+		RowHandle RowHandle = GetRowFromAssetViewItem(Item);
 		
 		if(DataStorage->IsRowAssigned(RowHandle))
 		{
@@ -125,27 +109,49 @@ namespace UE::Editor::ContentBrowser
 		return nullptr;
 	}
 
+	DataStorage::RowHandle FTedsContentBrowserViewExtender::GetRowFromAssetViewItem(const TSharedPtr<FAssetViewItem>& Item)
+	{
+		FAssetData ItemAssetData;
+		FName PackagePath;
+		
+		RowHandle RowHandle = InvalidRowHandle;
+		
+		if (Item->GetItem().Legacy_TryGetAssetData(ItemAssetData))
+		{
+			IndexHash IndexHash = GenerateIndexHash(ItemAssetData.GetSoftObjectPath());
+			RowHandle = DataStorage->FindIndexedRow(IndexHash);
+		}
+		else if(Item->GetItem().Legacy_TryGetPackagePath(PackagePath))
+		{
+			IndexHash IndexHash = GenerateIndexHash(PackagePath);
+			RowHandle = DataStorage->FindIndexedRow(IndexHash);
+		}
+
+		return RowHandle;
+	}
+
 	FTedsContentBrowserViewExtender::FTedsContentBrowserViewExtender()
 	{
 		DataStorage = GetMutableDataStorageFeature<IEditorDataStorageProvider>(StorageFeatureName);
 		
 		RowQueryStack = MakeShared<FQueryStackNode_RowView>(&Rows);
 
-		// Sample dynamic column to display the "Triangles" attribute on static meshes
+		// Sample dynamic column to display the "Skeleton" attribute on skeletal meshes
 		// We probably want the dynamic columns in the table viewer to be data driven based on the rows in the future
-		const UScriptStruct* DynamicStaticMeshTrianglesColumn = DataStorage->GenerateDynamicColumn(FDynamicColumnDescription
+		const UScriptStruct* DynamicSkeletalMeshSkeletonColumn = DataStorage->GenerateDynamicColumn(FDynamicColumnDescription
 						{
 							.TemplateType = FItemStringAttributeColumn_Experimental::StaticStruct(),
-							.Identifier = "Triangles"
+							.Identifier = "Skeleton"
 						});
-
+		
 		// Create the table viewer widget
 		TableViewer = SNew(STedsTableViewer)
 					.QueryStack(RowQueryStack)
 					.CellWidgetPurposes({TEXT("General.RowLabel"), TEXT("General.Cell")})
 					// Default list of columns to display
-					.Columns({ FNameColumn::StaticStruct(), FAssetClassColumn::StaticStruct(), FAssetTag::StaticStruct(), FAssetPathColumn_Experimental::StaticStruct(),
-						FDiskSizeColumn::StaticStruct(), FVirtualPathColumn_Experimental::StaticStruct(), DynamicStaticMeshTrianglesColumn })
+					.Columns({ FNameColumn::StaticStruct(), FTypedElementAlertColumn::StaticStruct(),
+						FAssetClassColumn::StaticStruct(), FAssetTag::StaticStruct(), FAssetPathColumn_Experimental::StaticStruct(),
+						FDiskSizeColumn::StaticStruct(), FVirtualPathColumn_Experimental::StaticStruct(), DynamicSkeletalMeshSkeletonColumn })
 					.ListSelectionMode(ESelectionMode::Multi)
 					.OnSelectionChanged_Lambda([this](RowHandle Row)
 					{
@@ -256,17 +262,42 @@ namespace UE::Editor::ContentBrowser
 
 	void FTedsContentBrowserViewExtender::SetSelection(const TSharedPtr<FAssetViewItem>& Item, bool bSelected, const ESelectInfo::Type SelectInfo)
 	{
-		// CB 2.0 TODO: Implement using Teds selection
+		RowHandle Row = GetRowFromAssetViewItem(Item);
+
+		if(DataStorage->IsRowAssigned(Row))
+		{
+			// We have to defer the selection by a tick because this fires on path change which has to refresh the internal list of assets.
+			// The table viewer doesn't refresh immediately but rather on tick by checking if the query stack is dirty. If we set the selection
+			// before the refresh happens SListView will deselect the item since it isn't visible in the list yet.
+			// Long term selection should also be handled through TEDS so it happens at the proper time automatically.
+			FTSTicker::GetCoreTicker().AddTicker(FTickerDelegate::CreateLambda([this, Row, bSelected, SelectInfo](float)
+			{
+				TableViewer->SetSelection(Row, bSelected, SelectInfo);
+				return false;
+			}));
+		}
 	}
 
 	void FTedsContentBrowserViewExtender::RequestScrollIntoView(const TSharedPtr<FAssetViewItem>& Item)
 	{
-		// CB 2.0 TODO: Implement using a Teds column
+		RowHandle Row = GetRowFromAssetViewItem(Item);
+
+		if(DataStorage->IsRowAssigned(Row))
+		{
+			// We have to defer the scroll by a tick because this fires on path change which has to refresh the internal list of assets.
+			// The table viewer doesn't refresh immediately but rather on tick by checking if the query stack is dirty. If we request scroll
+			// before the refresh happens SListView will ignore the request since the item isn't visible in the list yet.
+			FTSTicker::GetCoreTicker().AddTicker(FTickerDelegate::CreateLambda([this, Row](float)
+			{
+				TableViewer->ScrollIntoView(Row);
+				return false;
+			}));
+		}
 	}
 
 	void FTedsContentBrowserViewExtender::ClearSelection()
 	{
-		// CB 2.0 TODO: Implement using a Teds column
+		TableViewer->ClearSelection();
 	}
 
 	bool FTedsContentBrowserViewExtender::IsRightClickScrolling()
diff --git a/Engine/Plugins/Experimental/EditorDataStorageFeatures/Source/TedsContentBrowser/Private/TedsContentBrowserModule.h b/Engine/Plugins/Experimental/EditorDataStorageFeatures/Source/TedsContentBrowser/Private/TedsContentBrowserModule.h
--- a/Engine/Plugins/Experimental/EditorDataStorageFeatures/Source/TedsContentBrowser/Private/TedsContentBrowserModule.h
+++ b/Engine/Plugins/Experimental/EditorDataStorageFeatures/Source/TedsContentBrowser/Private/TedsContentBrowserModule.h
@@ -53,6 +53,7 @@ namespace UE::Editor::ContentBrowser
 		// Get the internal FAssetViewItem from a row handle
 		TSharedPtr<FAssetViewItem> GetAssetViewItemFromRow(DataStorage::RowHandle Row);
 
+		DataStorage::RowHandle GetRowFromAssetViewItem(const TSharedPtr<FAssetViewItem>& Item);
 
 	private:
 
@@ -88,13 +89,10 @@ namespace UE::Editor::ContentBrowser
 
 		FTedsContentBrowserModule() = default;
 		
-		TSharedPtr<IContentBrowserViewExtender> GetContentBrowserViewExtender();
+		static TSharedPtr<IContentBrowserViewExtender> CreateContentBrowserViewExtender();
 
 		// IModuleInterface interface
 		virtual void StartupModule() override;
 		virtual void ShutdownModule() override;
-
-	private:
-		TSharedPtr<FTedsContentBrowserViewExtender> ViewExtender;
 	};
 } // namespace UE::Editor::ContentBrowser
diff --git a/Engine/Plugins/Experimental/EditorDataStorageFeatures/Source/TedsTableViewer/Private/Widgets/STedsTableViewer.cpp b/Engine/Plugins/Experimental/EditorDataStorageFeatures/Source/TedsTableViewer/Private/Widgets/STedsTableViewer.cpp
--- a/Engine/Plugins/Experimental/EditorDataStorageFeatures/Source/TedsTableViewer/Private/Widgets/STedsTableViewer.cpp
+++ b/Engine/Plugins/Experimental/EditorDataStorageFeatures/Source/TedsTableViewer/Private/Widgets/STedsTableViewer.cpp
@@ -166,6 +166,21 @@ namespace UE::Editor::DataStorage
 		return TedsWidget->GetRowHandle();
 	}
 
+	void STedsTableViewer::SetSelection(RowHandle Row, bool bSelected, const ESelectInfo::Type SelectInfo) const
+	{
+		ListView->SetItemSelection(FTedsRowHandle(Row), bSelected, SelectInfo);
+	}
+
+	void STedsTableViewer::ScrollIntoView(RowHandle Row) const
+	{
+		ListView->RequestScrollIntoView(FTedsRowHandle(Row));
+	}
+
+	void STedsTableViewer::ClearSelection() const
+	{
+		ListView->ClearSelection();
+	}
+
 	bool STedsTableViewer::IsItemVisible(TableViewerItemPtr InItem) const
 	{
 		return ListView->IsItemVisible(InItem);
diff --git a/Engine/Plugins/Experimental/EditorDataStorageFeatures/Source/TedsTableViewer/Public/Widgets/STedsTableViewer.h b/Engine/Plugins/Experimental/EditorDataStorageFeatures/Source/TedsTableViewer/Public/Widgets/STedsTableViewer.h
--- a/Engine/Plugins/Experimental/EditorDataStorageFeatures/Source/TedsTableViewer/Public/Widgets/STedsTableViewer.h
+++ b/Engine/Plugins/Experimental/EditorDataStorageFeatures/Source/TedsTableViewer/Public/Widgets/STedsTableViewer.h
@@ -78,6 +78,14 @@ namespace UE::Editor::DataStorage
 
 		// Get the row handle for the widget row the table viewer's contents are stored in
 		TEDSTABLEVIEWER_API RowHandle GetWidgetRowHandle() const;
+		
+		// Select the given row in the table viewer
+		TEDSTABLEVIEWER_API void SetSelection(RowHandle Row, bool bSelected, const ESelectInfo::Type SelectInfo) const;
+
+		// Scroll the given row into view in the table viewer
+		TEDSTABLEVIEWER_API void ScrollIntoView(RowHandle Row) const;
+
+		TEDSTABLEVIEWER_API void ClearSelection() const;
 
 	protected:
 		
diff --git a/Engine/Source/Editor/ContentBrowser/Private/ContentBrowserModule.cpp b/Engine/Source/Editor/ContentBrowser/Private/ContentBrowserModule.cpp
--- a/Engine/Source/Editor/ContentBrowser/Private/ContentBrowserModule.cpp
+++ b/Engine/Source/Editor/ContentBrowser/Private/ContentBrowserModule.cpp
@@ -48,12 +48,12 @@ IContentBrowserSingleton& FContentBrowserModule::Get() const
 	return *ContentBrowserSingleton;
 }
 
-void FContentBrowserModule::SetContentBrowserViewExtender(const TSharedPtr<IContentBrowserViewExtender>& InViewExtender)
+void FContentBrowserModule::SetContentBrowserViewExtender(const FCreateViewExtender& InViewExtender)
 {
 	ContentBrowserViewExtender = InViewExtender;
 }
 
-TSharedPtr<IContentBrowserViewExtender> FContentBrowserModule::GetContentBrowserViewExtender()
+FContentBrowserModule::FCreateViewExtender FContentBrowserModule::GetContentBrowserViewExtender()
 {
 	return ContentBrowserViewExtender;
 }
diff --git a/Engine/Source/Editor/ContentBrowser/Private/SAssetView.cpp b/Engine/Source/Editor/ContentBrowser/Private/SAssetView.cpp
--- a/Engine/Source/Editor/ContentBrowser/Private/SAssetView.cpp
+++ b/Engine/Source/Editor/ContentBrowser/Private/SAssetView.cpp
@@ -1230,15 +1230,19 @@ void SAssetView::Construct( const FArguments& InArgs )
 	if(InArgs._AllowCustomView)
 	{
 		FContentBrowserModule& ContentBrowserModule = FModuleManager::GetModuleChecked<FContentBrowserModule>(TEXT("ContentBrowser"));
-		ViewExtender = ContentBrowserModule.GetContentBrowserViewExtender();
 
-		// Bind the delegates the custom view is responsible for firing
-		if(ViewExtender)
+		if(ContentBrowserModule.GetContentBrowserViewExtender().IsBound())
 		{
-			ViewExtender->OnSelectionChanged().BindSP(this, &SAssetView::AssetSelectionChanged);
-			ViewExtender->OnContextMenuOpened().BindSP(this, &SAssetView::OnGetContextMenuContent);
-			ViewExtender->OnItemScrolledIntoView().BindSP(this, &SAssetView::ItemScrolledIntoView);
-			ViewExtender->OnItemDoubleClicked().BindSP(this, &SAssetView::OnListMouseButtonDoubleClick);
+			ViewExtender = ContentBrowserModule.GetContentBrowserViewExtender().Execute();
+
+			// Bind the delegates the custom view is responsible for firing
+			if(ViewExtender)
+			{
+				ViewExtender->OnSelectionChanged().BindSP(this, &SAssetView::AssetSelectionChanged);
+				ViewExtender->OnContextMenuOpened().BindSP(this, &SAssetView::OnGetContextMenuContent);
+				ViewExtender->OnItemScrolledIntoView().BindSP(this, &SAssetView::ItemScrolledIntoView);
+				ViewExtender->OnItemDoubleClicked().BindSP(this, &SAssetView::OnListMouseButtonDoubleClick);
+			}
 		}
 	}
 	
diff --git a/Engine/Source/Editor/ContentBrowser/Public/ContentBrowserModule.h b/Engine/Source/Editor/ContentBrowser/Public/ContentBrowserModule.h
--- a/Engine/Source/Editor/ContentBrowser/Public/ContentBrowserModule.h
+++ b/Engine/Source/Editor/ContentBrowser/Public/ContentBrowserModule.h
@@ -126,6 +126,8 @@ public:
 	DECLARE_DELEGATE_OneParam( FDefaultSelectedPathsDelegate, TArray<FName>& /*VirtualPaths*/ );
 	/** */
 	DECLARE_DELEGATE_OneParam( FDefaultPathsToExpandDelegate, TArray<FName>& /*VirtualPaths*/ );
+	/** Delegate that creates an instance of the custom view extender */
+	DECLARE_DELEGATE_RetVal(TSharedPtr<IContentBrowserViewExtender>, FCreateViewExtender);
 	
 	/**
 	 * Called right after the plugin DLL has been loaded and the plugin object has been created
@@ -174,11 +176,10 @@ public:
 	/** Delegates to be called to extend list of content browser Plugin Filters*/
 	virtual TArray<FAddPathViewPluginFilters>& GetAddPathViewPluginFilters() { return PathViewPluginFilters; }
 
-	/** Register a custom view that can show up inside SAssetView in the Content Browser */
-	virtual void SetContentBrowserViewExtender(const TSharedPtr<IContentBrowserViewExtender>& InExtender);
-
+	/** Register a delegate that creates a custom view that can show up inside SAssetView in the Content Browser */
+	virtual void SetContentBrowserViewExtender(const FCreateViewExtender& InCreateViewExtender);
 	/** Get the currently registered custom view */
-	virtual TSharedPtr<IContentBrowserViewExtender> GetContentBrowserViewExtender();
+	virtual FCreateViewExtender GetContentBrowserViewExtender();
 
 	/** Delegate accessors */
 	FOnFilterChanged& GetOnFilterChanged() { return OnFilterChanged; } 
@@ -264,5 +265,5 @@ private:
 	FDefaultPathsToExpandDelegate DefaultPathsToExpandDelegate;
 
 	/** Extension used to provide a custom view to the Content Browser */
-	TSharedPtr<IContentBrowserViewExtender> ContentBrowserViewExtender;
+	FCreateViewExtender ContentBrowserViewExtender;
 };
