TedsTableViewer: Fix incorrect construction of FTedsRowHandle

diff --git a/Engine/Plugins/Experimental/EditorDataStorageFeatures/Source/TedsTableViewer/Private/Widgets/STedsTableViewer.cpp b/Engine/Plugins/Experimental/EditorDataStorageFeatures/Source/TedsTableViewer/Private/Widgets/STedsTableViewer.cpp
--- a/Engine/Plugins/Experimental/EditorDataStorageFeatures/Source/TedsTableViewer/Private/Widgets/STedsTableViewer.cpp
+++ b/Engine/Plugins/Experimental/EditorDataStorageFeatures/Source/TedsTableViewer/Private/Widgets/STedsTableViewer.cpp
@@ -168,12 +168,14 @@ namespace UE::Editor::DataStorage
 
 	void STedsTableViewer::SetSelection(RowHandle Row, bool bSelected, const ESelectInfo::Type SelectInfo) const
 	{
-		ListView->SetItemSelection(FTedsRowHandle(Row), bSelected, SelectInfo);
+		FTedsRowHandle TedsRowHandle{ .RowHandle = Row };
+		ListView->SetItemSelection(TedsRowHandle, bSelected, SelectInfo);
 	}
 
 	void STedsTableViewer::ScrollIntoView(RowHandle Row) const
 	{
-		ListView->RequestScrollIntoView(FTedsRowHandle(Row));
+		FTedsRowHandle TedsRowHandle{ .RowHandle = Row };
+		ListView->RequestScrollIntoView(TedsRowHandle);
 	}
 
 	void STedsTableViewer::ClearSelection() const
