Revision Control Processors: Update overlay status when a PackageReference column is added to the actor row. This handles the case where the revision control status is pushed to TEDS before the actor and SCC rows are linked, since it only updates the overlay status from the revision control update when both rows are linked.

diff --git a/Engine/Plugins/Experimental/EditorDataStorageFeatures/Source/TedsRevisionControl/Private/Processors/RevisionControlProcessors.cpp b/Engine/Plugins/Experimental/EditorDataStorageFeatures/Source/TedsRevisionControl/Private/Processors/RevisionControlProcessors.cpp
--- a/Engine/Plugins/Experimental/EditorDataStorageFeatures/Source/TedsRevisionControl/Private/Processors/RevisionControlProcessors.cpp
+++ b/Engine/Plugins/Experimental/EditorDataStorageFeatures/Source/TedsRevisionControl/Private/Processors/RevisionControlProcessors.cpp
@@ -290,6 +290,9 @@ void URevisionControlDataStorageFactory::RegisterQueries(IEditorDataStorageProvi
 				DataStorage.UnregisterQuery(SelectionRemoved);
 				SelectionRemoved = InvalidQueryHandle;
 
+				DataStorage.UnregisterQuery(PackageReferenceAdded);
+				PackageReferenceAdded = InvalidQueryHandle;
+
 				DataStorage.UnregisterQuery(FlushPackageUpdates);
 				FlushPackageUpdates = InvalidQueryHandle;
 
@@ -572,6 +575,37 @@ void URevisionControlDataStorageFactory::RegisterApplyOverlays(IEditorDataStorag
 					.DependsOn()
 						.SubQuery(Subqueries)
 					.Compile());
+
+	/**
+	 * Usually, when a revision control update is requested for an SCC row it adds a new row with FTypedElementPackageUpdateColumn and a reference to
+	 * the actor row and the SCC row to update the overlays. However, if the revision control update happens before the actor row and SCC row have a
+	 * chance to link to each other via the FTypedElementPackageReference column, FTypedElementPackageUpdateColumn cannot be added.
+	 * So we add an observer to track for FTypedElementPackageReference addition to the actor rows and manually execute an overlay update.
+	 */
+	PackageReferenceAdded = DataStorage.RegisterQuery(
+					Select(
+					TEXT("Add overlay on package reference added"),
+					FObserver::OnAdd<FTypedElementPackageReference>(),
+					[this](IQueryContext& ActorQueryContext, RowHandle ObjectRow, const FTypedElementPackageReference& PackageReference,
+						const FTypedElementUObjectColumn& Actor)
+					{
+						ActorQueryContext.RunSubquery(EApplyOverlaysObjectToSCC, PackageReference.Row, CreateSubqueryCallbackBinding(
+							[&ActorQueryContext, &ObjectRow, &Actor](ISubqueryContext& SubQueryContext)
+							{
+								FColor Color = DetermineOverlayColor(SubQueryContext, Actor, ActorQueryContext.HasColumn<FTypedElementSelectionColumn>());
+								if (Color.Bits != 0)
+								{
+									ActorQueryContext.AddColumn<FTypedElementViewportOverlayColorColumn>(ObjectRow, { .OverlayColor = Color });
+								}
+							})
+						);
+					})
+					.Where()
+						.All<FTypedElementActorTag>()
+						.None<FTypedElementViewportOverlayColorColumn>()
+					.DependsOn()
+						.SubQuery(Subqueries)
+					.Compile());
 }
 
 void URevisionControlDataStorageFactory::RegisterRemoveOverlays(IEditorDataStorageProvider& DataStorage)
diff --git a/Engine/Plugins/Experimental/EditorDataStorageFeatures/Source/TedsRevisionControl/Private/Processors/RevisionControlProcessors.h b/Engine/Plugins/Experimental/EditorDataStorageFeatures/Source/TedsRevisionControl/Private/Processors/RevisionControlProcessors.h
--- a/Engine/Plugins/Experimental/EditorDataStorageFeatures/Source/TedsRevisionControl/Private/Processors/RevisionControlProcessors.h
+++ b/Engine/Plugins/Experimental/EditorDataStorageFeatures/Source/TedsRevisionControl/Private/Processors/RevisionControlProcessors.h
@@ -44,6 +44,8 @@ private:
 	UE::Editor::DataStorage::QueryHandle FlushPackageUpdates = UE::Editor::DataStorage::InvalidQueryHandle;
 	UE::Editor::DataStorage::QueryHandle SelectionAdded = UE::Editor::DataStorage::InvalidQueryHandle;
 	UE::Editor::DataStorage::QueryHandle SelectionRemoved = UE::Editor::DataStorage::InvalidQueryHandle;
+	UE::Editor::DataStorage::QueryHandle PackageReferenceAdded = UE::Editor::DataStorage::InvalidQueryHandle;
+
 
 	// General queries to gather SCC rows
 	TMap<const UScriptStruct*, UE::Editor::DataStorage::QueryHandle> GeneralQueriesMap;
