TedsSettings : add an inactive settings section table and allow rows to migrate between active and inactive tables.

When a settings object is registered with the ISettingsModule it will become active.  If there is an existing inactive row that matches the ContainerName + CategoryName + SectionName ( index hash ) then that row will be removed / replaced by an active row.

When a settings object is unregistered with the ISettingsModule it will become inactive.  External modules can also add an inactive settings section using the FindOrAddSettingsSection method.

When rows transition between the inactive and active tables, any column that is unknown to TedsSettings will be copied over to the new row.  This allows for external modules to add columns to rows and have those columns preserved as settings are registered and unregistered with the ISettingsModule.

diff --git a/Engine/Plugins/Experimental/EditorDataStorageFeatures/Source/TedsSettings/Private/TedsSettingsEditorSubsystem.cpp b/Engine/Plugins/Experimental/EditorDataStorageFeatures/Source/TedsSettings/Private/TedsSettingsEditorSubsystem.cpp
--- a/Engine/Plugins/Experimental/EditorDataStorageFeatures/Source/TedsSettings/Private/TedsSettingsEditorSubsystem.cpp
+++ b/Engine/Plugins/Experimental/EditorDataStorageFeatures/Source/TedsSettings/Private/TedsSettingsEditorSubsystem.cpp
@@ -31,6 +31,11 @@ UTedsSettingsEditorSubsystem::FOnEnabledChanged& UTedsSettingsEditorSubsystem::O
 	return EnabledChangedDelegate;
 }
 
+UE::Editor::DataStorage::RowHandle UTedsSettingsEditorSubsystem::FindOrAddSettingsSection(const FName& ContainerName, const FName& CategoryName, const FName& SectionName)
+{
+	return SettingsManager->FindOrAddSettingsSection(ContainerName, CategoryName, SectionName);
+}
+
 void UTedsSettingsEditorSubsystem::Initialize(FSubsystemCollectionBase& Collection)
 {
 	Super::Initialize(Collection);
diff --git a/Engine/Plugins/Experimental/EditorDataStorageFeatures/Source/TedsSettings/Private/TedsSettingsManager.cpp b/Engine/Plugins/Experimental/EditorDataStorageFeatures/Source/TedsSettings/Private/TedsSettingsManager.cpp
--- a/Engine/Plugins/Experimental/EditorDataStorageFeatures/Source/TedsSettings/Private/TedsSettingsManager.cpp
+++ b/Engine/Plugins/Experimental/EditorDataStorageFeatures/Source/TedsSettings/Private/TedsSettingsManager.cpp
@@ -5,6 +5,7 @@
 #include "Elements/Columns/TypedElementCompatibilityColumns.h"
 #include "Elements/Columns/TypedElementMiscColumns.h"
 #include "Elements/Columns/TypedElementSlateWidgetColumns.h"
+#include "Elements/Columns/TypedElementTypeInfoColumns.h"
 #include "Elements/Common/EditorDataStorageFeatures.h"
 #include "Elements/Framework/TypedElementIndexHasher.h"
 #include "Elements/Framework/TypedElementQueryBuilder.h"
@@ -31,13 +32,31 @@ namespace UE::Editor::Settings::Private
 		Hash = CityHash128to64({ Hash, UE::Editor::DataStorage::GenerateIndexHash(CategoryName) });
 		return CityHash128to64({ Hash, UE::Editor::DataStorage::GenerateIndexHash(SectionName) });
 	}
+
+	static bool IsUnknownColumn(const UScriptStruct* ColumnType)
+	{
+		check(ColumnType);
+		return ColumnType != FSettingsSectionTag::StaticStruct()
+			&& ColumnType != FSettingsInactiveSectionTag::StaticStruct()
+			&& ColumnType != FSettingsContainerReferenceColumn::StaticStruct()
+			&& ColumnType != FSettingsCategoryReferenceColumn::StaticStruct()
+			&& ColumnType != FNameColumn::StaticStruct()
+			&& ColumnType != FDisplayNameColumn::StaticStruct()
+			&& ColumnType != FDescriptionColumn::StaticStruct()
+			&& ColumnType != FTypedElementUObjectColumn::StaticStruct()
+			&& ColumnType != FTypedElementUObjectIdColumn::StaticStruct()
+			&& ColumnType != FTypedElementClassTypeInfoColumn::StaticStruct()
+			&& ColumnType != FTypedElementClassDefaultObjectTag::StaticStruct();
+	}
 }
 
 FTedsSettingsManager::FTedsSettingsManager()
 	: bIsInitialized{ false }
-	, SelectAllSettingsQuery{ UE::Editor::DataStorage::InvalidQueryHandle }
+	, SelectAllActiveSettingsQuery{ UE::Editor::DataStorage::InvalidQueryHandle }
+	, SelectAllInactiveSettingsQuery{ UE::Editor::DataStorage::InvalidQueryHandle }
 	, SettingsContainerTable{ UE::Editor::DataStorage::InvalidTableHandle }
 	, SettingsCategoryTable{ UE::Editor::DataStorage::InvalidTableHandle }
+	, SettingsInactiveSectionTable{ UE::Editor::DataStorage::InvalidTableHandle }
 {
 }
 
@@ -53,7 +72,7 @@ void FTedsSettingsManager::Initialize()
 
 				RegisterTables(*DataStorage);
 				RegisterQueries(*DataStorage);
-				RegisterSettings();
+				RegisterActiveSettings();
 			};
 
 		if (AreEditorDataStorageFeaturesEnabled())
@@ -81,7 +100,8 @@ void FTedsSettingsManager::Shutdown()
 			IEditorDataStorageProvider* DataStorage = GetMutableDataStorageFeature<IEditorDataStorageProvider>(StorageFeatureName);
 			check(DataStorage);
 
-			UnregisterSettings();
+			UnregisterActiveSettings();
+			UnregisterInactiveSettings();
 			UnregisterQueries(*DataStorage);
 		}
 
@@ -89,6 +109,43 @@ void FTedsSettingsManager::Shutdown()
 	}
 }
 
+UE::Editor::DataStorage::RowHandle FTedsSettingsManager::FindOrAddSettingsSection(const FName& ContainerName, const FName& CategoryName, const FName& SectionName)
+{
+	using namespace UE::Editor::DataStorage;
+	using namespace UE::Editor::Settings::Private;
+
+	if (bIsInitialized)
+	{
+		IEditorDataStorageProvider* DataStorage = GetMutableDataStorageFeature<IEditorDataStorageProvider>(StorageFeatureName);
+		if (ensure(DataStorage))
+		{
+			// Ensure that the SettingsInactiveSectionTable is registered since it is possible for OnDataStorageInterfacesSet to be called
+			// for a delegate bound in another module before it has been called for our bound delegate.
+			RegisterTables(*DataStorage);
+
+			uint64 SectionIndexHash = GenerateIndexHash(ContainerName, CategoryName, SectionName);
+
+			RowHandle SectionRow = DataStorage->FindIndexedRow(SectionIndexHash);
+
+			if (SectionRow == InvalidRowHandle)
+			{
+				SectionRow = DataStorage->AddRow(SettingsInactiveSectionTable);
+
+				DataStorage->AddColumn<FSettingsInactiveSectionTag>(SectionRow);
+				DataStorage->AddColumn<FNameColumn>(SectionRow, { .Name = SectionName });
+				DataStorage->AddColumn<FSettingsContainerReferenceColumn>(SectionRow, { .ContainerName = ContainerName, .ContainerRow = InvalidRowHandle });
+				DataStorage->AddColumn<FSettingsCategoryReferenceColumn>(SectionRow, { .CategoryName = CategoryName, .CategoryRow = InvalidRowHandle });
+
+				DataStorage->IndexRow(SectionIndexHash, SectionRow);
+			}
+
+			return SectionRow;
+		}
+	}
+
+	return InvalidRowHandle;
+}
+
 void FTedsSettingsManager::RegisterTables(IEditorDataStorageProvider& DataStorage)
 {
 	if (SettingsContainerTable == UE::Editor::DataStorage::InvalidTableHandle)
@@ -104,32 +161,52 @@ void FTedsSettingsManager::RegisterTables(IEditorDataStorageProvider& DataStorag
 			TTypedElementColumnTypeList<FSettingsContainerReferenceColumn, FNameColumn, FDisplayNameColumn, FDescriptionColumn, FSettingsCategoryTag>(),
 			FName(TEXT("Editor_SettingsCategoryTable")));
 	}
+
+	if (SettingsInactiveSectionTable == UE::Editor::DataStorage::InvalidTableHandle)
+	{
+		SettingsInactiveSectionTable = DataStorage.RegisterTable(
+			TTypedElementColumnTypeList<FSettingsContainerReferenceColumn, FSettingsCategoryReferenceColumn, FNameColumn, FSettingsInactiveSectionTag>(),
+			FName(TEXT("Editor_SettingsInactiveSectionTable")));
+	}
 }
 
 void FTedsSettingsManager::RegisterQueries(IEditorDataStorageProvider& DataStorage)
 {
 	using namespace UE::Editor::DataStorage::Queries;
 
-	if (SelectAllSettingsQuery == UE::Editor::DataStorage::InvalidQueryHandle)
+	if (SelectAllActiveSettingsQuery == UE::Editor::DataStorage::InvalidQueryHandle)
 	{
-		SelectAllSettingsQuery = DataStorage.RegisterQuery(
+		SelectAllActiveSettingsQuery = DataStorage.RegisterQuery(
 			Select()
 				.ReadOnly<FSettingsContainerReferenceColumn, FSettingsCategoryReferenceColumn, FNameColumn>()
 			.Where()
 				.All<FSettingsSectionTag>()
 			.Compile());
 	}
+
+	if (SelectAllInactiveSettingsQuery == UE::Editor::DataStorage::InvalidQueryHandle)
+	{
+		SelectAllInactiveSettingsQuery = DataStorage.RegisterQuery(
+			Select()
+				.ReadOnly<FSettingsContainerReferenceColumn, FSettingsCategoryReferenceColumn, FNameColumn>()
+			.Where()
+				.All<FSettingsInactiveSectionTag>()
+			.Compile());
+	}
 }
 
 void FTedsSettingsManager::UnregisterQueries(IEditorDataStorageProvider& DataStorage)
 {
-	DataStorage.UnregisterQuery(SelectAllSettingsQuery);
-	SelectAllSettingsQuery = UE::Editor::DataStorage::InvalidQueryHandle;
+	DataStorage.UnregisterQuery(SelectAllActiveSettingsQuery);
+	SelectAllActiveSettingsQuery = UE::Editor::DataStorage::InvalidQueryHandle;
+
+	DataStorage.UnregisterQuery(SelectAllInactiveSettingsQuery);
+	SelectAllInactiveSettingsQuery = UE::Editor::DataStorage::InvalidQueryHandle;
 }
 
-void FTedsSettingsManager::RegisterSettings()
+void FTedsSettingsManager::RegisterActiveSettings()
 {
-	TRACE_CPUPROFILER_EVENT_SCOPE(TedsSettingsManager.RegisterSettings);
+	TRACE_CPUPROFILER_EVENT_SCOPE(TedsSettingsManager.RegisterActiveSettings);
 
 	ISettingsModule* SettingsModule = FModuleManager::GetModulePtr<ISettingsModule>("Settings");
 	check(SettingsModule);
@@ -157,7 +234,7 @@ void FTedsSettingsManager::RegisterSettingsContainer(const FName& ContainerName)
 	IEditorDataStorageProvider* DataStorage = GetMutableDataStorageFeature<IEditorDataStorageProvider>(StorageFeatureName);
 	check(DataStorage);
 
-	UE_LOG(LogTedsSettings, Log, TEXT("Register Settings Container : '%s'"), *ContainerName.ToString());
+	UE_LOG(LogTedsSettings, Verbose, TEXT("Register Settings Container : '%s'"), *ContainerName.ToString());
 
 	ISettingsContainerPtr ContainerPtr = SettingsModule->GetContainer(ContainerName);
 
@@ -186,7 +263,7 @@ void FTedsSettingsManager::RegisterSettingsContainer(const FName& ContainerName)
 	// OnCategoryModified is called at the same time as OnSectionRemoved so we only bind to OnCategoryModified for add / update / remove
 	ContainerPtr->OnCategoryModified().AddSPLambda(this, [this, ContainerPtr, ContainerRow](const FName& ModifiedCategoryName)
 		{
-			UE_LOG(LogTedsSettings, Log, TEXT("Settings Category modified : '%s->%s'"), *ContainerPtr->GetName().ToString(), *ModifiedCategoryName.ToString());
+			UE_LOG(LogTedsSettings, Verbose, TEXT("Settings Category modified : '%s->%s'"), *ContainerPtr->GetName().ToString(), *ModifiedCategoryName.ToString());
 
 			ISettingsCategoryPtr CategoryPtr = ContainerPtr->GetCategory(ModifiedCategoryName);
 
@@ -194,12 +271,12 @@ void FTedsSettingsManager::RegisterSettingsContainer(const FName& ContainerName)
 		});
 }
 
-void FTedsSettingsManager::UnregisterSettings()
+void FTedsSettingsManager::UnregisterActiveSettings()
 {
 	using namespace UE::Editor::DataStorage;
 	using namespace UE::Editor::Settings::Private;
 
-	TRACE_CPUPROFILER_EVENT_SCOPE(TedsSettingsManager.UnregisterSettings);
+	TRACE_CPUPROFILER_EVENT_SCOPE(TedsSettingsManager.UnregisterActiveSettings);
 
 	ISettingsModule* SettingsModule = FModuleManager::GetModulePtr<ISettingsModule>("Settings");
 	check(SettingsModule);
@@ -217,7 +294,7 @@ void FTedsSettingsManager::UnregisterSettings()
 
 	for (FName ContainerName : ContainerNames)
 	{
-		UE_LOG(LogTedsSettings, Log, TEXT("Unregister Settings Container : '%s'"), *ContainerName.ToString());
+		UE_LOG(LogTedsSettings, Verbose, TEXT("Unregister Settings Container : '%s'"), *ContainerName.ToString());
 
 		ISettingsContainerPtr ContainerPtr = SettingsModule->GetContainer(ContainerName);
 
@@ -228,7 +305,7 @@ void FTedsSettingsManager::UnregisterSettings()
 
 		for (ISettingsCategoryPtr CategoryPtr : Categories)
 		{
-			UE_LOG(LogTedsSettings, Log, TEXT("Unregister Settings Category : '%s'"), *CategoryPtr->GetName().ToString());
+			UE_LOG(LogTedsSettings, Verbose, TEXT("Unregister Settings Category : '%s'"), *CategoryPtr->GetName().ToString());
 
 			TArray<ISettingsSectionPtr> Sections;
 			const bool bIgnoreVisibility = true;
@@ -242,7 +319,7 @@ void FTedsSettingsManager::UnregisterSettings()
 
 					DataStorage->RemoveIndex(GenerateIndexHash(ContainerName, CategoryPtr->GetName(), SectionPtr->GetName()));
 
-					UE_LOG(LogTedsSettings, Log, TEXT("Removed Settings Section : '%s'"), *SectionPtr->GetName().ToString());
+					UE_LOG(LogTedsSettings, Verbose, TEXT("Removed Settings Section : '%s'"), *SectionPtr->GetName().ToString());
 				}
 			}
 
@@ -265,6 +342,30 @@ void FTedsSettingsManager::UnregisterSettings()
 	}
 }
 
+void FTedsSettingsManager::UnregisterInactiveSettings()
+{
+	TRACE_CPUPROFILER_EVENT_SCOPE(TedsSettingsManager.UnregisterInactiveSettings);
+
+	using namespace UE::Editor::DataStorage::Queries;
+	using namespace UE::Editor::Settings::Private;
+
+	IEditorDataStorageProvider* DataStorage = GetMutableDataStorageFeature<IEditorDataStorageProvider>(StorageFeatureName);
+	check(DataStorage);
+	
+	TSet<RowHandle> RowHandles;
+
+	DataStorage->RunQuery(SelectAllInactiveSettingsQuery, CreateDirectQueryCallbackBinding(
+		[&RowHandles](IDirectQueryContext& Context, const RowHandle*)
+		{
+			RowHandles.Append(Context.GetRowHandles());
+		}));
+
+	for (const RowHandle& Row : RowHandles)
+	{
+		DataStorage->RemoveRow(Row);
+	}
+}
+
 void FTedsSettingsManager::UpdateSettingsCategory(TSharedPtr<ISettingsCategory> SettingsCategory, UE::Editor::DataStorage::RowHandle ContainerRow, const bool bQueryExistingRows)
 {
 	using namespace UE::Editor::DataStorage;
@@ -281,7 +382,7 @@ void FTedsSettingsManager::UpdateSettingsCategory(TSharedPtr<ISettingsCategory>
 	const FName& ContainerName = DataStorage->GetColumn<FNameColumn>(ContainerRow)->Name;
 	const FName& CategoryName = SettingsCategory->GetName();
 
-	UE_LOG(LogTedsSettings, Log, TEXT("Update Settings Category: '%s->%s'"), *ContainerName.ToString(), *CategoryName.ToString());
+	UE_LOG(LogTedsSettings, Verbose, TEXT("Update Settings Category: '%s->%s'"), *ContainerName.ToString(), *CategoryName.ToString());
 
 	uint64 CategoryIndexHash = GenerateIndexHash(SettingsCategory.Get());
 
@@ -302,12 +403,12 @@ void FTedsSettingsManager::UpdateSettingsCategory(TSharedPtr<ISettingsCategory>
 	TArray<RowHandle> OldRowHandles;
 	TArray<FName> OldSectionNames;
 
-	// Gather all existing rows for the given { ContainerName, CategoryName } pair.
+	// Gather all existing active rows for the given { ContainerName, CategoryName } pair.
 	if (bQueryExistingRows)
 	{
 		using namespace UE::Editor::DataStorage::Queries;
 
-		DataStorage->RunQuery(SelectAllSettingsQuery, CreateDirectQueryCallbackBinding(
+		DataStorage->RunQuery(SelectAllActiveSettingsQuery, CreateDirectQueryCallbackBinding(
 			[&OldRowHandles, &OldSectionNames, &ContainerName, &CategoryName](
 				IDirectQueryContext& Context,
 				const FSettingsContainerReferenceColumn* ContainerColumns,
@@ -346,16 +447,24 @@ void FTedsSettingsManager::UpdateSettingsCategory(TSharedPtr<ISettingsCategory>
 			NewSectionNames.Emplace(SectionName);
 
 			uint64 SectionIndexHash = GenerateIndexHash(ContainerName, CategoryName, SectionName);
-			
+
+			TArray<void*> ColumnsToCopy;
+			TArray<const UScriptStruct*> ColumnTypesToCopy;
+
 			RowHandle OldSectionRow = DataStorage->FindIndexedRow(SectionIndexHash);
 			if (OldSectionRow != InvalidRowHandle)
 			{
 				UE_LOG(LogTedsSettings, Verbose, TEXT("Settings Section : '%s' is already in data storage"), *SectionName.ToString());
-
-				// Remove the row, the SettingsObjectPtr may have changed so we need to re-add the row with the new object
-				DataStorage->RemoveRow(OldSectionRow);
-
-				UE_LOG(LogTedsSettings, Log, TEXT("Removed Settings Section : '%s'"), *SectionName.ToString());
+				
+				// Capture unknown columns and copy them to the new row before we delete the old row.
+				DataStorage->ListColumns(OldSectionRow, [&ColumnsToCopy, &ColumnTypesToCopy](void* Column, const UScriptStruct& ColumnType)
+					{
+						if (IsUnknownColumn(&ColumnType))
+						{
+							ColumnsToCopy.Add(Column);
+							ColumnTypesToCopy.Add(&ColumnType);
+						}
+					});
 			}
 
 			RowHandle NewSectionRow = DataStorageCompatibility->AddCompatibleObject(SettingsObjectPtr);
@@ -367,13 +476,24 @@ void FTedsSettingsManager::UpdateSettingsCategory(TSharedPtr<ISettingsCategory>
 			DataStorage->AddColumn<FDisplayNameColumn>(NewSectionRow, { .DisplayName = SectionPtr->GetDisplayName() });
 			DataStorage->AddColumn<FDescriptionColumn>(NewSectionRow, { .Description = SectionPtr->GetDescription() });
 			
-			DataStorage->IndexRow(SectionIndexHash, NewSectionRow);
+			if (OldSectionRow != InvalidRowHandle && OldSectionRow != NewSectionRow)
+			{
+				AddColumns(*DataStorage, NewSectionRow, ColumnsToCopy, ColumnTypesToCopy);
+
+				// Remove the old row after we copy the unknown columns to the new row.
+				DataStorage->RemoveRow(OldSectionRow);
+
+				UE_LOG(LogTedsSettings, Verbose, TEXT("Removed Settings Section : '%s'"), *SectionName.ToString());
+			}
+			
+			UE_CLOG(OldSectionRow != NewSectionRow, LogTedsSettings, Verbose, TEXT("Added Settings Section : '%s'"), *SectionName.ToString());
+			UE_CLOG(OldSectionRow == NewSectionRow, LogTedsSettings, Verbose, TEXT("Updated Settings Section : '%s'"), *SectionName.ToString());
 
-			UE_LOG(LogTedsSettings, Log, TEXT("Added Settings Section : '%s'"), *SectionName.ToString());
+			DataStorage->IndexRow(SectionIndexHash, NewSectionRow);			
 		}
 	}
 
-	// Iterate the old sections and remove rows not in the new sections list.
+	// Iterate the old active sections and inactivate rows not in the new active sections list.
 	for (int32 RowIndex = 0; RowIndex < OldSectionNames.Num(); ++RowIndex)
 	{
 		const FName& OldSectionName = OldSectionNames[RowIndex];
@@ -383,13 +503,55 @@ void FTedsSettingsManager::UpdateSettingsCategory(TSharedPtr<ISettingsCategory>
 			continue;
 		}
 
-		RowHandle OldRowHandle = OldRowHandles[RowIndex];
-		check(OldRowHandle != InvalidRowHandle);
+		RowHandle OldSectionRow = OldRowHandles[RowIndex];
+		check(OldSectionRow != InvalidRowHandle);
+
+		// Add an inactive row for the settings section.
+		RowHandle NewSectionRow = DataStorage->AddRow(SettingsInactiveSectionTable);
+
+		DataStorage->AddColumn<FSettingsInactiveSectionTag>(NewSectionRow);
+		DataStorage->AddColumn<FNameColumn>(NewSectionRow, { .Name = OldSectionName });
+		DataStorage->AddColumn<FSettingsContainerReferenceColumn>(NewSectionRow, { .ContainerName = ContainerName, .ContainerRow = InvalidRowHandle });
+		DataStorage->AddColumn<FSettingsCategoryReferenceColumn>(NewSectionRow, { .CategoryName = CategoryName, .CategoryRow = InvalidRowHandle });
+
+		// Capture unknown columns and copy them to the new row before we delete the old row.
+		TArray<void*> ColumnsToCopy;
+		TArray<const UScriptStruct*> ColumnTypesToCopy;		
+		
+		DataStorage->ListColumns(OldSectionRow, [&ColumnsToCopy, &ColumnTypesToCopy](void* Column, const UScriptStruct& ColumnType)
+			{
+				if (IsUnknownColumn(&ColumnType))
+				{
+					ColumnsToCopy.Add(Column);
+					ColumnTypesToCopy.Add(&ColumnType);
+				}
+			});
 
-		DataStorage->RemoveRow(OldRowHandle);
+		AddColumns(*DataStorage, NewSectionRow, ColumnsToCopy, ColumnTypesToCopy);
 
-		DataStorage->RemoveIndex(GenerateIndexHash(ContainerName, CategoryName, OldSectionName));
+		DataStorage->RemoveRow(OldSectionRow);
 
-		UE_LOG(LogTedsSettings, Log, TEXT("Removed Settings Section : '%s'"), *OldSectionName.ToString());
+		DataStorage->IndexRow(GenerateIndexHash(ContainerName, CategoryName, OldSectionName), NewSectionRow);
+
+		UE_LOG(LogTedsSettings, Verbose, TEXT("Inactivated Settings Section : '%s'"), *OldSectionName.ToString());
+	}
+}
+
+void FTedsSettingsManager::AddColumns(IEditorDataStorageProvider& DataStorage, UE::Editor::DataStorage::RowHandle Row, const TArray<void*>& Columns, const TArray<const UScriptStruct*>& ColumnTypes)
+{
+	for (int32 ColumnIndex = 0; ColumnIndex < Columns.Num(); ++ColumnIndex)
+	{
+		void* SourceColumn = Columns[ColumnIndex];
+		const UScriptStruct* SourceColumnType = ColumnTypes[ColumnIndex];
+
+		DataStorage.AddColumnData(Row, SourceColumnType,
+			[SourceColumn](void* TargetColumn, const UScriptStruct& ColumnType)
+			{
+				ColumnType.CopyScriptStruct(TargetColumn, SourceColumn);
+			},
+			[](const UScriptStruct& ColumnType, void* Destination, void* Source)
+			{
+				ColumnType.CopyScriptStruct(Destination, Source);
+			});
 	}
 }
diff --git a/Engine/Plugins/Experimental/EditorDataStorageFeatures/Source/TedsSettings/Private/TedsSettingsManager.h b/Engine/Plugins/Experimental/EditorDataStorageFeatures/Source/TedsSettings/Private/TedsSettingsManager.h
--- a/Engine/Plugins/Experimental/EditorDataStorageFeatures/Source/TedsSettings/Private/TedsSettingsManager.h
+++ b/Engine/Plugins/Experimental/EditorDataStorageFeatures/Source/TedsSettings/Private/TedsSettingsManager.h
@@ -21,6 +21,8 @@ public:
 	void Initialize();
 	void Shutdown();
 
+	UE::Editor::DataStorage::RowHandle FindOrAddSettingsSection(const FName& ContainerName, const FName& CategoryName, const FName& SectionName);
+
 private:
 
 	void RegisterTables(IEditorDataStorageProvider& DataStorage);
@@ -28,16 +30,21 @@ private:
 	void RegisterQueries(IEditorDataStorageProvider& DataStorage);
 	void UnregisterQueries(IEditorDataStorageProvider& DataStorage);
 
-	void RegisterSettings();
-	void UnregisterSettings();
+	void RegisterActiveSettings();
+	void UnregisterActiveSettings();
+	void UnregisterInactiveSettings();
 
 	void RegisterSettingsContainer(const FName& ContainerName);
 
 	void UpdateSettingsCategory(TSharedPtr<ISettingsCategory> SettingsCategory, UE::Editor::DataStorage::RowHandle ContainerRow, const bool bQueryExistingRows = true);
 
+	void AddColumns(IEditorDataStorageProvider& DataStorage, UE::Editor::DataStorage::RowHandle Row, const TArray<void*>& Columns, const TArray<const UScriptStruct*>& ColumnTypes);
+
 	bool bIsInitialized;
-	UE::Editor::DataStorage::QueryHandle SelectAllSettingsQuery;
+	UE::Editor::DataStorage::QueryHandle SelectAllActiveSettingsQuery;
+	UE::Editor::DataStorage::QueryHandle SelectAllInactiveSettingsQuery;
 	UE::Editor::DataStorage::TableHandle SettingsContainerTable;
 	UE::Editor::DataStorage::TableHandle SettingsCategoryTable;
+	UE::Editor::DataStorage::TableHandle SettingsInactiveSectionTable;
 
 };
diff --git a/Engine/Plugins/Experimental/EditorDataStorageFeatures/Source/TedsSettings/Public/TedsSettingsColumns.h b/Engine/Plugins/Experimental/EditorDataStorageFeatures/Source/TedsSettings/Public/TedsSettingsColumns.h
--- a/Engine/Plugins/Experimental/EditorDataStorageFeatures/Source/TedsSettings/Public/TedsSettingsColumns.h
+++ b/Engine/Plugins/Experimental/EditorDataStorageFeatures/Source/TedsSettings/Public/TedsSettingsColumns.h
@@ -48,3 +48,9 @@ struct FSettingsSectionTag final : public FEditorDataStorageTag
 {
 	GENERATED_BODY()
 };
+
+USTRUCT(meta = (DisplayName = "Inactive Settings Section"))
+struct FSettingsInactiveSectionTag final : public FEditorDataStorageTag
+{
+	GENERATED_BODY()
+};
diff --git a/Engine/Plugins/Experimental/EditorDataStorageFeatures/Source/TedsSettings/Public/TedsSettingsEditorSubsystem.h b/Engine/Plugins/Experimental/EditorDataStorageFeatures/Source/TedsSettings/Public/TedsSettingsEditorSubsystem.h
--- a/Engine/Plugins/Experimental/EditorDataStorageFeatures/Source/TedsSettings/Public/TedsSettingsEditorSubsystem.h
+++ b/Engine/Plugins/Experimental/EditorDataStorageFeatures/Source/TedsSettings/Public/TedsSettingsEditorSubsystem.h
@@ -3,6 +3,7 @@
 #pragma once
 
 #include "EditorSubsystem.h"
+#include "Elements/Interfaces/TypedElementDataStorageInterface.h"
 
 #include "TedsSettingsEditorSubsystem.generated.h"
 
@@ -22,6 +23,9 @@ public:
 	DECLARE_MULTICAST_DELEGATE(FOnEnabledChanged)
 	FOnEnabledChanged& OnEnabledChanged();
 
+	// Finds an existing row (may be active or inactive) or adds a new inactive settings section row if no existing row is found.
+	UE::Editor::DataStorage::RowHandle FindOrAddSettingsSection(const FName& ContainerName, const FName& CategoryName, const FName& SectionName);
+
 protected: // USubsystem
 
 	virtual void Initialize(FSubsystemCollectionBase& Collection) override;
diff --git a/Engine/Plugins/Experimental/EditorDataStorageFeatures/Source/TedsSettings/Tests/TedsSettingsTests.cpp b/Engine/Plugins/Experimental/EditorDataStorageFeatures/Source/TedsSettings/Tests/TedsSettingsTests.cpp
--- a/Engine/Plugins/Experimental/EditorDataStorageFeatures/Source/TedsSettings/Tests/TedsSettingsTests.cpp
+++ b/Engine/Plugins/Experimental/EditorDataStorageFeatures/Source/TedsSettings/Tests/TedsSettingsTests.cpp
@@ -86,7 +86,7 @@ namespace UE::Editor::Settings::Tests
 				CountAllSettingsQuery = DataStorage->RegisterQuery(
 					Count()
 					.Where()
-						.All<FSettingsSectionTag>()
+						.Any<FSettingsSectionTag, FSettingsInactiveSectionTag>()
 					.Compile());
 			}
 
@@ -144,7 +144,7 @@ namespace UE::Editor::Settings::Tests
 
 		Describe("UnregisterSettings", [this]()
 		{
-			LatentIt("Should remove a row from editor data storage", [this](const FDoneDelegate& Done)
+			LatentIt("Should inactivate a row in editor data storage", [this](const FDoneDelegate& Done)
 			{
 				const FName& ContainerName = FName(TEXT("TestContainer"));
 				const FName& CategoryName = FName(TEXT("TestCategory"));
@@ -174,12 +174,24 @@ namespace UE::Editor::Settings::Tests
 					SettingsModule->UnregisterSettings(ContainerName, CategoryName, SectionName);
 
 					uint32 AfterUnregisterRowCount = CountSettingsRowsInDataStorage();
-					TestEqual(TEXT("RowCount"), AfterUnregisterRowCount, BeforeRowCount);
+					TestEqual(TEXT("RowCount"), AfterUnregisterRowCount, BeforeRowCount + 1);
 
 					TestFalse(TEXT("IsRowAssigned"), DataStorage->IsRowAssigned(RowHandle));
 
 					DataStorage::RowHandle TestInvalidRowHandle = DataStorageCompatibility->FindRowWithCompatibleObject(TestSettingsObject);
 					TestEqual(TEXT("InvalidRowHandle"), TestInvalidRowHandle, DataStorage::InvalidRowHandle);
+
+					UTedsSettingsEditorSubsystem* SettingsEditorSubsystem = GEditor->GetEditorSubsystem<UTedsSettingsEditorSubsystem>();
+
+					DataStorage::RowHandle InactiveRowHandle = SettingsEditorSubsystem->FindOrAddSettingsSection(ContainerName, CategoryName, SectionName);
+					TestNotEqual(TEXT("InactiveRowHandle"), InactiveRowHandle, DataStorage::InvalidRowHandle);
+
+					TestRowHandles.Push(InactiveRowHandle);
+
+					uint32 AfterFindOrAddRowCount = CountSettingsRowsInDataStorage();
+					TestEqual(TEXT("RowCount"), AfterFindOrAddRowCount, BeforeRowCount + 1);
+
+					TestTrue(TEXT("FSettingsInactiveSectionTag"), DataStorage->HasColumns<FSettingsInactiveSectionTag>(InactiveRowHandle));
 				});
 			});
 		});
@@ -229,6 +241,96 @@ namespace UE::Editor::Settings::Tests
 				});
 			});
 		});
+
+		Describe("FindOrAddSettingsSection with no existing active rows", [this]()
+		{
+			LatentIt("Should result in only a single inactive row in editor data storage", [this](const FDoneDelegate& Done)
+			{
+				UTedsSettingsEditorSubsystem* SettingsEditorSubsystem = GEditor->GetEditorSubsystem<UTedsSettingsEditorSubsystem>();
+
+				const FName& ContainerName = FName(TEXT("TestContainer"));
+				const FName& CategoryName = FName(TEXT("TestCategory"));
+				const FName& SectionName = FName(TEXT("TestSection"));
+
+				DataStorage::RowHandle RowHandle = SettingsEditorSubsystem->FindOrAddSettingsSection(ContainerName, CategoryName, SectionName);
+				TestNotEqual(TEXT("RowHandle"), RowHandle, DataStorage::InvalidRowHandle);
+
+				if (RowHandle == DataStorage::InvalidRowHandle)
+				{
+					Done.Execute();
+					return;
+				}
+
+				TestRowHandles.Push(RowHandle);
+
+				AwaitRowHandleThenVerify(RowHandle, Done, [this, RowHandle, ContainerName, CategoryName, SectionName]()
+				{
+					uint32 AfterRowCount = CountSettingsRowsInDataStorage();
+
+					TestEqual(TEXT("RowCount"), AfterRowCount, BeforeRowCount + 1);
+
+					TestEqual(TEXT("ContainerName"), DataStorage->GetColumn<FSettingsContainerReferenceColumn>(RowHandle)->ContainerName, ContainerName);
+					TestEqual(TEXT("CategoryName"), DataStorage->GetColumn<FSettingsCategoryReferenceColumn>(RowHandle)->CategoryName, CategoryName);
+					TestEqual(TEXT("SectionName"), DataStorage->GetColumn<FNameColumn>(RowHandle)->Name, SectionName);
+
+					TestTrue(TEXT("FSettingsInactiveSectionTag"), DataStorage->HasColumns<FSettingsInactiveSectionTag>(RowHandle));
+				});
+			});
+		});
+
+		Describe("RegisterSettings with an existing inactive row for same container/category/section", [this]()
+		{
+			LatentIt("Should result in only a single active row in editor data storage", [this](const FDoneDelegate& Done)
+			{
+				UTedsSettingsEditorSubsystem* SettingsEditorSubsystem = GEditor->GetEditorSubsystem<UTedsSettingsEditorSubsystem>();
+
+				const FName& ContainerName = FName(TEXT("TestContainer"));
+				const FName& CategoryName = FName(TEXT("TestCategory"));
+				const FName& SectionName = FName(TEXT("TestSection"));
+
+				DataStorage::RowHandle RowHandle = SettingsEditorSubsystem->FindOrAddSettingsSection(ContainerName, CategoryName, SectionName);
+				TestNotEqual(TEXT("RowHandle"), RowHandle, DataStorage::InvalidRowHandle);
+
+				if (RowHandle == DataStorage::InvalidRowHandle)
+				{
+					Done.Execute();
+					return;
+				}
+
+				TestRowHandles.Push(RowHandle);
+
+				TestTrue(TEXT("FSettingsInactiveSectionTag"), DataStorage->HasColumns<FSettingsInactiveSectionTag>(RowHandle));
+
+				TObjectPtr<UTestSettings> TestSettingsObject = NewObject<UTestSettings>();
+
+				SettingsModule->RegisterSettings(ContainerName, CategoryName, SectionName, FText(), FText(), TestSettingsObject);
+
+				RowHandle = DataStorageCompatibility->FindRowWithCompatibleObject(TestSettingsObject);
+				TestNotEqual(TEXT("RowHandle"), RowHandle, DataStorage::InvalidRowHandle);
+
+				if (RowHandle == DataStorage::InvalidRowHandle)
+				{
+					Done.Execute();
+					return;
+				}
+
+				TestRowHandles.Push(RowHandle);
+
+				AwaitRowHandleThenVerify(RowHandle, Done, [this, RowHandle, ContainerName, CategoryName, SectionName]()
+				{
+					uint32 AfterRowCount = CountSettingsRowsInDataStorage();
+
+					TestEqual(TEXT("RowCount"), AfterRowCount, BeforeRowCount + 1);
+
+					TestEqual(TEXT("ContainerName"), DataStorage->GetColumn<FSettingsContainerReferenceColumn>(RowHandle)->ContainerName, ContainerName);
+					TestEqual(TEXT("CategoryName"), DataStorage->GetColumn<FSettingsCategoryReferenceColumn>(RowHandle)->CategoryName, CategoryName);
+					TestEqual(TEXT("SectionName"), DataStorage->GetColumn<FNameColumn>(RowHandle)->Name, SectionName);
+
+					TestFalse(TEXT("FSettingsInactiveSectionTag"), DataStorage->HasColumns<FSettingsInactiveSectionTag>(RowHandle));
+					TestTrue(TEXT("FSettingsSectionTag"), DataStorage->HasColumns<FSettingsSectionTag>(RowHandle));
+				});
+			});
+		});
 	}
 } // namespace UE::Editor::Settings::Tests
 
