Revision Control Overlays:
Remove the package update table in favor of an activatable query that sweeps through all actors in the Scene to update their overlays
Activate said query when the SCC status of an object changes

TEDS:
Fix a crash with activatable queries where you could have duplicates in the pending activatable queries, but unregistering a query only removed one instance

PrimitiveComponent:
Add functionality to persist OverlayColor in UPrimitiveComponent using the ComponentInstanceDataCache

diff --git a/Engine/Plugins/Developer/PerforceSourceControl/Source/PerforceSourceControl/Private/PerforceSourceControlOperations.cpp b/Engine/Plugins/Developer/PerforceSourceControl/Source/PerforceSourceControl/Private/PerforceSourceControlOperations.cpp
--- a/Engine/Plugins/Developer/PerforceSourceControl/Source/PerforceSourceControl/Private/PerforceSourceControlOperations.cpp
+++ b/Engine/Plugins/Developer/PerforceSourceControl/Source/PerforceSourceControl/Private/PerforceSourceControlOperations.cpp
@@ -2133,15 +2133,12 @@ bool FPerforceUpdateStatusWorker::UpdateStates() const
 			DataStorage->RemoveColumn(Row, Column);
 		}
 		
-        if (FTypedElementPackageReference* BackReference = DataStorage->GetColumn<FTypedElementPackageReference>(Row))
-        {
-        	if (RowHandle ObjectRow = BackReference->Row; ObjectRow != InvalidRowHandle)
-        	{
-        		static TableHandle Table = DataStorage->FindTable(FName("Editor_PackageUpdateTable"));
-        		RowHandle UpdateRow = DataStorage->AddRow(Table);
-        		DataStorage->AddColumn(UpdateRow, FTypedElementPackageUpdateColumn{ .ObjectRow = ObjectRow, .PackageRow = Row });
-        	}
-        }
+		// While the backreference via FTypedElementPackageReference isn't reliable because you can have multiple actors referencing a single SCC
+		// row, it at least lets us filter out rows that don't have any actors
+		if (DataStorage->HasColumns<FTypedElementPackageReference>(Row))
+		{
+			DataStorage->ActivateQueries(TEXT("UpdateSCCForActors"));
+		}
 	};
 
 	bool bUpdated = false;
diff --git a/Engine/Plugins/Experimental/EditorDataStorage/Source/TedsCore/Private/Queries/TypedElementExtendedQueryStore.cpp b/Engine/Plugins/Experimental/EditorDataStorage/Source/TedsCore/Private/Queries/TypedElementExtendedQueryStore.cpp
--- a/Engine/Plugins/Experimental/EditorDataStorage/Source/TedsCore/Private/Queries/TypedElementExtendedQueryStore.cpp
+++ b/Engine/Plugins/Experimental/EditorDataStorage/Source/TedsCore/Private/Queries/TypedElementExtendedQueryStore.cpp
@@ -314,7 +314,7 @@ namespace UE::Editor::DataStorage
 				checkf(!QueryData.Description.Callback.ActivationName.IsNone(),
 					TEXT("Attempting to enable the query '%s' which isn't activatable."), *QueryData.Description.Callback.Name.ToString());
 #endif
-				PendingActivatables.Add(Query);
+				PendingActivatables.AddUnique(Query);
 			}
 		}
 	}
diff --git a/Engine/Plugins/Experimental/EditorDataStorageFeatures/Source/TedsRevisionControl/Private/Processors/RevisionControlProcessors.cpp b/Engine/Plugins/Experimental/EditorDataStorageFeatures/Source/TedsRevisionControl/Private/Processors/RevisionControlProcessors.cpp
--- a/Engine/Plugins/Experimental/EditorDataStorageFeatures/Source/TedsRevisionControl/Private/Processors/RevisionControlProcessors.cpp
+++ b/Engine/Plugins/Experimental/EditorDataStorageFeatures/Source/TedsRevisionControl/Private/Processors/RevisionControlProcessors.cpp
@@ -44,8 +44,6 @@ namespace UE::Editor::RevisionControl::Private
 			Factory->UpdateOverlayColors(DataStorage);
 		}
 	}
-	
-	
 
 	static bool gEnableOverlays = false;
 	TAutoConsoleVariable<bool> CVarEnableOverlays(
@@ -238,10 +236,6 @@ void URevisionControlDataStorageFactory::RegisterTables(IEditorDataStorageProvid
 			FTypedElementPackagePathColumn, FTypedElementPackageLoadedPathColumn,
 			FSCCRevisionIdColumn, FSCCExternalRevisionIdColumn>(),
 		FName("Editor_RevisionControlTable"));
-
-	DataStorage.RegisterTable(
-		TTypedElementColumnTypeList<FTypedElementPackageUpdateColumn>(),
-		FName("Editor_PackageUpdateTable"));
 }
 
 void URevisionControlDataStorageFactory::RegisterQueries(IEditorDataStorageProvider& DataStorage)
@@ -276,11 +270,6 @@ void URevisionControlDataStorageFactory::RegisterQueries(IEditorDataStorageProvi
 			}
 			else
 			{
-				DataStorage.UnregisterQuery(ApplyNewOverlays);
-				ApplyNewOverlays = InvalidQueryHandle;
-				DataStorage.UnregisterQuery(ChangeOverlay);
-				ChangeOverlay = InvalidQueryHandle;
-
 				DataStorage.UnregisterQuery(ApplyOverlaysObjectToSCC);
 				ApplyOverlaysObjectToSCC = InvalidQueryHandle;
 
@@ -293,8 +282,11 @@ void URevisionControlDataStorageFactory::RegisterQueries(IEditorDataStorageProvi
 				DataStorage.UnregisterQuery(PackageReferenceAdded);
 				PackageReferenceAdded = InvalidQueryHandle;
 
-				DataStorage.UnregisterQuery(FlushPackageUpdates);
-				FlushPackageUpdates = InvalidQueryHandle;
+				DataStorage.UnregisterQuery(UpdateSCCForActors);
+				UpdateSCCForActors = InvalidQueryHandle;
+				
+				DataStorage.UnregisterQuery(UpdateOverlays);
+				UpdateOverlays = InvalidQueryHandle;
 
 				RegisterRemoveOverlays(DataStorage);
 			}
@@ -320,65 +312,12 @@ void URevisionControlDataStorageFactory::RegisterQueries(IEditorDataStorageProvi
 
 void URevisionControlDataStorageFactory::UpdateOverlaysForSCCState(IEditorDataStorageProvider* DataStorage, const UScriptStruct* Column) const
 {
-	using namespace UE::Editor::DataStorage::Queries;
-
-	// The column provided must be present in our map
-	const QueryHandle* QueryForColumn = GeneralQueriesMap.Find(Column);
-	if(!QueryForColumn)
-	{
-		return;
-	}
-
-	// Collect all SCC and Object rows to request a package update
-	TArray<FTypedElementPackageUpdateColumn> PackagesToUpdate;
-
-	DirectQueryCallback RowCollector = CreateDirectQueryCallbackBinding(
-		[&PackagesToUpdate, DataStorage] (const IDirectQueryContext& Context, const FTypedElementPackageReference* PackageReference)
-		{
-			uint32 RowCount = Context.GetRowCount();
-			TConstArrayView<RowHandle> RowHandles = Context.GetRowHandles();
-
-			for(uint32 RowIndex = 0; RowIndex < RowCount; ++RowIndex)
-			{
-				if(DataStorage->IsRowAvailable(PackageReference[RowIndex].Row))
-				{
-					// We are querying the SCC rows so the package row is the one present in the callback and the object
-					// row is the one present in FTypedElementPackageReference
- 					PackagesToUpdate.Add(FTypedElementPackageUpdateColumn{.ObjectRow = PackageReference[RowIndex].Row, .PackageRow = RowHandles[RowIndex]});
-				}
-			}
-		});
-
-	DataStorage->RunQuery(*QueryForColumn, RowCollector);
-	RequestPackageUpdates(DataStorage, PackagesToUpdate);
+	DataStorage->ActivateQueries(TEXT("UpdateSCCForActors"));
 }
 
 void URevisionControlDataStorageFactory::UpdateOverlayColors(IEditorDataStorageProvider* DataStorage) const
 {
-	using namespace UE::Editor::DataStorage::Queries;
-
-	TArray<FTypedElementPackageUpdateColumn> PackagesToUpdate;
-
-	DirectQueryCallback RowCollector = CreateDirectQueryCallbackBinding(
-		[&PackagesToUpdate] (const IDirectQueryContext& Context, const FTypedElementPackageReference* PackageReference)
-		{
-			uint32 RowCount = Context.GetRowCount();
-			TConstArrayView<RowHandle> RowHandles = Context.GetRowHandles();
-
-			for(uint32 RowIndex = 0; RowIndex < RowCount; ++RowIndex)
-			{
-				if(PackageReference[RowIndex].Row != InvalidRowHandle)
-				{
-					// The viewport color column is on the actor rows, so the row handle in the callback is the Object Row and the row handle
-					// in FTypedElementPackageReference is the SCC row
-					PackagesToUpdate.Add(FTypedElementPackageUpdateColumn{.ObjectRow = RowHandles[RowIndex], .PackageRow = PackageReference[RowIndex].Row});
-				}
-			}
-		});
-
-	DataStorage->RunQuery(FetchOverlayColors, RowCollector);
-	RequestPackageUpdates(DataStorage, PackagesToUpdate);
-
+	DataStorage->ActivateQueries(TEXT("UpdateOverlayForActors"));
 }
 
 void URevisionControlDataStorageFactory::RegisterFetchUpdates(IEditorDataStorageProvider& DataStorage)
@@ -422,108 +361,66 @@ void URevisionControlDataStorageFactory::RegisterApplyOverlays(IEditorDataStorag
 				.ReadOnly<FSCCStatusColumn>(EOptional::Yes)
 			.Compile());
 	}
-
-	if (ApplyNewOverlays == InvalidQueryHandle)
-	{
-		ApplyNewOverlays = DataStorage.RegisterQuery(
-			Select()
-			.ReadOnly<FTypedElementUObjectColumn, FTypedElementPackageReference>()
-			.Where()
-				.All<FTypedElementActorTag>()
-				.None<FTypedElementViewportOverlayColorColumn>()
-			.Compile());
-	}
-
-	if (ChangeOverlay == InvalidQueryHandle)
+	
+	if (UpdateSCCForActors == InvalidQueryHandle)
 	{
-		ChangeOverlay = DataStorage.RegisterQuery(
-			Select()
-			.ReadOnly<FTypedElementUObjectColumn, FTypedElementPackageReference, FTypedElementViewportOverlayColorColumn>()
-			.Where()
-				.All<FTypedElementActorTag>()
-			.Compile());
-	}
+		check(ApplyOverlaysObjectToSCC != InvalidQueryHandle);
 
-	enum EFlushPackageUpdatesSubqueries
-	{
-		EApplyOverlaysObjectToSCC,
-		EApplyNewOverlays,
-		EChangeOverlay,
-			
-		Num
-	};
-	TArray<RowHandle> Subqueries;
-	Subqueries.AddUninitialized(EFlushPackageUpdatesSubqueries::Num);
-	Subqueries[EApplyOverlaysObjectToSCC] = ApplyOverlaysObjectToSCC;
-	Subqueries[EApplyNewOverlays] = ApplyNewOverlays;
-	Subqueries[EChangeOverlay] = ChangeOverlay;
-
-	if (FlushPackageUpdates == InvalidQueryHandle)
-	{
-		check(ApplyOverlaysObjectToSCC != InvalidQueryHandle && ApplyNewOverlays != InvalidQueryHandle && ChangeOverlay!= InvalidQueryHandle);
-		
-		FlushPackageUpdates = DataStorage.RegisterQuery(
+		// Query:
+		// For all actors having a package reference:
+		//	Determine if a color should be applied based on SCC status tags
+		//		If so, either add the OverlayColorColumn to the actor row or remove and re-add it to update it
+		//		If not, remove the OverlayColorColumn from the actor row
+		UpdateSCCForActors = DataStorage.RegisterQuery(
 			Select(
-				TEXT("Consume collected package updates"),
+				TEXT("Update overlay for all actors"),
 				FProcessor(EQueryTickPhase::PrePhysics, DataStorage.GetQueryTickGroupName(EQueryTickGroups::Update))
-					.SetExecutionMode(EExecutionMode::GameThread),
-				[](IQueryContext& Context, RowHandle Row, const FTypedElementPackageUpdateColumn& Update)
+					.SetExecutionMode(EExecutionMode::GameThread)
+					.MakeActivatable(TEXT("UpdateSCCForActors")),
+				[](IQueryContext& ActorQueryContext, RowHandle ObjectRow, const FTypedElementUObjectColumn& Actor, const FTypedElementPackageReference& PackageReference)
 				{
-					// Query:
-					// For all actors without an overlay color column AND having a package reference:
-					//   Determine if a color should be applied based on SCC status tags
-					//   If so, add the OverlayColorColumn to the actor row
-					Context.RunSubquery(EApplyNewOverlays, Update.ObjectRow, CreateSubqueryCallbackBinding(
-						[&Context](ISubqueryContext& ActorQueryContext, RowHandle ObjectRow, const FTypedElementUObjectColumn& Actor, const FTypedElementPackageReference& PackageReference)
+					// Run a subquery on the SCC row to determine the overlay color
+					ActorQueryContext.RunSubquery(0, PackageReference.Row, CreateSubqueryCallbackBinding(
+						[&ActorQueryContext, &ObjectRow, &Actor](const ISubqueryContext& SubQueryContext)
 						{
-							Context.RunSubquery(EApplyOverlaysObjectToSCC, PackageReference.Row, CreateSubqueryCallbackBinding(
-								[&ActorQueryContext, &ObjectRow, &Actor](ISubqueryContext& SubQueryContext)
-								{
-									FColor Color = DetermineOverlayColor(SubQueryContext, Actor, ActorQueryContext.HasColumn<FTypedElementSelectionColumn>());
-									if (Color.Bits != 0)
-									{
-										ActorQueryContext.AddColumn<FTypedElementViewportOverlayColorColumn>(ObjectRow, { .OverlayColor = Color });
-									}
-								})
-							);
-						}
-					));
-
-					// Query:
-					// For all actors WITH an overlay color column AND having a package reference:
-					//   Re-check the color that should be applied based on SCC status tags
-					//   If the color has changed, remove and re-add the OverlayColorColumn to the actor row
-					//
-					// Note: Remove and re-add will trigger observer in TypedElementActorViewportProcessors to SetOverlayColor on the primitive components
-					Context.RunSubquery(EChangeOverlay, Update.ObjectRow, CreateSubqueryCallbackBinding(
-						[&Context](ISubqueryContext& ActorQueryContext, RowHandle ObjectRow, const FTypedElementUObjectColumn& Actor, const FTypedElementPackageReference& PackageReference, const FTypedElementViewportOverlayColorColumn& OverlayColorColumn)
-						{
-							Context.RunSubquery(EApplyOverlaysObjectToSCC, PackageReference.Row, CreateSubqueryCallbackBinding(
-								[&ActorQueryContext, &ObjectRow, &Actor, &OverlayColorColumn](ISubqueryContext& SubQueryContext)
+							FColor Color = DetermineOverlayColor(SubQueryContext, Actor, ActorQueryContext.HasColumn<FTypedElementSelectionColumn>());
+
+							// If there is no overlay color, simply remove the column
+							if (Color.Bits == 0)
+							{
+								ActorQueryContext.RemoveColumns<FTypedElementViewportOverlayColorColumn>(ObjectRow);
+								return;
+							}
+
+							// If we already have the color column, only update it if the color has changed
+							if(const FTypedElementViewportOverlayColorColumn* OverlayColorColumn = ActorQueryContext.GetColumn<FTypedElementViewportOverlayColorColumn>())
+							{
+								if (Color != OverlayColorColumn->OverlayColor)
 								{
-									FColor Color = DetermineOverlayColor(SubQueryContext, Actor, ActorQueryContext.HasColumn<FTypedElementSelectionColumn>());
-									if (Color.Bits == 0)
-									{
-										ActorQueryContext.RemoveColumns<FTypedElementViewportOverlayColorColumn>(ObjectRow);
-									}
-									else if (Color != OverlayColorColumn.OverlayColor)
-									{
-										// Remove and re-add to trigger the observer
-										ActorQueryContext.RemoveColumns<FTypedElementViewportOverlayColorColumn>(ObjectRow);
-										ActorQueryContext.AddColumn<FTypedElementViewportOverlayColorColumn>(ObjectRow, { .OverlayColor = Color });
-									}
-								})
-							);
-						}
-					));
-					Context.RemoveRow(Row);
+									// Remove and re-add to trigger the observer
+									ActorQueryContext.RemoveColumns<FTypedElementViewportOverlayColorColumn>(ObjectRow);
+									ActorQueryContext.AddColumn<FTypedElementViewportOverlayColorColumn>(ObjectRow, { .OverlayColor = Color });
+								}
+							}
+							// If we don't have the color column, add it
+							else
+							{
+								ActorQueryContext.AddColumn<FTypedElementViewportOverlayColorColumn>(ObjectRow, { .OverlayColor = Color });
+							}
+						})
+					);
 				})
+			.ReadOnly<FTypedElementViewportOverlayColorColumn>(EOptional::Yes)
+			.Where()
+				.All<FTypedElementActorTag>()
 			.DependsOn()
-				.SubQuery(Subqueries)
+				.SubQuery(ApplyOverlaysObjectToSCC)
 			.Compile());
 	}
 
-	SelectionAdded = DataStorage.RegisterQuery(
+	if (SelectionAdded == InvalidQueryHandle)
+	{
+		SelectionAdded = DataStorage.RegisterQuery(
 					Select(
 					TEXT("Update Overlay on Selection"),
 					FObserver::OnAdd<FTypedElementSelectionColumn>(),
@@ -542,8 +439,11 @@ void URevisionControlDataStorageFactory::RegisterApplyOverlays(IEditorDataStorag
 					.Where()
 						.All<FTypedElementActorTag, FTypedElementViewportOverlayColorColumn, FTypedElementPackageReference>()
 					.Compile());
-	
-	SelectionRemoved = DataStorage.RegisterQuery(
+	}
+
+	if (SelectionRemoved == InvalidQueryHandle)
+	{
+		SelectionRemoved = DataStorage.RegisterQuery(
 					Select(
 					TEXT("Update Overlay on Deselection"),
 					FObserver::OnRemove<FTypedElementSelectionColumn>(),
@@ -555,8 +455,8 @@ void URevisionControlDataStorageFactory::RegisterApplyOverlays(IEditorDataStorag
 						if(SelectionColumn.SelectionSet.IsNone())
 						{
 							// When an item is deselected, add the viewport overlay color column to it if applicable
-							ActorQueryContext.RunSubquery(EApplyOverlaysObjectToSCC, PackageReference.Row, CreateSubqueryCallbackBinding(
-								[&ActorQueryContext, &ObjectRow, &Actor](ISubqueryContext& SubQueryContext)
+							ActorQueryContext.RunSubquery(0, PackageReference.Row, CreateSubqueryCallbackBinding(
+								[&ActorQueryContext, &ObjectRow, &Actor](const ISubqueryContext& SubQueryContext)
 								{
 									// We manually set the selection as false because the item is being deslected, but the row still has the selection column
 									// when the observer is fired
@@ -573,39 +473,76 @@ void URevisionControlDataStorageFactory::RegisterApplyOverlays(IEditorDataStorag
 					.Where()
 						.All<FTypedElementActorTag>()
 					.DependsOn()
-						.SubQuery(Subqueries)
+						.SubQuery(ApplyOverlaysObjectToSCC)
 					.Compile());
+	}
 
-	/**
+	if (PackageReferenceAdded == InvalidQueryHandle)
+	{
+		/**
 	 * Usually, when a revision control update is requested for an SCC row it adds a new row with FTypedElementPackageUpdateColumn and a reference to
 	 * the actor row and the SCC row to update the overlays. However, if the revision control update happens before the actor row and SCC row have a
 	 * chance to link to each other via the FTypedElementPackageReference column, FTypedElementPackageUpdateColumn cannot be added.
 	 * So we add an observer to track for FTypedElementPackageReference addition to the actor rows and manually execute an overlay update.
 	 */
-	PackageReferenceAdded = DataStorage.RegisterQuery(
-					Select(
-					TEXT("Add overlay on package reference added"),
-					FObserver::OnAdd<FTypedElementPackageReference>(),
-					[this](IQueryContext& ActorQueryContext, RowHandle ObjectRow, const FTypedElementPackageReference& PackageReference,
-						const FTypedElementUObjectColumn& Actor)
+		PackageReferenceAdded = DataStorage.RegisterQuery(
+						Select(
+						TEXT("Add overlay on package reference added"),
+						FObserver::OnAdd<FTypedElementPackageReference>(),
+						[this](IQueryContext& ActorQueryContext, RowHandle ObjectRow, const FTypedElementPackageReference& PackageReference,
+							const FTypedElementUObjectColumn& Actor)
+						{
+							ActorQueryContext.RunSubquery(0, PackageReference.Row, CreateSubqueryCallbackBinding(
+								[&ActorQueryContext, &ObjectRow, &Actor](const ISubqueryContext& SubQueryContext)
+								{
+									FColor Color = DetermineOverlayColor(SubQueryContext, Actor, ActorQueryContext.HasColumn<FTypedElementSelectionColumn>());
+									if (Color.Bits != 0)
+									{
+										ActorQueryContext.AddColumn<FTypedElementViewportOverlayColorColumn>(ObjectRow, { .OverlayColor = Color });
+									}
+								})
+							);
+						})
+						.Where()
+							.All<FTypedElementActorTag>()
+							.None<FTypedElementViewportOverlayColorColumn>()
+						.DependsOn()
+							.SubQuery(ApplyOverlaysObjectToSCC)
+						.Compile());
+	}
+
+	// Query to update the color for all rows that currently have the overlay color column
+	if (UpdateOverlays == InvalidQueryHandle)
+	{
+		UpdateOverlays = DataStorage.RegisterQuery(
+			Select(
+				TEXT("Update overlay color for actors with overlays"),
+				FProcessor(EQueryTickPhase::PrePhysics, DataStorage.GetQueryTickGroupName(EQueryTickGroups::Update))
+					.SetExecutionMode(EExecutionMode::GameThread)
+					.MakeActivatable(TEXT("UpdateOverlayForActors")),
+					[](IQueryContext& ActorQueryContext, RowHandle ObjectRow, const FTypedElementUObjectColumn& Actor,
+						const FTypedElementPackageReference& PackageReference, FTypedElementViewportOverlayColorColumn& OverlayColorColumn)
 					{
-						ActorQueryContext.RunSubquery(EApplyOverlaysObjectToSCC, PackageReference.Row, CreateSubqueryCallbackBinding(
-							[&ActorQueryContext, &ObjectRow, &Actor](ISubqueryContext& SubQueryContext)
+						ActorQueryContext.RunSubquery(0, PackageReference.Row, CreateSubqueryCallbackBinding(
+							[&ActorQueryContext, &ObjectRow, &Actor, &OverlayColorColumn](const ISubqueryContext& SubQueryContext)
 							{
 								FColor Color = DetermineOverlayColor(SubQueryContext, Actor, ActorQueryContext.HasColumn<FTypedElementSelectionColumn>());
-								if (Color.Bits != 0)
+							
+								if (Color != OverlayColorColumn.OverlayColor)
 								{
+									// Remove and re-add to trigger the observer
+									ActorQueryContext.RemoveColumns<FTypedElementViewportOverlayColorColumn>(ObjectRow);
 									ActorQueryContext.AddColumn<FTypedElementViewportOverlayColorColumn>(ObjectRow, { .OverlayColor = Color });
 								}
 							})
 						);
 					})
-					.Where()
-						.All<FTypedElementActorTag>()
-						.None<FTypedElementViewportOverlayColorColumn>()
-					.DependsOn()
-						.SubQuery(Subqueries)
-					.Compile());
+			.Where()
+				.All<FTypedElementActorTag>()
+			.DependsOn()
+				.SubQuery(ApplyOverlaysObjectToSCC)
+			.Compile());
+	}
 }
 
 void URevisionControlDataStorageFactory::RegisterRemoveOverlays(IEditorDataStorageProvider& DataStorage)
@@ -638,41 +575,8 @@ void URevisionControlDataStorageFactory::RegisterRemoveOverlays(IEditorDataStora
 void URevisionControlDataStorageFactory::RegisterGeneralQueries(IEditorDataStorageProvider& DataStorage)
 {
 	using namespace UE::Editor::DataStorage::Queries;
+	using namespace UE::Editor::RevisionControl::Private;
 
-	// SCC Columns we want to update the overlay state for
-	const TArray<const UScriptStruct*> Columns{ FSCCNotCurrentTag::StaticStruct(), FSCCExternallyLockedColumn::StaticStruct(),
-		FSCCStatusColumn::StaticStruct(), FSCCLockedTag::StaticStruct() };
-
-	for(const UScriptStruct* Column : Columns)
-	{
-		QueryHandle Query = DataStorage.RegisterQuery(
-			Select()
-				.ReadOnly<FTypedElementPackageReference>()
-			.Where()
-				.All(Column)
-			.Compile());
-
-		GeneralQueriesMap.Add(Column, Query);
-	}
-
-	if(FetchOverlayColors == InvalidQueryHandle)
-	{
-		FetchOverlayColors = DataStorage.RegisterQuery(
-		Select()
-			.ReadOnly<FTypedElementPackageReference>()
-		.Where()
-			.All<FTypedElementViewportOverlayColorColumn, FTypedElementActorTag>()
-		.Compile());
-	}
+	
 }
 
-void URevisionControlDataStorageFactory::RequestPackageUpdates(IEditorDataStorageProvider* DataStorage, TArray<FTypedElementPackageUpdateColumn>& Packages) const
-{
-	UE::Editor::DataStorage::TableHandle Table = DataStorage->FindTable(FName("Editor_PackageUpdateTable"));
-
-	for(FTypedElementPackageUpdateColumn& PackageUpdate : Packages)
-	{
-		UE::Editor::DataStorage::RowHandle UpdateRow = DataStorage->AddRow(Table);
-		DataStorage->AddColumn(UpdateRow, MoveTemp(PackageUpdate));
-	}
-}
diff --git a/Engine/Plugins/Experimental/EditorDataStorageFeatures/Source/TedsRevisionControl/Private/Processors/RevisionControlProcessors.h b/Engine/Plugins/Experimental/EditorDataStorageFeatures/Source/TedsRevisionControl/Private/Processors/RevisionControlProcessors.h
--- a/Engine/Plugins/Experimental/EditorDataStorageFeatures/Source/TedsRevisionControl/Private/Processors/RevisionControlProcessors.h
+++ b/Engine/Plugins/Experimental/EditorDataStorageFeatures/Source/TedsRevisionControl/Private/Processors/RevisionControlProcessors.h
@@ -34,23 +34,15 @@ private:
 	void RegisterRemoveOverlays(IEditorDataStorageProvider& DataStorage);
 	void RegisterGeneralQueries(IEditorDataStorageProvider& DataStorage);
 
-	void RequestPackageUpdates(IEditorDataStorageProvider* DataStorage, TArray<FTypedElementPackageUpdateColumn>& Packages) const;
-	
 	UE::Editor::DataStorage::QueryHandle FetchUpdates = UE::Editor::DataStorage::InvalidQueryHandle;
-	UE::Editor::DataStorage::QueryHandle ApplyNewOverlays = UE::Editor::DataStorage::InvalidQueryHandle;
-	UE::Editor::DataStorage::QueryHandle ChangeOverlay = UE::Editor::DataStorage::InvalidQueryHandle;
 	UE::Editor::DataStorage::QueryHandle ApplyOverlaysObjectToSCC = UE::Editor::DataStorage::InvalidQueryHandle;
 	UE::Editor::DataStorage::QueryHandle RemoveOverlays = UE::Editor::DataStorage::InvalidQueryHandle;
-	UE::Editor::DataStorage::QueryHandle FlushPackageUpdates = UE::Editor::DataStorage::InvalidQueryHandle;
+	UE::Editor::DataStorage::QueryHandle UpdateSCCForActors = UE::Editor::DataStorage::InvalidQueryHandle;
 	UE::Editor::DataStorage::QueryHandle SelectionAdded = UE::Editor::DataStorage::InvalidQueryHandle;
 	UE::Editor::DataStorage::QueryHandle SelectionRemoved = UE::Editor::DataStorage::InvalidQueryHandle;
 	UE::Editor::DataStorage::QueryHandle PackageReferenceAdded = UE::Editor::DataStorage::InvalidQueryHandle;
 
-
-	// General queries to gather SCC rows
-	TMap<const UScriptStruct*, UE::Editor::DataStorage::QueryHandle> GeneralQueriesMap;
-
 	// Query to fetch all rows with overlays
-	UE::Editor::DataStorage::QueryHandle FetchOverlayColors = UE::Editor::DataStorage::InvalidQueryHandle;
+	UE::Editor::DataStorage::QueryHandle UpdateOverlays = UE::Editor::DataStorage::InvalidQueryHandle;
 
 };
diff --git a/Engine/Source/Runtime/Engine/Classes/Components/PrimitiveComponent.h b/Engine/Source/Runtime/Engine/Classes/Components/PrimitiveComponent.h
--- a/Engine/Source/Runtime/Engine/Classes/Components/PrimitiveComponent.h
+++ b/Engine/Source/Runtime/Engine/Classes/Components/PrimitiveComponent.h
@@ -3110,6 +3110,11 @@ private:
 
 	UPROPERTY()
 	TObjectPtr<UPrimitiveComponent> LODParent = nullptr;
+
+#if WITH_EDITOR
+	FColor OverlayColor;
+#endif
+	
 };
 
 
diff --git a/Engine/Source/Runtime/Engine/Private/Components/PrimitiveComponent.cpp b/Engine/Source/Runtime/Engine/Private/Components/PrimitiveComponent.cpp
--- a/Engine/Source/Runtime/Engine/Private/Components/PrimitiveComponent.cpp
+++ b/Engine/Source/Runtime/Engine/Private/Components/PrimitiveComponent.cpp
@@ -731,8 +731,21 @@ FPrimitiveComponentInstanceData::FPrimitiveComponentInstanceData(const UPrimitiv
 	, VisibilityId(SourceComponent->VisibilityId)
 	, LODParent(SourceComponent->GetLODParentPrimitive())
 {
-	const_cast<UPrimitiveComponent*>(SourceComponent)->ConditionalUpdateComponentToWorld(); // sadness
+	UPrimitiveComponent* PrimitiveComponent = const_cast<UPrimitiveComponent*>(SourceComponent);
+	PrimitiveComponent->ConditionalUpdateComponentToWorld(); // sadness
 	ComponentTransform = SourceComponent->GetComponentTransform();
+
+#if WITH_EDITOR
+	// Only persist the overlay color if the component wants editor effects
+	if (PrimitiveComponent->bWantsEditorEffects)
+	{
+		OverlayColor = PrimitiveComponent->OverlayColor;
+	}
+	else
+	{
+		OverlayColor = FColor(ForceInitToZero);
+	}
+#endif
 }
 
 void FPrimitiveComponentInstanceData::ApplyToComponent(UActorComponent* Component, const ECacheApplyPhase CacheApplyPhase)
@@ -758,11 +771,24 @@ void FPrimitiveComponentInstanceData::ApplyToComponent(UActorComponent* Componen
 		PrimitiveComponent->ResetCustomPrimitiveData();
 		Component->MarkRenderStateDirty();
 	}
+	
+#if WITH_EDITOR
+	if (OverlayColor != FColor(ForceInitToZero))
+	{
+		PrimitiveComponent->SetOverlayColor(OverlayColor);
+	}
+#endif
 }
 
 bool FPrimitiveComponentInstanceData::ContainsData() const
 {
-	return (Super::ContainsData() || LODParent || (VisibilityId != INDEX_NONE));
+	bool bContainsData = (Super::ContainsData() || LODParent || (VisibilityId != INDEX_NONE));
+	
+#if WITH_EDITOR
+	bContainsData |= (OverlayColor != FColor(ForceInitToZero));
+#endif
+
+	return bContainsData;
 }
 
 void FPrimitiveComponentInstanceData::AddReferencedObjects(FReferenceCollector& Collector)
@@ -5033,6 +5059,13 @@ void UPrimitiveComponent::AssignSceneProxy(FPrimitiveSceneProxy* InSceneProxy)
 		{
 			SceneData.OwnerLastRenderTimePtr->NumAlwaysVisibleComponents.fetch_add(1, std::memory_order_relaxed);
 		}
+
+#if WITH_EDITOR
+		if (bWantsEditorEffects)
+		{
+			SetOverlayColor(OverlayColor);
+		}
+#endif
 	}
 }
 
diff --git a/Engine/Source/Runtime/TypedElementFramework/Public/Elements/Columns/TypedElementPackageColumns.h b/Engine/Source/Runtime/TypedElementFramework/Public/Elements/Columns/TypedElementPackageColumns.h
--- a/Engine/Source/Runtime/TypedElementFramework/Public/Elements/Columns/TypedElementPackageColumns.h
+++ b/Engine/Source/Runtime/TypedElementFramework/Public/Elements/Columns/TypedElementPackageColumns.h
@@ -46,18 +46,6 @@ struct FTypedElementPackageUpdatedTag final : public FEditorDataStorageTag
 	GENERATED_BODY()
 };
 
-/**
- * Column that stores relevant rows when a package is updated
- */
-USTRUCT()
-struct FTypedElementPackageUpdateColumn final : public FEditorDataStorageColumn
-{
-	GENERATED_BODY()
-	
-	UE::Editor::DataStorage::RowHandle ObjectRow;
-	UE::Editor::DataStorage::RowHandle PackageRow;
-};
-
 /**
  * Column that stores the path of a package.
  */
