TedsSettings : Register general purpose widgets for FSettingsContainerReferenceColumn, FSettingsCategoryReferenceColumn, and FSettingsSectionTag.

These are used in general purpose cells such as the TEDS Debugger, for example.

diff --git a/Engine/Plugins/Experimental/EditorDataStorageFeatures/Source/TedsSettings/Private/TedsSettingsFactory.cpp b/Engine/Plugins/Experimental/EditorDataStorageFeatures/Source/TedsSettings/Private/TedsSettingsFactory.cpp
new file mode 100644
--- /dev/null
+++ b/Engine/Plugins/Experimental/EditorDataStorageFeatures/Source/TedsSettings/Private/TedsSettingsFactory.cpp
@@ -0,0 +1,15 @@
+// Copyright Epic Games, Inc. All Rights Reserved.
+
+#include "TedsSettingsFactory.h"
+
+#include "TedsSettingsColumns.h"
+#include "TedsSettingsWidgets.h"
+
+void UTedsSettingsFactory::RegisterWidgetConstructors(IEditorDataStorageProvider& DataStorage, IEditorDataStorageUiProvider& DataStorageUi) const
+{
+	using namespace UE::Editor::DataStorage::Queries;
+
+	DataStorageUi.RegisterWidgetFactory<FSettingsContainerReferenceWidgetConstructor>(FName(TEXT("General.Cell")), TColumn<FSettingsContainerReferenceColumn>());
+	DataStorageUi.RegisterWidgetFactory<FSettingsCategoryReferenceWidgetConstructor>(FName(TEXT("General.Cell")), TColumn<FSettingsCategoryReferenceColumn>());
+	DataStorageUi.RegisterWidgetFactory<FSettingsSectionWidgetConstructor>(FName(TEXT("General.Cell")), TColumn<FSettingsSectionTag>());
+}
diff --git a/Engine/Plugins/Experimental/EditorDataStorageFeatures/Source/TedsSettings/Private/TedsSettingsFactory.h b/Engine/Plugins/Experimental/EditorDataStorageFeatures/Source/TedsSettings/Private/TedsSettingsFactory.h
new file mode 100644
--- /dev/null
+++ b/Engine/Plugins/Experimental/EditorDataStorageFeatures/Source/TedsSettings/Private/TedsSettingsFactory.h
@@ -0,0 +1,21 @@
+// Copyright Epic Games, Inc. All Rights Reserved.
+
+#pragma once
+
+#include "Elements/Interfaces/TypedElementDataStorageFactory.h"
+#include "UObject/ObjectMacros.h"
+
+#include "TedsSettingsFactory.generated.h"
+
+UCLASS()
+class UTedsSettingsFactory : public UEditorDataStorageFactory
+{
+	GENERATED_BODY()
+
+public:
+
+	~UTedsSettingsFactory() override = default;
+
+	void RegisterWidgetConstructors(IEditorDataStorageProvider& DataStorage, IEditorDataStorageUiProvider& DataStorageUi) const override;
+
+};
diff --git a/Engine/Plugins/Experimental/EditorDataStorageFeatures/Source/TedsSettings/Private/TedsSettingsWidgets.cpp b/Engine/Plugins/Experimental/EditorDataStorageFeatures/Source/TedsSettings/Private/TedsSettingsWidgets.cpp
new file mode 100644
--- /dev/null
+++ b/Engine/Plugins/Experimental/EditorDataStorageFeatures/Source/TedsSettings/Private/TedsSettingsWidgets.cpp
@@ -0,0 +1,140 @@
+// Copyright Epic Games, Inc. All Rights Reserved.
+
+#include "TedsSettingsWidgets.h"
+
+#include "Elements/Columns/TypedElementMiscColumns.h"
+#include "Elements/Columns/TypedElementSlateWidgetColumns.h"
+#include "Elements/Framework/TypedElementAttributeBinding.h"
+#include "TedsSettingsColumns.h"
+#include "Widgets/SToolTip.h"
+#include "Widgets/Layout/SBox.h"
+#include "Widgets/Text/STextBlock.h"
+
+FSettingsContainerReferenceWidgetConstructor::FSettingsContainerReferenceWidgetConstructor()
+	: FSimpleWidgetConstructor(FSettingsContainerReferenceWidgetConstructor::StaticStruct())
+{
+}
+
+TSharedPtr<SWidget> FSettingsContainerReferenceWidgetConstructor::CreateWidget(
+	IEditorDataStorageProvider* DataStorage,
+	IEditorDataStorageUiProvider* DataStorageUi,
+	UE::Editor::DataStorage::RowHandle TargetRow,
+	UE::Editor::DataStorage::RowHandle WidgetRow,
+	const UE::Editor::DataStorage::FMetaDataView& Arguments)
+{
+	using namespace UE::Editor::DataStorage;
+
+	if (DataStorage->IsRowAvailable(TargetRow))
+	{
+		RowHandle ContainerRow = DataStorage->GetColumn<FSettingsContainerReferenceColumn>(TargetRow)->ContainerRow;
+
+		if (DataStorage->IsRowAvailable(ContainerRow))
+		{
+			FAttributeBinder Binder(ContainerRow, DataStorage);
+
+			return SNew(SBox)
+				.Padding(4.0f, 0.0f)
+				.VAlign(VAlign_Center)
+				[
+					SNew(STextBlock)
+						.Text(Binder.BindData(&FDisplayNameColumn::DisplayName))
+						.ToolTip(SNew(SToolTip).Text(Binder.BindData(&FDescriptionColumn::Description)))
+				];
+		}
+		else // Fallback to render just the FName if we can't follow the reference
+		{
+			FAttributeBinder Binder(TargetRow, DataStorage);
+
+			return SNew(SBox)
+				.Padding(4.0f, 0.0f)
+				.VAlign(VAlign_Center)
+				[
+					SNew(STextBlock)
+						.ColorAndOpacity(FSlateColor::UseSubduedForeground())
+						.Text(Binder.BindText(&FSettingsContainerReferenceColumn::ContainerName))
+				];
+		}
+	}
+
+	return SNullWidget::NullWidget;
+}
+
+FSettingsCategoryReferenceWidgetConstructor::FSettingsCategoryReferenceWidgetConstructor()
+	: FSimpleWidgetConstructor(FSettingsCategoryReferenceWidgetConstructor::StaticStruct())
+{
+}
+
+TSharedPtr<SWidget> FSettingsCategoryReferenceWidgetConstructor::CreateWidget(
+	IEditorDataStorageProvider* DataStorage,
+	IEditorDataStorageUiProvider* DataStorageUi,
+	UE::Editor::DataStorage::RowHandle TargetRow,
+	UE::Editor::DataStorage::RowHandle WidgetRow,
+	const UE::Editor::DataStorage::FMetaDataView& Arguments)
+{
+	using namespace UE::Editor::DataStorage;
+
+	if (DataStorage->IsRowAvailable(TargetRow))
+	{
+		RowHandle CategoryRow = DataStorage->GetColumn<FSettingsCategoryReferenceColumn>(TargetRow)->CategoryRow;
+
+		if (DataStorage->IsRowAvailable(CategoryRow))
+		{
+			FAttributeBinder Binder(CategoryRow, DataStorage);
+
+			return SNew(SBox)
+				.Padding(4.0f, 0.0f)
+				.VAlign(VAlign_Center)
+				[
+					SNew(STextBlock)
+						.Text(Binder.BindData(&FDisplayNameColumn::DisplayName))
+						.ToolTip(SNew(SToolTip).Text(Binder.BindData(&FDescriptionColumn::Description)))
+				];
+		}
+		else // Fallback to render just the FName if we can't follow the reference
+		{
+			FAttributeBinder Binder(TargetRow, DataStorage);
+
+			return SNew(SBox)
+				.Padding(4.0f, 0.0f)
+				.VAlign(VAlign_Center)
+				[
+					SNew(STextBlock)
+						.ColorAndOpacity(FSlateColor::UseSubduedForeground())
+						.Text(Binder.BindText(&FSettingsCategoryReferenceColumn::CategoryName))
+				];
+		}
+	}
+
+	return SNullWidget::NullWidget;
+}
+
+FSettingsSectionWidgetConstructor::FSettingsSectionWidgetConstructor()
+	: FSimpleWidgetConstructor(FSettingsSectionWidgetConstructor::StaticStruct())
+{
+}
+
+TSharedPtr<SWidget> FSettingsSectionWidgetConstructor::CreateWidget(
+	IEditorDataStorageProvider* DataStorage,
+	IEditorDataStorageUiProvider* DataStorageUi,
+	UE::Editor::DataStorage::RowHandle TargetRow,
+	UE::Editor::DataStorage::RowHandle WidgetRow,
+	const UE::Editor::DataStorage::FMetaDataView& Arguments)
+{
+	using namespace UE::Editor::DataStorage;
+
+	if (DataStorage->IsRowAvailable(TargetRow))
+	{
+		FAttributeBinder Binder(TargetRow, DataStorage);
+
+		return SNew(SBox)
+			.Padding(4.0f, 0.0f)
+			.VAlign(VAlign_Center)
+			[
+				SNew(STextBlock)
+					.Text(Binder.BindData(&FDisplayNameColumn::DisplayName))
+					.ToolTip(SNew(SToolTip).Text(Binder.BindData(&FDescriptionColumn::Description)))
+			];
+	}
+
+	return SNullWidget::NullWidget;
+}
diff --git a/Engine/Plugins/Experimental/EditorDataStorageFeatures/Source/TedsSettings/Private/TedsSettingsWidgets.h b/Engine/Plugins/Experimental/EditorDataStorageFeatures/Source/TedsSettings/Private/TedsSettingsWidgets.h
new file mode 100644
--- /dev/null
+++ b/Engine/Plugins/Experimental/EditorDataStorageFeatures/Source/TedsSettings/Private/TedsSettingsWidgets.h
@@ -0,0 +1,65 @@
+// Copyright Epic Games, Inc. All Rights Reserved.
+
+#pragma once
+
+#include "Elements/Interfaces/TypedElementDataStorageUiInterface.h"
+#include "UObject/ObjectMacros.h"
+
+#include "TedsSettingsWidgets.generated.h"
+
+USTRUCT()
+struct FSettingsContainerReferenceWidgetConstructor : public FSimpleWidgetConstructor
+{
+	GENERATED_BODY()
+
+public:
+
+	FSettingsContainerReferenceWidgetConstructor();
+	~FSettingsContainerReferenceWidgetConstructor() override = default;
+
+	virtual TSharedPtr<SWidget> CreateWidget(
+		IEditorDataStorageProvider* DataStorage,
+		IEditorDataStorageUiProvider* DataStorageUi,
+		UE::Editor::DataStorage::RowHandle TargetRow,
+		UE::Editor::DataStorage::RowHandle WidgetRow,
+		const UE::Editor::DataStorage::FMetaDataView& Arguments) override;
+
+};
+
+USTRUCT()
+struct FSettingsCategoryReferenceWidgetConstructor : public FSimpleWidgetConstructor
+{
+	GENERATED_BODY()
+
+public:
+
+	FSettingsCategoryReferenceWidgetConstructor();
+	~FSettingsCategoryReferenceWidgetConstructor() override = default;
+
+	virtual TSharedPtr<SWidget> CreateWidget(
+		IEditorDataStorageProvider* DataStorage,
+		IEditorDataStorageUiProvider* DataStorageUi,
+		UE::Editor::DataStorage::RowHandle TargetRow,
+		UE::Editor::DataStorage::RowHandle WidgetRow,
+		const UE::Editor::DataStorage::FMetaDataView& Arguments) override;
+
+};
+
+USTRUCT()
+struct FSettingsSectionWidgetConstructor : public FSimpleWidgetConstructor
+{
+	GENERATED_BODY()
+
+public:
+
+	FSettingsSectionWidgetConstructor();
+	~FSettingsSectionWidgetConstructor() override = default;
+
+	virtual TSharedPtr<SWidget> CreateWidget(
+		IEditorDataStorageProvider* DataStorage,
+		IEditorDataStorageUiProvider* DataStorageUi,
+		UE::Editor::DataStorage::RowHandle TargetRow,
+		UE::Editor::DataStorage::RowHandle WidgetRow,
+		const UE::Editor::DataStorage::FMetaDataView& Arguments) override;
+
+};
diff --git a/Engine/Plugins/Experimental/EditorDataStorageFeatures/Source/TedsSettings/TedsSettings.Build.cs b/Engine/Plugins/Experimental/EditorDataStorageFeatures/Source/TedsSettings/TedsSettings.Build.cs
--- a/Engine/Plugins/Experimental/EditorDataStorageFeatures/Source/TedsSettings/TedsSettings.Build.cs
+++ b/Engine/Plugins/Experimental/EditorDataStorageFeatures/Source/TedsSettings/TedsSettings.Build.cs
@@ -19,6 +19,8 @@ public class TedsSettings : ModuleRules
 					"Core",
 					"CoreUObject",
 					"Engine",
+					"Slate",
+					"SlateCore",
 					"UnrealEd",
 				});
 
