[TEDS] Enable searching in the outliner by exposing the label as the display string

diff --git a/Engine/Plugins/Experimental/EditorDataStorageFeatures/Source/TedsOutliner/Private/TedsOutlinerItem.cpp b/Engine/Plugins/Experimental/EditorDataStorageFeatures/Source/TedsOutliner/Private/TedsOutlinerItem.cpp
--- a/Engine/Plugins/Experimental/EditorDataStorageFeatures/Source/TedsOutliner/Private/TedsOutlinerItem.cpp
+++ b/Engine/Plugins/Experimental/EditorDataStorageFeatures/Source/TedsOutliner/Private/TedsOutlinerItem.cpp
@@ -3,6 +3,7 @@
 #include "TedsOutlinerItem.h"
 
 #include "Elements/Columns/TypedElementLabelColumns.h"
+#include "Elements/Columns/TypedElementMiscColumns.h"
 #include "HAL/PlatformApplicationMisc.h"
 #include "TedsOutlinerImpl.h"
 
@@ -33,7 +34,16 @@ FSceneOutlinerTreeItemID FTedsOutlinerTreeItem::GetID() const
 
 FString FTedsOutlinerTreeItem::GetDisplayString() const
 {
-	return TEXT("TEDS Item"); // TEDS-Outliner TODO: Used for searching by name, how to get this from TEDS
+	if (const FTypedElementLabelColumn* LabelColumn = TedsOutlinerImpl->GetStorage()->GetColumn<FTypedElementLabelColumn>(RowHandle))
+	{
+		return LabelColumn->Label;
+	}
+	if (const FNameColumn* NameColumn = TedsOutlinerImpl->GetStorage()->GetColumn<FNameColumn>(RowHandle))
+	{
+		return NameColumn->Name.ToString();
+	}
+
+	return TEXT("TEDS Item");
 }
 
 bool FTedsOutlinerTreeItem::CanInteract() const
diff --git a/Engine/Plugins/Experimental/EditorDataStorageFeatures/Source/TedsPropertyEditor/Private/Widgets/SPropertyMenuTedsRowPicker.cpp b/Engine/Plugins/Experimental/EditorDataStorageFeatures/Source/TedsPropertyEditor/Private/Widgets/SPropertyMenuTedsRowPicker.cpp
--- a/Engine/Plugins/Experimental/EditorDataStorageFeatures/Source/TedsPropertyEditor/Private/Widgets/SPropertyMenuTedsRowPicker.cpp
+++ b/Engine/Plugins/Experimental/EditorDataStorageFeatures/Source/TedsPropertyEditor/Private/Widgets/SPropertyMenuTedsRowPicker.cpp
@@ -64,7 +64,7 @@ void SPropertyMenuTedsRowPicker::Construct(const FArguments& InArgs)
 				FSceneOutlinerInitializationOptions InitOptions;
 				InitOptions.bShowHeaderRow = true;
 				InitOptions.bShowTransient = true;
-				InitOptions.bShowSearchBox = false; // Search not currently supported in TEDS Outliner
+				InitOptions.bShowSearchBox = true;
 
 				InitOptions.ModeFactory = FCreateSceneOutlinerMode::CreateLambda([&](SSceneOutliner* Outliner)
 					{
