TEDS CB: Fix a crash when you set folder color through the debugger
TEDS UI: Add a default attribute binding for FSlateColorColumn for all TEDS UI Widgets (behind a cvar)

diff --git a/Engine/Plugins/Experimental/EditorDataStorage/Source/TedsCore/Private/TypedElementDatabaseUI.cpp b/Engine/Plugins/Experimental/EditorDataStorage/Source/TedsCore/Private/TypedElementDatabaseUI.cpp
--- a/Engine/Plugins/Experimental/EditorDataStorage/Source/TedsCore/Private/TypedElementDatabaseUI.cpp
+++ b/Engine/Plugins/Experimental/EditorDataStorage/Source/TedsCore/Private/TypedElementDatabaseUI.cpp
@@ -462,7 +462,8 @@ void UEditorDataStorageUi::CreateStandardArchetypes()
 	WidgetTable = Storage->RegisterTable(MakeArrayView(
 		{
 			FTypedElementSlateWidgetReferenceColumn::StaticStruct(),
-			FTypedElementSlateWidgetReferenceDeletesRowTag::StaticStruct()
+			FTypedElementSlateWidgetReferenceDeletesRowTag::StaticStruct(),
+			FSlateColorColumn::StaticStruct()
 		}), FName(TEXT("Editor_WidgetTable")));
 }
 
diff --git a/Engine/Plugins/Experimental/EditorDataStorageFeatures/Source/TedsAssetData/Private/Processors/AssetProcessors.cpp b/Engine/Plugins/Experimental/EditorDataStorageFeatures/Source/TedsAssetData/Private/Processors/AssetProcessors.cpp
--- a/Engine/Plugins/Experimental/EditorDataStorageFeatures/Source/TedsAssetData/Private/Processors/AssetProcessors.cpp
+++ b/Engine/Plugins/Experimental/EditorDataStorageFeatures/Source/TedsAssetData/Private/Processors/AssetProcessors.cpp
@@ -30,7 +30,7 @@ namespace UE::TedsAssetDataFactory::Private
 
 		// Extract the X Dimension first (everything to the left of the first x)
 		int32 FirstXPos;
-		if(!Dimensions.FindChar('x', FirstXPos))
+		if (!Dimensions.FindChar('x', FirstXPos))
 		{
 			return false; // Failsafe, in case the dimension isn't in the correct format somehow.
 		}
@@ -42,7 +42,7 @@ namespace UE::TedsAssetDataFactory::Private
 		LexFromString(X, XDimension);
 
 		// If it isn't a power of 2, early out since we already know this is a non square without needing to check the rest
-		if(!FMath::IsPowerOfTwo(X))
+		if (!FMath::IsPowerOfTwo(X))
 		{
 			return true;
 		}
@@ -52,7 +52,7 @@ namespace UE::TedsAssetDataFactory::Private
 		int32 SecondXPos;
 
 		// If we find another x, this is a 3D texture
-		if(YDimension.FindChar('x', SecondXPos))
+		if (YDimension.FindChar('x', SecondXPos))
 		{
 			// Extract out the Z Dimension first, which is everything to the right of the second X. And then convert it to an int
 			FStringView ZDimension = YDimension.RightChop(SecondXPos + 1);
@@ -60,7 +60,7 @@ namespace UE::TedsAssetDataFactory::Private
 			LexFromString(Z, ZDimension);
 
 			// Early out if the Z Dimension is not a power of 2 so we don't need to check Y Dimension
-			if(!FMath::IsPowerOfTwo(Z))
+			if (!FMath::IsPowerOfTwo(Z))
 			{
 				return true;
 			}
@@ -73,7 +73,7 @@ namespace UE::TedsAssetDataFactory::Private
 		int32 Y;
 		LexFromString(Y, YDimension);
 
-		if(!FMath::IsPowerOfTwo(Y))
+		if (!FMath::IsPowerOfTwo(Y))
 		{
 			return true;
 		}
@@ -110,6 +110,17 @@ void UTedsAssetDataFactory::RegisterQueries(IEditorDataStorageProvider& DataStor
 			.All<FFolderTag, FUpdatedPathTag, FVirtualPathColumn_Experimental>()
 		.Compile()
 		);
+	
+	struct FSetFolderColorCommand
+	{
+		void operator()()
+		{
+			
+			UE::Editor::ContentBrowser::ExtensionUtils::SetFolderColor(FolderPath, NewFolderColor);
+		}
+		FName FolderPath;
+		FLinearColor NewFolderColor;;
+	};
 
 	DataStorage.RegisterQuery(
 		Select(
@@ -122,7 +133,18 @@ void UTedsAssetDataFactory::RegisterQueries(IEditorDataStorageProvider& DataStor
 
 				for (int32 Index = 0; Index < NumOfRowToProcess; ++Index)
 				{
-					UE::Editor::ContentBrowser::ExtensionUtils::SetFolderColor(PathColumn[Index].Path, ColorColumn[Index].Color.GetSpecifiedColor());
+					FSlateColor FolderColor = ColorColumn[Index].Color.GetSpecifiedColor();
+
+					if (FolderColor.IsColorSpecified())
+					{
+						// Defer the add because it will fire a CB delegate that UTedsAssetDataFactory::OnSetFolderColor registers to and ends up
+						// accessing the data storage in the middle of a processor callback which is not allowed
+						Context.PushCommand(FSetFolderColorCommand
+						{
+							.FolderPath = PathColumn[Index].Path,
+							.NewFolderColor = ColorColumn[Index].Color.GetSpecifiedColor()
+						});
+					}
 				}
 			}
 		)
@@ -145,7 +167,7 @@ void UTedsAssetDataFactory::RegisterQueries(IEditorDataStorageProvider& DataStor
 				
 				for (int32 Index = 0; Index < NumRows; ++Index)
 				{
-					if(UE::TedsAssetDataFactory::Private::IsTextureNonSquare(ColumnView[Index].Value))
+					if (UE::TedsAssetDataFactory::Private::IsTextureNonSquare(ColumnView[Index].Value))
 					{
 						FTypedElementAlertColumn Alert;
 						Alert.AlertType = FTypedElementAlertColumnType::Error;
@@ -187,7 +209,7 @@ void UTedsAssetDataFactory::OnSetFolderColor(const FString& Path, IEditorDataSto
 	const UE::Editor::DataStorage::IndexHash PathHash = UE::Editor::DataStorage::GenerateIndexHash(FName(Path));
 	const UE::Editor::DataStorage::RowHandle Row = DataStorage->FindIndexedRow(PathHash);
 
-	if(DataStorage->IsRowAvailable(Row))
+	if (DataStorage->IsRowAvailable(Row))
 	{
 		DataStorage->AddColumn<FUpdatedPathTag>(Row);
 	}
diff --git a/Engine/Source/Runtime/TypedElementFramework/Private/Elements/Framework/TypedElementDataStorageWidget.cpp b/Engine/Source/Runtime/TypedElementFramework/Private/Elements/Framework/TypedElementDataStorageWidget.cpp
--- a/Engine/Source/Runtime/TypedElementFramework/Private/Elements/Framework/TypedElementDataStorageWidget.cpp
+++ b/Engine/Source/Runtime/TypedElementFramework/Private/Elements/Framework/TypedElementDataStorageWidget.cpp
@@ -3,8 +3,21 @@
 #include "Elements/Framework/TypedElementDataStorageWidget.h"
 #include "Elements/Columns/TypedElementSlateWidgetColumns.h"
 #include "Elements/Common//EditorDataStorageFeatures.h"
+#include "Elements/Framework/TypedElementAttributeBinding.h"
 #include "Elements/Interfaces/TypedElementDataStorageInterface.h"
 
+namespace UE::Editor::DataStorage::TedsWidget::Private
+{
+	bool bUseDefaultAttributeBindings = false;
+
+	static FAutoConsoleVariableRef CVarUseDefaultAttributeBindings(
+		TEXT("TEDS.UI.UseDefaultAttributeBindings"),
+		bUseDefaultAttributeBindings,
+		TEXT("If true, widgets created through TEDS UI will provide a default set of attribute bindings. Does not apply to existing widgets retroactively"),
+		ECVF_Default
+	);
+}
+
 STedsWidget::STedsWidget()
 	: UiRowHandle(UE::Editor::DataStorage::InvalidRowHandle)
 {
@@ -21,6 +34,21 @@ void STedsWidget::Construct(const FArguments& InArgs)
 		RegisterTedsWidget(InArgs._Content.Widget);
 	}
 
+	if (UE::Editor::DataStorage::TedsWidget::Private::bUseDefaultAttributeBindings)
+	{
+		UE::Editor::DataStorage::FAttributeBinder Binder(UiRowHandle, GetStorageIfAvailable());
+
+		SetColorAndOpacity(Binder.BindData(&FSlateColorColumn::Color, [](const FSlateColor& InColor)
+		{
+			if(InColor.IsColorSpecified())
+			{
+				return InColor.GetSpecifiedColor();
+			}
+
+			return FLinearColor::White;
+		}));
+	}
+
 	ChildSlot
 	[
 		InArgs._Content.Widget
