TEDS UI: Add/remove the additional columns when the content is set/removed in STedsWidget
Table Viewer: Fix a bug where the column based virtualization would end up accessing the data storage interface in the middle of a query callback

This fixes a bug with the alert widget where it would cause a crash sometimes

diff --git a/Engine/Plugins/Experimental/EditorDataStorageFeatures/Source/TedsTableViewer/Private/TedsTableViewerColumn.cpp b/Engine/Plugins/Experimental/EditorDataStorageFeatures/Source/TedsTableViewer/Private/TedsTableViewerColumn.cpp
--- a/Engine/Plugins/Experimental/EditorDataStorageFeatures/Source/TedsTableViewer/Private/TedsTableViewerColumn.cpp
+++ b/Engine/Plugins/Experimental/EditorDataStorageFeatures/Source/TedsTableViewer/Private/TedsTableViewerColumn.cpp
@@ -274,58 +274,65 @@ namespace UE::Editor::DataStorage
 
 		using namespace UE::Editor::DataStorage::Queries;
 
-		// Query to find all widgets that belong to the row handles that need updates
-		DirectQueryCallback RowCollector = CreateDirectQueryCallbackBinding(
-		[this](const IDirectQueryContext& Context, const FTypedElementSlateWidgetReferenceColumn* ContainerWidgetReferenceColumns, const FTypedElementRowReferenceColumn* RowReferenceColumns)
+		TArray<RowHandle> MatchedWidgetRows;
+		DirectQueryCallback RowCollector = CreateDirectQueryCallbackBinding([&MatchedWidgetRows](const IDirectQueryContext& Context, const RowHandle*)
 		{
-			const FTypedElementSlateWidgetReferenceColumn* WidgetsIt = ContainerWidgetReferenceColumns;
-			const FTypedElementRowReferenceColumn* RowRefsIt = RowReferenceColumns;
-			const TConstArrayView<RowHandle> Rows = Context.GetRowHandles();
+			MatchedWidgetRows.Append(Context.GetRowHandles());
+		});
+
+		// Run query to gather all widget rows
+		Storage->RunQuery(WidgetQuery, RowCollector);
 
-			for(unsigned RowIndex = 0; RowIndex < Context.GetRowCount(); ++RowIndex, ++WidgetsIt, ++RowRefsIt)
+		// Run the actual logic outside the query because updating the widget can add/remove columns through the DataStorage which is invalid
+		// when you are inside a query callback
+		for(RowHandle Row : MatchedWidgetRows)
+		{
+			const FTypedElementSlateWidgetReferenceColumn* WidgetColumn = Storage->GetColumn<FTypedElementSlateWidgetReferenceColumn>(Row);
+			const FTypedElementRowReferenceColumn* RowReferenceColumn = Storage->GetColumn<FTypedElementRowReferenceColumn>(Row);
+
+			if(!ensureMsgf(WidgetColumn && RowReferenceColumn,
+				TEXT("Expected to have the widget reference and row reference columns since we queried for them")))
 			{
-				// Check if this widget's owning row is in our rows to update
-				bool* bColumnAddedPtr = RowsToUpdate.Find(RowRefsIt->Row);
-				// If not, skip it
-				if(!bColumnAddedPtr)
-				{
-					continue;
-				}
+				continue;
+			}
 
-				// Check if the container TEDSWidget exists, if not we cannot update this widget
-				const TSharedPtr<STedsWidget> TedsWidget = WidgetsIt->TedsWidget.Pin();
-				if(!TedsWidget)
-				{
-					continue;
-				}
+			// Check if this widget's owning row is in our rows to update
+			bool* bColumnAddedPtr = RowsToUpdate.Find(RowReferenceColumn->Row);
+			// If not, skip it
+			if(!bColumnAddedPtr)
+			{
+				continue;
+			}
 
-				// A row has numerous widgets, make sure we only update the one that was created by our column by checking the constructor
-				if(WidgetsIt->WidgetConstructor != CellWidgetConstructor)
-				{
-					continue;
-				}
+			// Check if the container TEDSWidget exists, if not we cannot update this widget
+			const TSharedPtr<STedsWidget> TedsWidget = WidgetColumn->TedsWidget.Pin();
+			if(!TedsWidget)
+			{
+				continue;
+			}
 
-				// If a column was added and we are here, we need to re-create the widget
-				// TEDS-Outliner TODO: Do we need to create the widget only if it doesn't exist? Or should we also update it to automatically respond
-				// to column changes even if it was already created
-				if(*bColumnAddedPtr)
-				{
-					const TSharedPtr<SWidget> RowWidget = CellWidgetConstructor->Construct(Rows[RowIndex], Storage, StorageUi, WidgetMetaData);
+			// A row has numerous widgets, make sure we only update the one that was created by our column by checking the constructor
+			if(WidgetColumn->WidgetConstructor != CellWidgetConstructor)
+			{
+				continue;
+			}
 
-					if(RowWidget)
-					{
-						TedsWidget->SetContent(RowWidget.ToSharedRef());
-					}
-				}
-				// If a column was removed (and we don't match anymore) delete the internal widget
-				else
+			// If a column was added and we are here, we need to re-create the widget
+			// TEDS-Outliner TODO: Do we need to create the widget only if it doesn't exist? Or should we also update it to automatically respond
+			// to column changes even if it was already created
+			if(*bColumnAddedPtr)
+			{
+				if(const TSharedPtr<SWidget> RowWidget = CellWidgetConstructor->Construct(Row, Storage, StorageUi, WidgetMetaData))
 				{
-					TedsWidget->SetContent(SNullWidget::NullWidget);
+					TedsWidget->SetContent(RowWidget.ToSharedRef());
 				}
 			}
-		});
-
-		Storage->RunQuery(WidgetQuery, RowCollector);
+			// If a column was removed (and we don't match anymore) delete the internal widget
+			else
+			{
+				TedsWidget->SetContent(SNullWidget::NullWidget);
+			}
+		}
 	}
 	
 	FName FTedsTableViewerColumn::GetColumnName() const
diff --git a/Engine/Source/Runtime/TypedElementFramework/Private/Elements/Framework/TypedElementDataStorageWidget.cpp b/Engine/Source/Runtime/TypedElementFramework/Private/Elements/Framework/TypedElementDataStorageWidget.cpp
--- a/Engine/Source/Runtime/TypedElementFramework/Private/Elements/Framework/TypedElementDataStorageWidget.cpp
+++ b/Engine/Source/Runtime/TypedElementFramework/Private/Elements/Framework/TypedElementDataStorageWidget.cpp
@@ -5,6 +5,7 @@
 #include "Elements/Common//EditorDataStorageFeatures.h"
 #include "Elements/Framework/TypedElementAttributeBinding.h"
 #include "Elements/Interfaces/TypedElementDataStorageInterface.h"
+#include "Elements/Interfaces/TypedElementDataStorageUiInterface.h"
 
 namespace UE::Editor::DataStorage::TedsWidget::Private
 {
@@ -29,7 +30,7 @@ void STedsWidget::Construct(const FArguments& InArgs)
 	UiRowHandle = InArgs._UiRowHandle;
 
 	// If the Ui row wasn't already registered externally, register it with Teds
-	if(UiRowHandle == UE::Editor::DataStorage::InvalidRowHandle)
+	if (UiRowHandle == UE::Editor::DataStorage::InvalidRowHandle)
 	{
 		RegisterTedsWidget(InArgs._Content.Widget);
 	}
@@ -49,10 +50,7 @@ void STedsWidget::Construct(const FArguments& InArgs)
 		}));
 	}
 
-	ChildSlot
-	[
-		InArgs._Content.Widget
-	];
+	SetContent(InArgs._Content.Widget);
 }
 
 void STedsWidget::RegisterTedsWidget(const TSharedPtr<SWidget>& InContentWidget)
@@ -61,20 +59,20 @@ void STedsWidget::RegisterTedsWidget(const TSharedPtr<SWidget>& InContentWidget)
 	IEditorDataStorageProvider* Storage = GetStorageIfAvailable();
 
 	// If TEDS is not enabled, STedsWidget will just behave like a regular widget
-	if(!Storage)
+	if (!Storage)
 	{
 		return;
 	}
 	
 	const TableHandle WidgetTable = Storage->FindTable(TEXT("Editor_WidgetTable"));
-	if(WidgetTable == InvalidTableHandle)
+	if (WidgetTable == InvalidTableHandle)
 	{
 		return;
 	}
 	
 	UiRowHandle = Storage->AddRow(WidgetTable);
 	
-	if(FTypedElementSlateWidgetReferenceColumn* WidgetReferenceColumn = Storage->GetColumn<FTypedElementSlateWidgetReferenceColumn>(UiRowHandle))
+	if (FTypedElementSlateWidgetReferenceColumn* WidgetReferenceColumn = Storage->GetColumn<FTypedElementSlateWidgetReferenceColumn>(UiRowHandle))
 	{
 		WidgetReferenceColumn->TedsWidget = SharedThis(this);
 		WidgetReferenceColumn->Widget = InContentWidget;
@@ -83,11 +81,34 @@ void STedsWidget::RegisterTedsWidget(const TSharedPtr<SWidget>& InContentWidget)
 
 void STedsWidget::SetContent(const TSharedRef< SWidget >& InContent)
 {
-	if(IEditorDataStorageProvider* Storage = GetStorageIfAvailable())
+	if (IEditorDataStorageProvider* Storage = GetStorageIfAvailable())
 	{
-		if(FTypedElementSlateWidgetReferenceColumn* WidgetReferenceColumn = Storage->GetColumn<FTypedElementSlateWidgetReferenceColumn>(UiRowHandle))
+		if (FTypedElementSlateWidgetReferenceColumn* WidgetReferenceColumn = Storage->GetColumn<FTypedElementSlateWidgetReferenceColumn>(UiRowHandle))
 		{
+			// First we set the widget reference on the column
 			WidgetReferenceColumn->Widget = InContent;
+
+			if (TSharedPtr<FTypedElementWidgetConstructor> Constructor = WidgetReferenceColumn->WidgetConstructor.Pin())
+			{
+				TConstArrayView<const UScriptStruct*> AdditionalColumns = Constructor->GetAdditionalColumnsList();
+				
+				// When we are setting valid content, we want to add the additional columns to the widget row so it gets picked up by any queries
+				// the widget constructor requires.
+				if (InContent != SNullWidget::NullWidget)
+				{
+					Storage->AddColumns(UiRowHandle, AdditionalColumns);
+				}
+				// If we are setting to null content, i.e removing the widget we also remove the additional columns so the row stops matching any
+				// queries designed to operate on the widget
+				else
+				{
+					// RemoveColumns ensures if the columns don't exist on the row
+					if (Storage->HasColumns(UiRowHandle, AdditionalColumns))
+					{
+						Storage->RemoveColumns(UiRowHandle, AdditionalColumns);
+					}
+				}
+			}
 		}
 	}
 	
