UE-224438 - Remove placeholder workflow
==============================================================================

* Created a workflow to delete placeholder components and entities from the outliner
* Adds a right click menu to delete placeholder entries in the entity outliner
* Created a custom column that allows you to add a context menu without having to deal with editor UI directly in modules that might not want to import those editor-UI specific modules.

diff --git a/Engine/Plugins/Experimental/EditorDataStorageFeatures/Source/TedsOutliner/Private/TedsOutlinerItem.cpp b/Engine/Plugins/Experimental/EditorDataStorageFeatures/Source/TedsOutliner/Private/TedsOutlinerItem.cpp
--- a/Engine/Plugins/Experimental/EditorDataStorageFeatures/Source/TedsOutliner/Private/TedsOutlinerItem.cpp
+++ b/Engine/Plugins/Experimental/EditorDataStorageFeatures/Source/TedsOutliner/Private/TedsOutlinerItem.cpp
@@ -7,6 +7,16 @@
 #include "HAL/PlatformApplicationMisc.h"
 #include "TedsOutlinerImpl.h"
 
+#include "UObject/PropertyBagRepository.h"
+
+#include "Columns/TedsOutlinerColumns.h"
+#include "Elements/Columns/TypedElementCompatibilityColumns.h"
+#include "Elements/Common//EditorDataStorageFeatures.h"
+#include "Elements/Interfaces/TypedElementDataStorageInterface.h"
+
+
+#include "UObject/InstanceDataObjectUtils.h"
+
 #define LOCTEXT_NAMESPACE "TedsOutliner"
 
 namespace UE::Editor::Outliner
@@ -59,22 +69,16 @@ TSharedRef<SWidget> FTedsOutlinerTreeItem::GenerateLabelWidget(ISceneOutliner& O
 
 void FTedsOutlinerTreeItem::GenerateContextMenu(UToolMenu* Menu, SSceneOutliner& Outliner)
 {
-	FToolMenuSection& Section = Menu->AddSection("Copy", LOCTEXT("CopySection", "Copy"));
-
-	Section.AddMenuEntry(
-		"CopyRowHandle",
-		LOCTEXT("CopyRowHandle_Title", "Copy row handle"),
-		LOCTEXT("CopyRowHandle_Tooltip", "Copy the row handle of this row to the clipboard."),
-		FSlateIcon(),
-		FUIAction(
-			FExecuteAction::CreateLambda([this]()
-			{
-				const FString ClipboardString(LexToString<FString>(RowHandle));
-				FPlatformApplicationMisc::ClipboardCopy(*ClipboardString);
-			}),
-			FCanExecuteAction()
-		)
-	);
+	using namespace UE::Editor::DataStorage;
+	IEditorDataStorageProvider* DataStorageInterface = GetMutableDataStorageFeature<IEditorDataStorageProvider>(StorageFeatureName);
+
+	if (FTedsOutlinerContextMenuColumn* Column = DataStorageInterface->GetColumn<FTedsOutlinerContextMenuColumn>(RowHandle))
+	{
+		if (Column->OnCreateContextMenu.IsBound())
+		{
+			Column->OnCreateContextMenu.Execute(Menu, Outliner);
+		}
+	}
 }
 
 DataStorage::RowHandle FTedsOutlinerTreeItem::GetRowHandle() const
diff --git a/Engine/Plugins/Experimental/EditorDataStorageFeatures/Source/TedsOutliner/Public/Columns/TedsOutlinerColumns.h b/Engine/Plugins/Experimental/EditorDataStorageFeatures/Source/TedsOutliner/Public/Columns/TedsOutlinerColumns.h
--- a/Engine/Plugins/Experimental/EditorDataStorageFeatures/Source/TedsOutliner/Public/Columns/TedsOutlinerColumns.h
+++ b/Engine/Plugins/Experimental/EditorDataStorageFeatures/Source/TedsOutliner/Public/Columns/TedsOutlinerColumns.h
@@ -8,6 +8,8 @@
 #include "TedsOutlinerColumns.generated.h"
 
 class ISceneOutliner;
+class UToolMenu;
+class SSceneOutliner;
 
 // Column used to store a reference to the Teds Outliner owning a specific row
 // Currently only added to widget rows in the Teds Outliner
@@ -18,3 +20,14 @@ struct FTedsOutlinerColumn final : public FEditorDataStorageColumn
 
 	TWeakPtr<ISceneOutliner> Outliner;
 };
+
+DECLARE_DELEGATE_TwoParams(FTedsOutlinerContextMenuDelegate, UToolMenu* Menu, SSceneOutliner& SceneOutliner)
+
+/** Column used to allow context menu to be extended for an item in the outliner */
+USTRUCT()
+struct FTedsOutlinerContextMenuColumn final : public FEditorDataStorageColumn
+{
+	GENERATED_BODY()
+
+	FTedsOutlinerContextMenuDelegate OnCreateContextMenu;
+};
\ No newline at end of file
