TEDS-Outliner: Add a column to show row handles

diff --git a/Engine/Plugins/Experimental/EditorDataStorageFeatures/Source/TedsOutliner/Private/Compatibility/SceneOutlinerRowHandleColumn.cpp b/Engine/Plugins/Experimental/EditorDataStorageFeatures/Source/TedsOutliner/Private/Compatibility/SceneOutlinerRowHandleColumn.cpp
--- a/Engine/Plugins/Experimental/EditorDataStorageFeatures/Source/TedsOutliner/Private/Compatibility/SceneOutlinerRowHandleColumn.cpp
+++ b/Engine/Plugins/Experimental/EditorDataStorageFeatures/Source/TedsOutliner/Private/Compatibility/SceneOutlinerRowHandleColumn.cpp
@@ -7,11 +7,24 @@
 #include "Elements/Common/EditorDataStorageFeatures.h"
 #include "Elements/Interfaces/TypedElementDataStorageUiInterface.h"
 #include "TedsTableViewerColumn.h"
+#include "Compatibility/SceneOutlinerTedsBridge.h"
 
 #define LOCTEXT_NAMESPACE "SceneOutlinerRowHandleColumn"
 
 FSceneOutlinerRowHandleColumn::FSceneOutlinerRowHandleColumn(ISceneOutliner& SceneOutliner)
 	: WeakSceneOutliner(StaticCastSharedRef<ISceneOutliner>(SceneOutliner.AsShared()))
+{
+	CreateWidgetConstructor();
+}
+
+FSceneOutlinerRowHandleColumn::FSceneOutlinerRowHandleColumn(ISceneOutliner& SceneOutliner, const FGetRowHandle& InGetRowHandle)
+	: WeakSceneOutliner(StaticCastSharedRef<ISceneOutliner>(SceneOutliner.AsShared()))
+	, GetRowHandle(InGetRowHandle)
+{
+	CreateWidgetConstructor();
+}
+
+void FSceneOutlinerRowHandleColumn::CreateWidgetConstructor()
 {
 	using namespace UE::Editor::DataStorage;
 	auto AssignWidgetToColumn = [this](TUniquePtr<FTypedElementWidgetConstructor> Constructor, TConstArrayView<TWeakObjectPtr<const UScriptStruct>>)
@@ -23,8 +36,8 @@ FSceneOutlinerRowHandleColumn::FSceneOutlinerRowHandleColumn(ISceneOutliner& Sce
 	
 	IEditorDataStorageUiProvider* StorageUi = GetMutableDataStorageFeature<IEditorDataStorageUiProvider>(UiFeatureName);
 	checkf(StorageUi, TEXT("FSceneOutlinerRowHandleColumn created before data storage interfaces were initialized."))
-
-	StorageUi->CreateWidgetConstructors(TEXT("General.Cell.RowHandle"), UE::Editor::DataStorage::FMetaDataView(), AssignWidgetToColumn);
+	
+	StorageUi->CreateWidgetConstructors(TEXT("General.Cell.RowHandle"), FMetaDataView(), AssignWidgetToColumn);
 }
 
 
@@ -48,18 +61,42 @@ SHeaderRow::FColumn::FArguments FSceneOutlinerRowHandleColumn::ConstructHeaderRo
 
 const TSharedRef<SWidget> FSceneOutlinerRowHandleColumn::ConstructRowWidget(FSceneOutlinerTreeItemRef TreeItem, const STableRow<FSceneOutlinerTreeItemPtr>& Row)
 {
+	if (!TableViewerColumn)
+	{
+		return SNew(STextBlock)
+				.Text(LOCTEXT("MissingWidgetConstructor",
+					"Row Handles cannot be displayed if widget constructor for the General.Cell.RowHandle purpose was not found"));
+	}
+	
 	auto SceneOutliner = WeakSceneOutliner.Pin();
 	check(SceneOutliner.IsValid());
+	
+	UE::Editor::DataStorage::RowHandle RowHandle = UE::Editor::DataStorage::InvalidRowHandle;;
 
-	if (const UE::Editor::Outliner::FTedsOutlinerTreeItem* OutlinerTreeItem = TreeItem->CastTo<UE::Editor::Outliner::FTedsOutlinerTreeItem>())
+	// Check any custom delegates to get the row handle of the item
+	if (GetRowHandle.IsBound())
 	{
-		const UE::Editor::DataStorage::RowHandle RowHandle = OutlinerTreeItem->GetRowHandle();
+		RowHandle = GetRowHandle.Execute(TreeItem.Get());
+	}
+	// Otherwise it could just be a regular TedsOutlinerTreeItem which we know how to handle
+	else if (const UE::Editor::Outliner::FTedsOutlinerTreeItem* OutlinerTreeItem = TreeItem->CastTo<UE::Editor::Outliner::FTedsOutlinerTreeItem>())
+	{
+		RowHandle = OutlinerTreeItem->GetRowHandle();
+	}
 
-		if(TSharedPtr<SWidget> Widget = TableViewerColumn->ConstructRowWidget(RowHandle))
+	IEditorDataStorageProvider* Storage = TableViewerColumn->GetStorage();
+
+	if (ensureMsgf(Storage, TEXT("Cannot create a widget to display row handles before TEDS is initialized")))
+	{
+		if (Storage->IsRowAssigned(RowHandle))
 		{
-			return Widget.ToSharedRef();
+			if (TSharedPtr<SWidget> Widget = TableViewerColumn->ConstructRowWidget(RowHandle))
+			{
+				return Widget.ToSharedRef();
+			}
 		}
 	}
+	
 	return SNullWidget::NullWidget;
 }
 
diff --git a/Engine/Plugins/Experimental/EditorDataStorageFeatures/Source/TedsOutliner/Private/Compatibility/SceneOutlinerRowHandleColumn.h b/Engine/Plugins/Experimental/EditorDataStorageFeatures/Source/TedsOutliner/Private/Compatibility/SceneOutlinerRowHandleColumn.h
deleted file mode 100644
diff --git a/Engine/Plugins/Experimental/EditorDataStorageFeatures/Source/TedsOutliner/Public/Compatibility/SceneOutlinerRowHandleColumn.h b/Engine/Plugins/Experimental/EditorDataStorageFeatures/Source/TedsOutliner/Public/Compatibility/SceneOutlinerRowHandleColumn.h
new file mode 100644
--- /dev/null
+++ b/Engine/Plugins/Experimental/EditorDataStorageFeatures/Source/TedsOutliner/Public/Compatibility/SceneOutlinerRowHandleColumn.h
@@ -0,0 +1,53 @@
+// Copyright Epic Games, Inc. All Rights Reserved.
+
+#pragma once
+
+#include "SceneOutlinerFwd.h"
+#include "ISceneOutliner.h"
+#include "ISceneOutlinerColumn.h"
+#include "Elements/Common/TypedElementHandles.h"
+
+namespace UE::Editor::DataStorage
+{
+	class FTedsTableViewerColumn;
+}
+
+/**
+ * This is a custom column for the TEDS Outliner/Table Viewer to display row handles for items. It is a special cased column instead of going through
+ * the TEDS UI layer because row handles are not information stored in a TEDS column
+ */
+class FSceneOutlinerRowHandleColumn : public ISceneOutlinerColumn
+{
+public:
+
+	DECLARE_DELEGATE_RetVal_OneParam(UE::Editor::DataStorage::RowHandle, FGetRowHandle, const ISceneOutlinerTreeItem&);
+	
+	TEDSOUTLINER_API explicit FSceneOutlinerRowHandleColumn(ISceneOutliner& SceneOutliner);
+	TEDSOUTLINER_API explicit FSceneOutlinerRowHandleColumn(ISceneOutliner& SceneOutliner, const FGetRowHandle& InGetRowHandle);
+	
+	virtual ~FSceneOutlinerRowHandleColumn() {}
+
+	TEDSOUTLINER_API static FName GetID();
+	
+	//////////////////////////////////////////////////////////////////////////
+	// Begin ISceneOutlinerColumn Implementation
+	virtual FName GetColumnID() override;
+	virtual SHeaderRow::FColumn::FArguments ConstructHeaderRowColumn() override;
+	virtual const TSharedRef<SWidget> ConstructRowWidget(FSceneOutlinerTreeItemRef TreeItem, const STableRow<FSceneOutlinerTreeItemPtr>& Row) override;
+	virtual void PopulateSearchStrings( const ISceneOutlinerTreeItem& Item, TArray< FString >& OutSearchStrings ) const override;
+	virtual bool SupportsSorting() const override { return true; }
+	virtual void SortItems(TArray<FSceneOutlinerTreeItemPtr>& OutItems, const EColumnSortMode::Type SortMode) const override;
+	// End ISceneOutlinerColumn Implementation
+	//////////////////////////////////////////////////////////////////////////
+
+private:
+
+	void CreateWidgetConstructor();
+	
+private:
+	TWeakPtr<ISceneOutliner> WeakSceneOutliner;
+	FGetRowHandle GetRowHandle;
+
+	// The Table Viewer column we are going to use internally to create the widget
+	TSharedPtr<UE::Editor::DataStorage::FTedsTableViewerColumn> TableViewerColumn;
+};
diff --git a/Engine/Plugins/Experimental/EditorDataStorageFeatures/Source/TedsTableViewer/Private/TedsTableViewerColumn.cpp b/Engine/Plugins/Experimental/EditorDataStorageFeatures/Source/TedsTableViewer/Private/TedsTableViewerColumn.cpp
--- a/Engine/Plugins/Experimental/EditorDataStorageFeatures/Source/TedsTableViewer/Private/TedsTableViewerColumn.cpp
+++ b/Engine/Plugins/Experimental/EditorDataStorageFeatures/Source/TedsTableViewer/Private/TedsTableViewerColumn.cpp
@@ -344,5 +344,10 @@ namespace UE::Editor::DataStorage
 	{
 		return MatchedColumns;
 	}
+
+	IEditorDataStorageProvider* FTedsTableViewerColumn::GetStorage() const
+	{
+		return Storage;
+	}
 }
 
diff --git a/Engine/Plugins/Experimental/EditorDataStorageFeatures/Source/TedsTableViewer/Public/TedsTableViewerColumn.h b/Engine/Plugins/Experimental/EditorDataStorageFeatures/Source/TedsTableViewer/Public/TedsTableViewerColumn.h
--- a/Engine/Plugins/Experimental/EditorDataStorageFeatures/Source/TedsTableViewer/Public/TedsTableViewerColumn.h
+++ b/Engine/Plugins/Experimental/EditorDataStorageFeatures/Source/TedsTableViewer/Public/TedsTableViewerColumn.h
@@ -48,7 +48,7 @@ namespace UE::Editor::DataStorage
 
 		TEDSTABLEVIEWER_API TConstArrayView<TWeakObjectPtr<const UScriptStruct>> GetMatchedColumns() const;
 
-
+		TEDSTABLEVIEWER_API IEditorDataStorageProvider* GetStorage() const;
 	protected:
 
 		void RegisterQueries();
