[UnrealED/Slate/TedsOutliner/SceneOutliner] Add MiddleEllipsis Overflow policy to Outliner and DockTab title
Added style setting to EditorStyleSettings to enable/disable MiddleEllipsis used for the Editor (currently only Outliner and DockTab).

diff --git a/Engine/Plugins/Experimental/EditorDataStorageFeatures/Source/TedsOutliner/Private/Widgets/OutlinerLabelWidget.cpp b/Engine/Plugins/Experimental/EditorDataStorageFeatures/Source/TedsOutliner/Private/Widgets/OutlinerLabelWidget.cpp
--- a/Engine/Plugins/Experimental/EditorDataStorageFeatures/Source/TedsOutliner/Private/Widgets/OutlinerLabelWidget.cpp
+++ b/Engine/Plugins/Experimental/EditorDataStorageFeatures/Source/TedsOutliner/Private/Widgets/OutlinerLabelWidget.cpp
@@ -13,6 +13,7 @@
 #include "Elements/Columns/TypedElementTypeInfoColumns.h"
 #include "Elements/Interfaces/Capabilities/TypedElementUiEditableCapability.h"
 #include "Elements/Interfaces/Capabilities/TypedElementUiStyleOverrideCapability.h"
+#include "Settings/EditorStyleSettings.h"
 #include "Widgets/Images/SImage.h"
 #include "Widgets/Layout/SSpacer.h"
 #include "Widgets/Text/STextBlock.h"
@@ -184,7 +185,9 @@ TSharedRef<SWidget> FOutlinerLabelWidgetConstructor::CreateLabel(IEditorDataStor
 	{
 		FAttributeBinder TargetRowBinder(TargetRow, DataStorage);
 		FAttributeBinder WidgetRowBinder(WidgetRow, DataStorage);
-			
+
+		const bool bShouldUseMiddleEllipsis = GetDefault<UEditorStyleSettings>()->bEnableMiddleEllipsis;
+
 		TSharedPtr<SInlineEditableTextBlock> TextBlock = SNew(SInlineEditableTextBlock)
 			.OnTextCommitted_Lambda(
 				[DataStorage, TargetRow](const FText& NewText, ETextCommit::Type CommitInfo)
@@ -209,6 +212,7 @@ TSharedRef<SWidget> FOutlinerLabelWidgetConstructor::CreateLabel(IEditorDataStor
 				})
 			.Text(TargetRowBinder.BindText(&FTypedElementLabelColumn::Label))
 			.ToolTipText(TargetRowBinder.BindText(&FTypedElementLabelColumn::Label))
+			.OverflowPolicy(bShouldUseMiddleEllipsis ? ETextOverflowPolicy::MiddleEllipsis : TOptional<ETextOverflowPolicy>())
 			.IsSelected(WidgetRowBinder.BindEvent(&FExternalWidgetSelectionColumn::IsSelected));
 
 		TextBlock->AddMetadata(MakeShared<TTypedElementUiEditableCapability<SInlineEditableTextBlock>>(*TextBlock));
diff --git a/Engine/Source/Editor/EditorStyle/Private/StarshipStyle.cpp b/Engine/Source/Editor/EditorStyle/Private/StarshipStyle.cpp
--- a/Engine/Source/Editor/EditorStyle/Private/StarshipStyle.cpp
+++ b/Engine/Source/Editor/EditorStyle/Private/StarshipStyle.cpp
@@ -16,6 +16,7 @@
 #include "Styling/SegmentedControlStyle.h"
 #include "Styling/StyleColors.h"
 #include "Framework/Application/SlateApplication.h"
+#include "Framework/Docking/TabManager.h"
 
 
 #define ICON_FONT(...) FSlateFontInfo(RootToContentDir("Fonts/FontAwesome", TEXT(".ttf")), __VA_ARGS__)
@@ -3060,6 +3061,7 @@ void FStarshipEditorStyle::FStyle::SetupProjectBadgeStyle()
 void FStarshipEditorStyle::FStyle::SetupDockingStyles()
 {
 	// Use the Docking Styles defined in StarshipCoreStyle
+	FGlobalTabmanager::Get()->SetShouldUseMiddleEllipsisForDockTabLabel(GetDefault<UEditorStyleSettings>()->bEnableMiddleEllipsis);
 }
 
 void FStarshipEditorStyle::FStyle::SetupTutorialStyles()
diff --git a/Engine/Source/Editor/SceneOutliner/Private/ActorTreeItem.cpp b/Engine/Source/Editor/SceneOutliner/Private/ActorTreeItem.cpp
--- a/Engine/Source/Editor/SceneOutliner/Private/ActorTreeItem.cpp
+++ b/Engine/Source/Editor/SceneOutliner/Private/ActorTreeItem.cpp
@@ -22,6 +22,7 @@
 #include "WorldPartition/LoaderAdapter/LoaderAdapterPinnedActors.h"
 #include "ToolMenu.h"
 #include "Engine/Level.h"
+#include "Settings/EditorStyleSettings.h"
 
 #define LOCTEXT_NAMESPACE "SceneOutliner_ActorTreeItem"
 
@@ -41,6 +42,8 @@ struct SActorTreeLabel : FSceneOutlinerCommonLabelData, public SCompoundWidget
 
 		HighlightText = SceneOutliner.GetFilterHighlightText();
 
+		const bool bShouldUseMiddleEllipsis = GetDefault<UEditorStyleSettings>()->bEnableMiddleEllipsis;
+
 		TSharedPtr<SInlineEditableTextBlock> InlineTextBlock;
 
 		auto MainContent = SNew(SHorizontalBox)
@@ -59,6 +62,7 @@ struct SActorTreeLabel : FSceneOutlinerCommonLabelData, public SCompoundWidget
 				.OnEnterEditingMode(this, &SActorTreeLabel::OnEnterEditingMode)
 				.OnExitEditingMode(this, &SActorTreeLabel::OnExitEditingMode)
 				.IsSelected(FIsSelected::CreateSP(&InRow, &STableRow<FSceneOutlinerTreeItemPtr>::IsSelectedExclusively))
+				.OverflowPolicy(bShouldUseMiddleEllipsis ? ETextOverflowPolicy::MiddleEllipsis : TOptional<ETextOverflowPolicy>())
 				.IsReadOnly(this, &SActorTreeLabel::IsReadOnly)
 			];
 
diff --git a/Engine/Source/Editor/UnrealEd/Classes/Settings/EditorStyleSettings.h b/Engine/Source/Editor/UnrealEd/Classes/Settings/EditorStyleSettings.h
--- a/Engine/Source/Editor/UnrealEd/Classes/Settings/EditorStyleSettings.h
+++ b/Engine/Source/Editor/UnrealEd/Classes/Settings/EditorStyleSettings.h
@@ -166,6 +166,10 @@ public:
 	UPROPERTY(config)
 	FGuid CurrentAppliedTheme; 
 
+	/** If true MiddleEllipsis will be used by certain editor widget (ex: Tab Title/Outliner Actor Names), otherwise the default policy will be used instead */
+	UPROPERTY(EditAnywhere, config, Category=Text, meta = (ConfigRestartRequired = true))
+	bool bEnableMiddleEllipsis = true;
+
 public:
 
 	/**
diff --git a/Engine/Source/Editor/UnrealEd/Private/Toolkits/AssetEditorToolkit.cpp b/Engine/Source/Editor/UnrealEd/Private/Toolkits/AssetEditorToolkit.cpp
--- a/Engine/Source/Editor/UnrealEd/Private/Toolkits/AssetEditorToolkit.cpp
+++ b/Engine/Source/Editor/UnrealEd/Private/Toolkits/AssetEditorToolkit.cpp
@@ -136,6 +136,7 @@ void FAssetEditorToolkit::InitAssetEditor( const EToolkitMode::Type Mode, const
 			.ToolTip(IDocumentation::Get()->CreateToolTip(ToolTipText, nullptr, DocLink, GetToolkitFName().ToString()))
 			.IconColor(this, &FAssetEditorToolkit::GetDefaultTabColor)
 			.Label(Label)
+			.LabelOverflowPolicy(FGlobalTabmanager::Get()->GetShouldUseMiddleEllipsisForDockTabLabel() ? ETextOverflowPolicy::MiddleEllipsis : ETextOverflowPolicy::Ellipsis)
 			.LabelSuffix(LabelSuffix);
 		const TAttribute<const FSlateBrush*> TabIcon = TAttribute<const FSlateBrush*>::CreateSP(this, &FAssetEditorToolkit::GetDefaultTabIcon);
 		NewMajorTab->SetTabIcon(TabIcon);
diff --git a/Engine/Source/Runtime/Slate/Private/Framework/Docking/TabManager.cpp b/Engine/Source/Runtime/Slate/Private/Framework/Docking/TabManager.cpp
--- a/Engine/Source/Runtime/Slate/Private/Framework/Docking/TabManager.cpp
+++ b/Engine/Source/Runtime/Slate/Private/Framework/Docking/TabManager.cpp
@@ -2079,7 +2079,12 @@ TSharedPtr<SDockTab> FTabManager::SpawnTab(const FTabId& TabId, const TSharedPtr
 				PendingMainNonClosableTab = nullptr;
 				MainNonCloseableTabID = TabId;
 			}
-			
+
+			if (FGlobalTabmanager::Get()->GetShouldUseMiddleEllipsisForDockTabLabel())
+			{
+				NewTabWidget->SetTabLabelOverflowPolicy(ETextOverflowPolicy::MiddleEllipsis);
+			}
+
 			NewTabWidget->SetLayoutIdentifier(TabId);
 			NewTabWidget->ProvideDefaultLabel(Spawner->GetDisplayName().IsEmpty() ? FText::FromName(Spawner->TabType) : Spawner->GetDisplayName());
 			NewTabWidget->ProvideDefaultIcon(Spawner->GetIcon().GetIcon());
@@ -2590,6 +2595,16 @@ const TSharedRef<FGlobalTabmanager>& FGlobalTabmanager::Get()
 	return Instance;
 }
 
+bool FGlobalTabmanager::GetShouldUseMiddleEllipsisForDockTabLabel() const
+{
+	return bShouldUseMiddleEllipsisForDockTabLabel;
+}
+
+void FGlobalTabmanager::SetShouldUseMiddleEllipsisForDockTabLabel(const bool bInShouldUseMiddleEllipsis)
+{
+	bShouldUseMiddleEllipsisForDockTabLabel = bInShouldUseMiddleEllipsis;
+}
+
 FDelegateHandle FGlobalTabmanager::OnActiveTabChanged_Subscribe( const FOnActiveTabChanged::FDelegate& InDelegate )
 {
 	return OnActiveTabChanged.Add( InDelegate );
diff --git a/Engine/Source/Runtime/Slate/Private/Widgets/Docking/SDockTab.cpp b/Engine/Source/Runtime/Slate/Private/Widgets/Docking/SDockTab.cpp
--- a/Engine/Source/Runtime/Slate/Private/Widgets/Docking/SDockTab.cpp
+++ b/Engine/Source/Runtime/Slate/Private/Widgets/Docking/SDockTab.cpp
@@ -679,7 +679,7 @@ void SDockTab::Construct( const FArguments& InArgs )
 						SAssignNew(LabelWidget, STextBlock)
 						.TextStyle(&GetCurrentStyle().TabTextStyle)
 						.Text(this, &SDockTab::GetTabLabel)
-						.OverflowPolicy(ETextOverflowPolicy::Ellipsis)
+						.OverflowPolicy(InArgs._LabelOverflowPolicy)
 					]
 
 					// Tab Label Suffix
@@ -912,6 +912,14 @@ void SDockTab::SetTabLabelSuffix(const TAttribute<FText>& InTabLabelSuffix)
 	TabLabelSuffix = InTabLabelSuffix;
 }
 
+void SDockTab::SetTabLabelOverflowPolicy(ETextOverflowPolicy InOverflowPolicy) const
+{
+	if (LabelWidget.IsValid())
+	{
+		LabelWidget->SetOverflowPolicy(InOverflowPolicy);
+	}
+}
+
 const FSlateBrush* SDockTab::GetTabIcon() const
 {
 	return TabIcon.Get();
diff --git a/Engine/Source/Runtime/Slate/Public/Framework/Docking/TabManager.h b/Engine/Source/Runtime/Slate/Public/Framework/Docking/TabManager.h
--- a/Engine/Source/Runtime/Slate/Public/Framework/Docking/TabManager.h
+++ b/Engine/Source/Runtime/Slate/Public/Framework/Docking/TabManager.h
@@ -1380,6 +1380,12 @@ public:
 
 	SLATE_API void SetProxyTabManager(TSharedPtr<FProxyTabmanager> InProxyTabManager);
 
+	/** Get if the DockTab label should use the MiddleEllipsis overflow policy */
+	SLATE_API bool GetShouldUseMiddleEllipsisForDockTabLabel() const;
+
+	/** Set whether the DockTab label should use the MiddleEllipsis overflow policy */
+	SLATE_API void SetShouldUseMiddleEllipsisForDockTabLabel(const bool bInShouldUseMiddleEllipsis);
+
 	DECLARE_DELEGATE(FOnOverrideDockableAreaRestore);
 	/** Used to override dockable area restoration behavior */
 	FOnOverrideDockableAreaRestore OnOverrideDockableAreaRestore_Handler;
@@ -1412,7 +1418,9 @@ public:
 	SLATE_API TSharedPtr<FTabManager::FLayout> GetInitialLayoutSP();
 	
 private:
-	
+	/** Whether tabs label should use MiddleEllipsis OverflowPolicy*/
+	bool bShouldUseMiddleEllipsisForDockTabLabel = false;
+
 	/** Pairs of Major Tab and the TabManager that manages tabs within it. */
 	struct FSubTabManager
 	{
diff --git a/Engine/Source/Runtime/Slate/Public/Widgets/Docking/SDockTab.h b/Engine/Source/Runtime/Slate/Public/Widgets/Docking/SDockTab.h
--- a/Engine/Source/Runtime/Slate/Public/Widgets/Docking/SDockTab.h
+++ b/Engine/Source/Runtime/Slate/Public/Widgets/Docking/SDockTab.h
@@ -89,6 +89,7 @@ public:
 		, _TabColorScale(FLinearColor::Transparent)
 		, _ForegroundColor(FSlateColor::UseStyle())
 		, _IconColor()
+		, _LabelOverflowPolicy(ETextOverflowPolicy::Ellipsis)
 		{}
 
 		SLATE_DEFAULT_SLOT( FArguments, Content )
@@ -120,6 +121,7 @@ public:
 		SLATE_ATTRIBUTE( FLinearColor, TabColorScale )
 		SLATE_ATTRIBUTE( FSlateColor, ForegroundColor )
 		SLATE_ATTRIBUTE( FLinearColor, IconColor)
+		SLATE_ARGUMENT( ETextOverflowPolicy, LabelOverflowPolicy )
 	SLATE_END_ARGS()
 
 
@@ -219,6 +221,9 @@ public:
 	/** Set Label Suffix.  A second text field at the end of the Label that takes precedence and isn't lost when space is restricted */
 	SLATE_API void SetTabLabelSuffix(const TAttribute<FText>& InTabLabelSuffix);
 
+	/** Set the Label overflow policy */
+	SLATE_API void SetTabLabelOverflowPolicy(ETextOverflowPolicy InOverflowPolicy) const;
+
 	/** The tooltip text that appears on the tab. */
 	SLATE_API void SetTabToolTipWidget(TSharedPtr<SToolTip> InTabToolTipWidget);
 
