Update Entity Reference Picker to remove Level and Simulation entities while allowing a soft interactable filter for entries we still want visible but un-selectable

diff --git a/Engine/Plugins/Experimental/EditorDataStorageFeatures/Source/TedsOutliner/Private/TedsOutlinerHierarchy.cpp b/Engine/Plugins/Experimental/EditorDataStorageFeatures/Source/TedsOutliner/Private/TedsOutlinerHierarchy.cpp
--- a/Engine/Plugins/Experimental/EditorDataStorageFeatures/Source/TedsOutliner/Private/TedsOutlinerHierarchy.cpp
+++ b/Engine/Plugins/Experimental/EditorDataStorageFeatures/Source/TedsOutliner/Private/TedsOutlinerHierarchy.cpp
@@ -70,7 +70,7 @@ FSceneOutlinerTreeItemPtr FTedsOutlinerHierarchy::FindOrCreateParentItem(const I
 	}
 	else if(bCreate)
 	{
-		return Mode->CreateItemFor<FTedsOutlinerTreeItem>(FTedsOutlinerTreeItem(ParentRowHandle, TedsOutlinerImpl), true);
+		return Mode->CreateItemFor<FTedsOutlinerTreeItem>(FTedsOutlinerTreeItem(ParentRowHandle, TedsOutlinerImpl), TedsOutlinerImpl->ShouldForceShowParentRows());
 	}
 	
 	return nullptr;
diff --git a/Engine/Plugins/Experimental/EditorDataStorageFeatures/Source/TedsOutliner/Private/TedsOutlinerImpl.cpp b/Engine/Plugins/Experimental/EditorDataStorageFeatures/Source/TedsOutliner/Private/TedsOutlinerImpl.cpp
--- a/Engine/Plugins/Experimental/EditorDataStorageFeatures/Source/TedsOutliner/Private/TedsOutlinerImpl.cpp
+++ b/Engine/Plugins/Experimental/EditorDataStorageFeatures/Source/TedsOutliner/Private/TedsOutlinerImpl.cpp
@@ -85,6 +85,7 @@ FTedsOutlinerImpl::FTedsOutlinerImpl(const FTedsOutlinerParams& InParams, IScene
 	, InitialQueryDescription(InParams.QueryDescription)
 	, HierarchyData(InParams.HierarchyData)
 	, SelectionSetName(InParams.SelectionSetOverride)
+	, bForceShowParents(InParams.bForceShowParents)
 	, SceneOutlinerMode(InMode)
 	, SceneOutliner(InParams.SceneOutliner)
 {
@@ -186,7 +187,7 @@ void FTedsOutlinerImpl::SetSelection(const TArray<DataStorage::RowHandle>& InSel
 	}
 }
 
-TSharedRef<SWidget> FTedsOutlinerImpl::CreateLabelWidgetForItem(DataStorage::RowHandle InRowHandle, const STableRow<FSceneOutlinerTreeItemPtr>& InRow) const
+TSharedRef<SWidget> FTedsOutlinerImpl::CreateLabelWidgetForItem(DataStorage::RowHandle InRowHandle, const STableRow<FSceneOutlinerTreeItemPtr>& InRow, bool bIsInteractable) const
 {
 	// Find the best matching label widget for a given list of columns
 	auto CreateWidgetConstructor = [this](TArray<TWeakObjectPtr<const UScriptStruct>>& ColumnTypes) -> TSharedPtr<FTypedElementWidgetConstructor>
@@ -222,13 +223,15 @@ TSharedRef<SWidget> FTedsOutlinerImpl::CreateLabelWidgetForItem(DataStorage::Row
 	
 	using namespace DataStorage::Queries;
 	// Query description to pass as metadata to allow the label column to be writable
-	static FQueryDescription MetaDataQueryDescription = Select().ReadWrite<FTypedElementLabelColumn>().Where().Compile();
+	static FQueryDescription MetaDataQueryReadWrite = Select().ReadWrite<FTypedElementLabelColumn>().Where().Compile();
+	static FQueryDescription MetaDataQueryRead = Select().ReadOnly<FTypedElementLabelColumn>().Where().Compile();
 
 	// Create a widget for this row using the given widget constructor
-	auto CreateWidget = [InRowHandle, this, &InRow](const TSharedPtr<FTypedElementWidgetConstructor>& WidgetConstructor) -> TSharedPtr<SWidget>
+	auto CreateWidget = [InRowHandle, this, &InRow, bIsInteractable](const TSharedPtr<FTypedElementWidgetConstructor>& WidgetConstructor) -> TSharedPtr<SWidget>
 	{
 		// Create metadata for the query
-		FQueryMetaDataView QueryMetaDataView(MetaDataQueryDescription);
+		FQueryMetaDataView QueryMetaDataView = bIsInteractable ?
+			FQueryMetaDataView(MetaDataQueryReadWrite) : FQueryMetaDataView(MetaDataQueryRead);
 
 		RowHandle UiRowHandle = Storage->AddRow(Storage->FindTable(DataStorage::TableViewerUtils::GetWidgetTableName()));
 
@@ -494,6 +497,11 @@ DataStorage::RowHandle FTedsOutlinerImpl::GetParentRow(DataStorage::RowHandle In
 	return ParentRowHandle;
 }
 
+bool FTedsOutlinerImpl::ShouldForceShowParentRows()
+{
+	return bForceShowParents;
+}
+
 void FTedsOutlinerImpl::OnItemAdded(DataStorage::RowHandle ItemRowHandle)
 {
 	if (!CanDisplayRow(ItemRowHandle))
diff --git a/Engine/Plugins/Experimental/EditorDataStorageFeatures/Source/TedsOutliner/Private/TedsOutlinerItem.cpp b/Engine/Plugins/Experimental/EditorDataStorageFeatures/Source/TedsOutliner/Private/TedsOutlinerItem.cpp
--- a/Engine/Plugins/Experimental/EditorDataStorageFeatures/Source/TedsOutliner/Private/TedsOutlinerItem.cpp
+++ b/Engine/Plugins/Experimental/EditorDataStorageFeatures/Source/TedsOutliner/Private/TedsOutlinerItem.cpp
@@ -58,13 +58,13 @@ FString FTedsOutlinerTreeItem::GetDisplayString() const
 
 bool FTedsOutlinerTreeItem::CanInteract() const
 {
-	return true; // TEDS-Outliner TODO: check item constness from TEDS maybe?
+	return Flags.bInteractive;
 }
 
 TSharedRef<SWidget> FTedsOutlinerTreeItem::GenerateLabelWidget(ISceneOutliner& Outliner,
 	const STableRow<FSceneOutlinerTreeItemPtr>& InRow)
 {
-	return TedsOutlinerImpl->CreateLabelWidgetForItem(RowHandle, InRow);
+	return TedsOutlinerImpl->CreateLabelWidgetForItem(RowHandle, InRow, CanInteract());
 }
 
 void FTedsOutlinerTreeItem::GenerateContextMenu(UToolMenu* Menu, SSceneOutliner& Outliner)
diff --git a/Engine/Plugins/Experimental/EditorDataStorageFeatures/Source/TedsOutliner/Public/TedsOutlinerImpl.h b/Engine/Plugins/Experimental/EditorDataStorageFeatures/Source/TedsOutliner/Public/TedsOutlinerImpl.h
--- a/Engine/Plugins/Experimental/EditorDataStorageFeatures/Source/TedsOutliner/Public/TedsOutlinerImpl.h
+++ b/Engine/Plugins/Experimental/EditorDataStorageFeatures/Source/TedsOutliner/Public/TedsOutlinerImpl.h
@@ -74,6 +74,7 @@ struct FTedsOutlinerParams
 	, QueryDescription()
 	, bUseDefaultTedsFilters(false)
 	, bShowRowHandleColumn(true)
+	, bForceShowParents(true)
 	, HierarchyData(FTedsOutlinerHierarchyData::GetDefaultHierarchyData())
 	, CellWidgetPurposes{TEXT("SceneOutliner.Cell"), TEXT("General.Cell")}
 	{}
@@ -94,6 +95,9 @@ struct FTedsOutlinerParams
 	// If true, this Outliner will include a column for row handle
 	bool bShowRowHandleColumn;
 
+	// If true, parent nodes will be remain visible if a child passes all filters even if the parent fails a filter
+	bool bForceShowParents;
+
 	// If specified, this is how the TEDS Outliner will handle hierarchies. If not specified - there will be no hierarchies shown as a
 	// parent-child relation in the tree view
 	TOptional<FTedsOutlinerHierarchyData> HierarchyData;
@@ -141,7 +145,7 @@ public:
 	void SetSelection(const TArray<DataStorage::RowHandle>& InSelectedRows);
 
 	// Helper function to create a label widget for a given row
-	TSharedRef<SWidget> CreateLabelWidgetForItem(DataStorage::RowHandle InRowHandle, const STableRow<FSceneOutlinerTreeItemPtr>& InRow) const;
+	TSharedRef<SWidget> CreateLabelWidgetForItem(DataStorage::RowHandle InRowHandle, const STableRow<FSceneOutlinerTreeItemPtr>& InRow, bool bIsInteractable = true) const;
 
 	// Get the hierarchy data associated with this table viewer
 	const TOptional<FTedsOutlinerHierarchyData>& GetHierarchyData();
@@ -164,9 +168,12 @@ public:
 	// Get the parent row for a given row
 	DataStorage::RowHandle GetParentRow(DataStorage::RowHandle InRowHandle);
 
+	bool ShouldForceShowParentRows();
+
 	// Recompile all queries used by this table viewer
 	void RecompileQueries();
 
+
 protected:
 	
 	void OnItemAdded(DataStorage::RowHandle ItemRowHandle);
@@ -219,6 +226,9 @@ protected:
 	
 	TOptional<FName> SelectionSetName;
 	bool bSelectionDirty = false;
+
+	// If true, parent nodes will remain visible if a child passes all filters even if the parent fails a filter
+	bool bForceShowParents = true;
 	
 	// Ticker for selection updates so we don't fire the delegate multiple times in one frame for multi select
 	FTSTicker::FDelegateHandle TickerHandle;
diff --git a/Engine/Plugins/Experimental/EditorDataStorageFeatures/Source/TedsOutliner/Public/TedsOutlinerItem.h b/Engine/Plugins/Experimental/EditorDataStorageFeatures/Source/TedsOutliner/Public/TedsOutlinerItem.h
--- a/Engine/Plugins/Experimental/EditorDataStorageFeatures/Source/TedsOutliner/Public/TedsOutlinerItem.h
+++ b/Engine/Plugins/Experimental/EditorDataStorageFeatures/Source/TedsOutliner/Public/TedsOutlinerItem.h
@@ -20,12 +20,18 @@ struct FTedsOutlinerTreeItem : ISceneOutlinerTreeItem
 public:
 	
 	DECLARE_DELEGATE_RetVal_OneParam(bool, FFilterPredicate, const DataStorage::RowHandle);
+	DECLARE_DELEGATE_RetVal_OneParam(bool, FInteractivePredicate, const DataStorage::RowHandle);
 
 	TEDSOUTLINER_API bool Filter(FFilterPredicate Pred) const
 	{
 		return Pred.Execute(RowHandle);
 	}
 
+	TEDSOUTLINER_API bool GetInteractiveState(FInteractivePredicate Pred) const
+	{
+		return Pred.Execute(RowHandle);
+	}
+
 	TEDSOUTLINER_API FTedsOutlinerTreeItem(const DataStorage::RowHandle& InRowHandle, const TSharedRef<const FTedsOutlinerImpl>& InTedsOutlinerImpl);
 
 	/* Begin ISceneOutlinerTreeItem Implementation */
diff --git a/Engine/Plugins/Experimental/EditorDataStorageFeatures/Source/TedsPropertyEditor/Public/TedsRowPickingMode.h b/Engine/Plugins/Experimental/EditorDataStorageFeatures/Source/TedsPropertyEditor/Public/TedsRowPickingMode.h
--- a/Engine/Plugins/Experimental/EditorDataStorageFeatures/Source/TedsPropertyEditor/Public/TedsRowPickingMode.h
+++ b/Engine/Plugins/Experimental/EditorDataStorageFeatures/Source/TedsPropertyEditor/Public/TedsRowPickingMode.h
@@ -8,6 +8,7 @@
 
 DECLARE_DELEGATE_OneParam(FOnTedsRowSelected, UE::Editor::DataStorage::RowHandle RowHandle);
 DECLARE_DELEGATE_RetVal_OneParam(bool, FOnShouldFilterTedsRow, const UE::Editor::DataStorage::RowHandle RowHandle);
+DECLARE_DELEGATE_RetVal_OneParam(bool, FOnShouldInteractTedsRow, const UE::Editor::DataStorage::RowHandle RowHandle);
 
 /*
 * Picking mode for TEDs Scene Outliner Widgets. Based off of FActorPickingMode
diff --git a/Engine/Plugins/Experimental/EditorDataStorageFeatures/Source/TedsPropertyEditor/Public/Widgets/SPropertyMenuTedsRowPicker.h b/Engine/Plugins/Experimental/EditorDataStorageFeatures/Source/TedsPropertyEditor/Public/Widgets/SPropertyMenuTedsRowPicker.h
--- a/Engine/Plugins/Experimental/EditorDataStorageFeatures/Source/TedsPropertyEditor/Public/Widgets/SPropertyMenuTedsRowPicker.h
+++ b/Engine/Plugins/Experimental/EditorDataStorageFeatures/Source/TedsPropertyEditor/Public/Widgets/SPropertyMenuTedsRowPicker.h
@@ -11,10 +11,12 @@ public:
 	SLATE_BEGIN_ARGS(SPropertyMenuTedsRowPicker)
 		: _AllowClear(true)
 		, _ElementFilter()
+		, _InteractiveFilter()
 	{}
 		SLATE_ARGUMENT(bool, AllowClear)
 		SLATE_ARGUMENT(UE::Editor::DataStorage::FQueryDescription, QueryFilter)
 		SLATE_ARGUMENT(FOnShouldFilterTedsRow, ElementFilter)
+		SLATE_ARGUMENT(FOnShouldInteractTedsRow, InteractiveFilter)
 		SLATE_EVENT(FOnTedsRowSelected, OnSet)
 		SLATE_EVENT(FSimpleDelegate, OnClose)
 	SLATE_END_ARGS()
@@ -36,6 +38,8 @@ private:
 
 	FOnShouldFilterTedsRow ElementFilter;
 
+	FOnShouldInteractTedsRow InteractiveFilter;
+
 	FOnTedsRowSelected OnSet;
 
 	FSimpleDelegate OnClose;
