Add missing file for Entity Reference Picker filter fix

diff --git a/Engine/Plugins/Experimental/EditorDataStorageFeatures/Source/TedsPropertyEditor/Private/Widgets/SPropertyMenuTedsRowPicker.cpp b/Engine/Plugins/Experimental/EditorDataStorageFeatures/Source/TedsPropertyEditor/Private/Widgets/SPropertyMenuTedsRowPicker.cpp
--- a/Engine/Plugins/Experimental/EditorDataStorageFeatures/Source/TedsPropertyEditor/Private/Widgets/SPropertyMenuTedsRowPicker.cpp
+++ b/Engine/Plugins/Experimental/EditorDataStorageFeatures/Source/TedsPropertyEditor/Private/Widgets/SPropertyMenuTedsRowPicker.cpp
@@ -13,6 +13,7 @@ void SPropertyMenuTedsRowPicker::Construct(const FArguments& InArgs)
 	bAllowClear = InArgs._AllowClear;
 	QueryFilter = InArgs._QueryFilter;
 	ElementFilter = InArgs._ElementFilter;
+	InteractiveFilter = InArgs._InteractiveFilter;
 	OnSet = InArgs._OnSet;
 
 	FMenuBuilder MenuBuilder(true, nullptr);
@@ -70,21 +71,30 @@ void SPropertyMenuTedsRowPicker::Construct(const FArguments& InArgs)
 					{
 						FTedsOutlinerParams Params(Outliner);
 						Params.QueryDescription = QueryFilter;
+						Params.bForceShowParents = false;
 						return new FTedsRowPickingMode(Params, OnItemPicked);
 					});
 
-				TSharedPtr<SSceneOutliner> Outliner = SNew(SSceneOutliner, InitOptions);
-
-				Outliner->AddFilter(
+				InitOptions.Filters->Add(
 					MakeShared<TSceneOutlinerPredicateFilter<FTedsOutlinerTreeItem>>(
 						FTedsOutlinerTreeItem::FFilterPredicate::CreateLambda([this](const UE::Editor::DataStorage::RowHandle RowHandle) -> bool
-					{
-						if (ElementFilter.IsBound())
-						{
-							return ElementFilter.Execute(RowHandle);
-						}
-						return true;
-					}), FSceneOutlinerFilter::EDefaultBehaviour::Pass));
+							{
+								if (ElementFilter.IsBound())
+								{
+									return ElementFilter.Execute(RowHandle);
+								}
+								return true;
+							}), FSceneOutlinerFilter::EDefaultBehaviour::Pass,
+						FTedsOutlinerTreeItem::FInteractivePredicate::CreateLambda([this](const UE::Editor::DataStorage::RowHandle RowHandle) -> bool
+							{
+								if (InteractiveFilter.IsBound())
+								{
+									return InteractiveFilter.Execute(RowHandle);
+								}
+								return true;
+							})));
+
+				TSharedPtr<SSceneOutliner> Outliner = SNew(SSceneOutliner, InitOptions);
 
 				MenuContent =
 					SNew(SBox)
