Relocated FUIHeaderPropertiesColumn and renamed to FHeaderWidgetSizeColumn.

diff --git a/Engine/Plugins/Experimental/Animation/HierarchyTable/Source/Editor/HierarchyTableEditor.Build.cs b/Engine/Plugins/Experimental/Animation/HierarchyTable/Source/Editor/HierarchyTableEditor.Build.cs
--- a/Engine/Plugins/Experimental/Animation/HierarchyTable/Source/Editor/HierarchyTableEditor.Build.cs
+++ b/Engine/Plugins/Experimental/Animation/HierarchyTable/Source/Editor/HierarchyTableEditor.Build.cs
@@ -21,6 +21,7 @@ public class HierarchyTableEditor : ModuleRules
 			"TypedElementFramework",
 			"SceneOutliner",
 			"ToolMenus",
+			"ToolWidgets",
 			"TedsTableViewer"
 		);
 	}
diff --git a/Engine/Plugins/Experimental/Animation/HierarchyTable/Source/Editor/Private/HierarchyTable/Widgets/OverrideWidgetConstructor.cpp b/Engine/Plugins/Experimental/Animation/HierarchyTable/Source/Editor/Private/HierarchyTable/Widgets/OverrideWidgetConstructor.cpp
--- a/Engine/Plugins/Experimental/Animation/HierarchyTable/Source/Editor/Private/HierarchyTable/Widgets/OverrideWidgetConstructor.cpp
+++ b/Engine/Plugins/Experimental/Animation/HierarchyTable/Source/Editor/Private/HierarchyTable/Widgets/OverrideWidgetConstructor.cpp
@@ -8,7 +8,7 @@
 #include "Widgets/Layout/SBox.h"
 #include "Widgets/SBoxPanel.h"
 #include "HierarchyTable.h"
-#include "Columns/UIPropertiesColumns.h"
+#include "Columns/SlateHeaderColumns.h"
 #include "ScopedTransaction.h"
 
 #define LOCTEXT_NAMESPACE "FTypedElementWidgetConstructor_Override"
@@ -137,7 +137,7 @@ TSharedPtr<SWidget> FTypedElementWidgetHeaderConstructor_Override::CreateWidget(
 bool FTypedElementWidgetHeaderConstructor_Override::FinalizeWidget(IEditorDataStorageProvider* DataStorage, 
 	IEditorDataStorageUiProvider* DataStorageUi, UE::Editor::DataStorage::RowHandle Row, const TSharedPtr<SWidget>& Widget)
 {
-	DataStorage->AddColumn(Row, FUIHeaderPropertiesColumn
+	DataStorage->AddColumn(Row, FHeaderWidgetSizeColumn
 		{
 			.ColumnSizeMode = EColumnSizeMode::Fixed,
 			.Width = 24.0f
diff --git a/Engine/Plugins/Experimental/EditorDataStorageFeatures/Source/TedsOutliner/Private/Widgets/AlertWidget.cpp b/Engine/Plugins/Experimental/EditorDataStorageFeatures/Source/TedsOutliner/Private/Widgets/AlertWidget.cpp
--- a/Engine/Plugins/Experimental/EditorDataStorageFeatures/Source/TedsOutliner/Private/Widgets/AlertWidget.cpp
+++ b/Engine/Plugins/Experimental/EditorDataStorageFeatures/Source/TedsOutliner/Private/Widgets/AlertWidget.cpp
@@ -2,7 +2,7 @@
 
 #include "Widgets/AlertWidget.h"
 
-#include "Columns/UIPropertiesColumns.h"
+#include "Columns/SlateHeaderColumns.h"
 #include "Elements/Columns/TypedElementAlertColumns.h"
 #include "Elements/Columns/TypedElementMiscColumns.h"
 #include "Elements/Columns/TypedElementSlateWidgetColumns.h"
@@ -420,7 +420,7 @@ TConstArrayView<const UScriptStruct*> FAlertHeaderWidgetConstructor::GetAddition
 bool FAlertHeaderWidgetConstructor::FinalizeWidget(IEditorDataStorageProvider* DataStorage, 
 	IEditorDataStorageUiProvider* DataStorageUi, UE::Editor::DataStorage::RowHandle Row, const TSharedPtr<SWidget>& Widget)
 {
-	DataStorage->AddColumn(Row, FUIHeaderPropertiesColumn
+	DataStorage->AddColumn(Row, FHeaderWidgetSizeColumn
 		{
 			.ColumnSizeMode = EColumnSizeMode::Fixed,
 			.Width = 24.0f
diff --git a/Engine/Plugins/Experimental/EditorDataStorageFeatures/Source/TedsTableViewer/Private/TedsTableViewerColumn.cpp b/Engine/Plugins/Experimental/EditorDataStorageFeatures/Source/TedsTableViewer/Private/TedsTableViewerColumn.cpp
--- a/Engine/Plugins/Experimental/EditorDataStorageFeatures/Source/TedsTableViewer/Private/TedsTableViewerColumn.cpp
+++ b/Engine/Plugins/Experimental/EditorDataStorageFeatures/Source/TedsTableViewer/Private/TedsTableViewerColumn.cpp
@@ -2,7 +2,7 @@
 
 #include "TedsTableViewerColumn.h"
 
-#include "Columns/UIPropertiesColumns.h"
+#include "Columns/SlateHeaderColumns.h"
 #include "Elements/Columns/TypedElementMiscColumns.h"
 #include "Elements/Columns/TypedElementSlateWidgetColumns.h"
 #include "Elements/Columns/TypedElementTypeInfoColumns.h"
@@ -129,7 +129,7 @@ namespace UE::Editor::DataStorage
 					]
 			];
 		
-		if (const FUIHeaderPropertiesColumn* HeaderProperties = Storage->GetColumn<FUIHeaderPropertiesColumn>(UiRowHandle))
+		if (const FHeaderWidgetSizeColumn* HeaderProperties = Storage->GetColumn<FHeaderWidgetSizeColumn>(UiRowHandle))
 		{
 			float Width = HeaderProperties->Width;
 			switch (HeaderProperties->ColumnSizeMode)
diff --git a/Engine/Plugins/Experimental/EditorDataStorageFeatures/Source/TedsTableViewer/Private/TedsTableViewerModule.cpp b/Engine/Plugins/Experimental/EditorDataStorageFeatures/Source/TedsTableViewer/Private/TedsTableViewerModule.cpp
--- a/Engine/Plugins/Experimental/EditorDataStorageFeatures/Source/TedsTableViewer/Private/TedsTableViewerModule.cpp
+++ b/Engine/Plugins/Experimental/EditorDataStorageFeatures/Source/TedsTableViewer/Private/TedsTableViewerModule.cpp
@@ -1,24 +1,7 @@
 // Copyright Epic Games, Inc. All Rights Reserved.
 
-#include "Elements/Columns/TypedElementAlertColumns.h"
-#include "Elements/Columns/TypedElementCompatibilityColumns.h"
-#include "Elements/Columns/TypedElementLabelColumns.h"
-#include "Elements/Common/EditorDataStorageFeatures.h"
-#include "Elements/Framework/TypedElementQueryBuilder.h"
-#include "Elements/Interfaces/TypedElementDataStorageInterface.h"
-#include "LevelEditor.h"
 #include "Modules/ModuleManager.h"
 #include "Modules/ModuleInterface.h"
-#include "Widgets/STedsTableViewer.h"
-#include "Widgets/Docking/SDockTab.h"
-#include "WorkspaceMenuStructureModule.h"
-#include "WorkspaceMenuStructure.h"
-#include "Elements/Columns/TypedElementSelectionColumns.h"
-#include "QueryStack/FQueryStackNode_RowView.h"
-
-#define LOCTEXT_NAMESPACE "TedsTableViewerModule"
-
-static FName TableViewerTabName(TEXT("TedsTableViewer"));
 
 class FTedsTableViewerModule : public IModuleInterface
 {
@@ -26,99 +9,12 @@ public:
 
 	virtual void StartupModule() override
 	{
-		RegisterTableViewerTab();
 	}
 
 	virtual void ShutdownModule() override
 	{
 	}
-
-	void RegisterTableViewerTab()
-	{
-		FLevelEditorModule& LevelEditorModule = FModuleManager::GetModuleChecked<FLevelEditorModule>("LevelEditor");
-		
-		LevelEditorTabManagerChangedHandle = LevelEditorModule.OnTabManagerChanged().AddLambda([this]()
-		{
-			FLevelEditorModule& LevelEditorModule = FModuleManager::GetModuleChecked<FLevelEditorModule>("LevelEditor");
-
-			TSharedPtr<FTabManager> LevelEditorTabManager = LevelEditorModule.GetLevelEditorTabManager();
-
-			LevelEditorTabManager->RegisterTabSpawner(TableViewerTabName, FOnSpawnTab::CreateRaw(this, &FTedsTableViewerModule::OpenTableViewer))
-			.SetDisplayName(LOCTEXT("TedsTableVIewerTitle", "Table Viewer (Experimental)"))
-			.SetGroup(WorkspaceMenu::GetMenuStructure().GetLevelEditorOutlinerCategory())
-			.SetIcon(FSlateIcon(FAppStyle::GetAppStyleSetName(), "LevelEditor.Tabs.Outliner"))
-			.SetAutoGenerateMenuEntry(false); // This can only be summoned from the Cvar now
-	
-		});
-	}
-	
-	
-	TSharedRef<SDockTab> OpenTableViewer(const FSpawnTabArgs& SpawnTabArgs)
-	{
-		return SNew(SDockTab)
-		.TabRole(ETabRole::NomadTab)
-		[
-			CreateTableViewer()
-		];
-	}
-	
-	TSharedRef<SWidget> CreateTableViewer()
-	{
-		using namespace UE::Editor::DataStorage;
-
-		if(!AreEditorDataStorageFeaturesEnabled())
-		{
-			return SNullWidget::NullWidget;
-		}
-
-		IEditorDataStorageProvider* DataStorage = GetMutableDataStorageFeature<IEditorDataStorageProvider>(StorageFeatureName);
-
-		using namespace UE::Editor::DataStorage::Queries;
-
-		// We'll just create a test table viewer that views all actor rows without actually updating it dynamically for now
-		static QueryHandle QueryHandle =
-			DataStorage->RegisterQuery(
-				Select()
-				.Where()
-					.All<FTypedElementActorTag>()
-				.Compile());
-
-		Rows.Empty();
-		
-		FQueryResult QueryResult = DataStorage->RunQuery(QueryHandle,
-			CreateDirectQueryCallbackBinding([this](const IEditorDataStorageProvider::IDirectQueryContext& Context, const RowHandle* RowHandles)
-			{
-				Rows.Append(RowHandles, Context.GetRowCount());
-			}));
-
-		return SNew(STedsTableViewer)
-			.QueryStack(MakeShared<FQueryStackNode_RowView>(&Rows))
-			.Columns({FTypedElementLabelColumn::StaticStruct(), FTypedElementSelectionColumn::StaticStruct(),
-				FTypedElementAlertColumn::StaticStruct(), FTypedElementChildAlertColumn::StaticStruct()});
-	}
-	
-	FDelegateHandle LevelEditorTabManagerChangedHandle;
-	TArray<UE::Editor::DataStorage::RowHandle> Rows;
 };
 
-// CVar to summon a test table viewer that views a snapshot of all actors at the moment when the Cvar is used
-static FAutoConsoleCommand OpenTableViewerConsoleCommand(
-	TEXT("TEDS.UI.OpenTableViewer"),
-	TEXT("Spawn the test TEDS Table Viewer."),
-	FConsoleCommandWithArgsDelegate::CreateLambda([](const TArray<FString>& Args)
-	{
-		FLevelEditorModule& LevelEditorModule = FModuleManager::GetModuleChecked<FLevelEditorModule>("LevelEditor");
-
-		TSharedPtr<FTabManager> LevelEditorTabManager = LevelEditorModule.GetLevelEditorTabManager();
-		if(TSharedPtr<SDockTab> DockTab = LevelEditorTabManager->FindExistingLiveTab(TableViewerTabName))
-		{
-			DockTab->RequestCloseTab();
-		}
-
-		LevelEditorTabManager->TryInvokeTab(TableViewerTabName);
-	}));
-
 IMPLEMENT_MODULE(FTedsTableViewerModule, TedsTableViewer);
 
-#undef LOCTEXT_NAMESPACE //"TedsTableViewerModule"
-
diff --git a/Engine/Plugins/Experimental/EditorDataStorageFeatures/Source/TedsTableViewer/Public/Columns/UIPropertiesColumns.h b/Engine/Plugins/Experimental/EditorDataStorageFeatures/Source/TedsTableViewer/Public/Columns/UIPropertiesColumns.h
deleted file mode 100644
diff --git a/Engine/Source/Developer/ToolWidgets/Public/Columns/SlateHeaderColumns.h b/Engine/Source/Developer/ToolWidgets/Public/Columns/SlateHeaderColumns.h
new file mode 100644
--- /dev/null
+++ b/Engine/Source/Developer/ToolWidgets/Public/Columns/SlateHeaderColumns.h
@@ -0,0 +1,28 @@
+// Copyright Epic Games, Inc. All Rights Reserved.
+
+#pragma once
+
+#include "Elements/Interfaces/TypedElementDataStorageInterface.h"
+#include "Widgets/Views/SHeaderRow.h"
+
+#include "SlateHeaderColumns.generated.h"
+
+/**
+ * Column added to a header widget row to control the sizing when the default behavior is incorrect.
+ */
+USTRUCT(meta = (DisplayName = "Header widget with custom size"))
+struct FHeaderWidgetSizeColumn final : public FEditorDataStorageColumn
+{
+	GENERATED_BODY()
+
+	EColumnSizeMode::Type ColumnSizeMode;
+
+	/*
+	 * Fill: Column stretches to this fraction of the header row
+	 * Fixed: Column is fixed at this width in slate units and cannot be resized
+	 * Manual: Column defaults to this width in slate units and can be user-sized
+	 * FillSized: Column stretches as Fill but is initialized with this width in slate units
+	 */
+	UPROPERTY(meta = (IgnoreForMemberInitializationTest))
+	float Width;
+};
\ No newline at end of file
