[SceneOutliner]
* Implement Visibility column for Actors and Entities in TEDSOutliner
- Created new WidgetFactories and Queries to control visibility through TEDS
- Implement EntityProxyActor to pass visibility state to entity and it's componnts.

diff --git a/Engine/Plugins/Experimental/EditorDataStorage/Source/TedsCore/Private/Compatibility/TypedElementActorVisibilityQueries.cpp b/Engine/Plugins/Experimental/EditorDataStorage/Source/TedsCore/Private/Compatibility/TypedElementActorVisibilityQueries.cpp
new file mode 100644
--- /dev/null
+++ b/Engine/Plugins/Experimental/EditorDataStorage/Source/TedsCore/Private/Compatibility/TypedElementActorVisibilityQueries.cpp
@@ -0,0 +1,81 @@
+#pragma once
+
+#include "Compatibility/TypedElementActorVisibilityQueries.h"
+
+#include "Elements/Columns/TypedElementCompatibilityColumns.h"
+#include "Elements/Columns/TypedElementMiscColumns.h"
+#include "Elements/Columns/TypedElementVisibilityColumns.h"
+#include "Elements/Common/TypedElementQueryTypes.h"
+#include "Elements/Framework/TypedElementQueryBuilder.h"
+#include "GameFramework/Actor.h"
+
+void UActorVisibilityDataStorageFactory::RegisterQueries(IEditorDataStorageProvider& DataStorage)
+{
+	RegisterActorAddVisibilityColumn(DataStorage);
+	RegisterActorVisibilityToColumnQuery(DataStorage);
+	RegisterVisibilityColumnToActorQuery(DataStorage);
+}
+
+void UActorVisibilityDataStorageFactory::RegisterActorAddVisibilityColumn(IEditorDataStorageProvider& DataStorage) const
+{
+	using namespace UE::Editor::DataStorage::Queries;
+
+	DataStorage.RegisterQuery(
+		Select(
+			TEXT("Actor Visibility Object to New Column"),
+			FProcessor(EQueryTickPhase::PrePhysics, DataStorage.GetQueryTickGroupName(EQueryTickGroups::SyncExternalToDataStorage)).SetExecutionMode(EExecutionMode::GameThread),
+			[](IQueryContext& Context, RowHandle Row, const FTypedElementUObjectColumn& Object)
+			{
+				if (const AActor* ActorInstance = Cast<AActor>(Object.Object); ActorInstance != nullptr)
+				{
+					Context.AddColumn(Row, FVisibleInEditorColumn{ .bIsVisibleInEditor = !ActorInstance->IsHiddenEd() });
+				}
+			})
+		.Where()
+			.All<FTypedElementActorTag, FTypedElementSyncFromWorldTag>()
+			.None<FVisibleInEditorColumn>()
+		.Compile()
+	);
+}
+
+void UActorVisibilityDataStorageFactory::RegisterActorVisibilityToColumnQuery(IEditorDataStorageProvider& DataStorage) const
+{
+	using namespace UE::Editor::DataStorage::Queries;
+
+	DataStorage.RegisterQuery(
+		Select(
+			TEXT("Actor Visibility Object to Column"),
+			FProcessor(EQueryTickPhase::PrePhysics, DataStorage.GetQueryTickGroupName(EQueryTickGroups::SyncExternalToDataStorage)).SetExecutionMode(EExecutionMode::GameThread),
+			[](IQueryContext& Context, RowHandle Row, const FTypedElementUObjectColumn& Object, FVisibleInEditorColumn& VisibilityColumn)
+			{
+				if (const AActor* ActorInstance = Cast<AActor>(Object.Object); ActorInstance != nullptr)
+				{
+					VisibilityColumn.bIsVisibleInEditor = !ActorInstance->IsHiddenEd();
+				}
+			})
+		.Where()
+			.All<FTypedElementActorTag, FTypedElementSyncFromWorldTag>()
+		.Compile()
+	);
+}
+
+void UActorVisibilityDataStorageFactory::RegisterVisibilityColumnToActorQuery(IEditorDataStorageProvider& DataStorage) const
+{
+	using namespace UE::Editor::DataStorage::Queries;
+
+	DataStorage.RegisterQuery(
+		Select(
+			TEXT("Actor Visibility Column to Object"),
+			FProcessor(EQueryTickPhase::PrePhysics, DataStorage.GetQueryTickGroupName(EQueryTickGroups::SyncDataStorageToExternal)).SetExecutionMode(EExecutionMode::GameThread),
+			[&](IQueryContext& Context, RowHandle Row, FTypedElementUObjectColumn& Object, const FVisibleInEditorColumn& VisibilityColumn)
+			{
+				if (AActor* ActorInstance = Cast<AActor>(Object.Object); ActorInstance != nullptr)
+				{
+					ActorInstance->SetIsTemporarilyHiddenInEditor(!VisibilityColumn.bIsVisibleInEditor);
+				}
+			})
+		.Where()
+			.All<FTypedElementActorTag, FTypedElementSyncBackToWorldTag>()
+		.Compile()
+	);
+}
diff --git a/Engine/Plugins/Experimental/EditorDataStorage/Source/TedsCore/Private/Compatibility/TypedElementActorVisibilityQueries.h b/Engine/Plugins/Experimental/EditorDataStorage/Source/TedsCore/Private/Compatibility/TypedElementActorVisibilityQueries.h
new file mode 100644
--- /dev/null
+++ b/Engine/Plugins/Experimental/EditorDataStorage/Source/TedsCore/Private/Compatibility/TypedElementActorVisibilityQueries.h
@@ -0,0 +1,38 @@
+// Copyright Epic Games, Inc. All Rights Reserved.
+
+#pragma once
+
+#include "Elements/Interfaces/TypedElementDataStorageFactory.h"
+#include "UObject/ObjectMacros.h"
+
+#include "TypedElementActorVisibilityQueries.generated.h"
+
+struct ISceneOutlinerTreeItem;
+
+UCLASS()
+class UActorVisibilityDataStorageFactory : public UEditorDataStorageFactory
+{
+	GENERATED_BODY()
+
+public:
+	~UActorVisibilityDataStorageFactory() override = default;
+
+	void RegisterQueries(IEditorDataStorageProvider& DataStorage) override;
+
+private:
+
+	/**
+	 * Adds the Visibility columns to new actors that do not have one already.
+	 */
+	void RegisterActorAddVisibilityColumn(IEditorDataStorageProvider& DataStorage) const;
+
+	/**
+	 * Takes the visibility set on an actor and copies it to the Data Storage if they differ.
+	 */
+	void RegisterActorVisibilityToColumnQuery(IEditorDataStorageProvider& DataStorage) const;
+	/**
+	 * Takes the visibility stored in the Data Storage and copies it to the actor's visibility if the FTypedElementSyncBackToWorldTag
+	 * has been set and the visibility differ.
+	 */
+	void RegisterVisibilityColumnToActorQuery(IEditorDataStorageProvider& DataStorage) const;
+};
diff --git a/Engine/Plugins/Experimental/EditorDataStorageFeatures/Source/TedsOutliner/Private/Compatibility/SceneOutlinerTedsBridge.cpp b/Engine/Plugins/Experimental/EditorDataStorageFeatures/Source/TedsOutliner/Private/Compatibility/SceneOutlinerTedsBridge.cpp
--- a/Engine/Plugins/Experimental/EditorDataStorageFeatures/Source/TedsOutliner/Private/Compatibility/SceneOutlinerTedsBridge.cpp
+++ b/Engine/Plugins/Experimental/EditorDataStorageFeatures/Source/TedsOutliner/Private/Compatibility/SceneOutlinerTedsBridge.cpp
@@ -3,11 +3,13 @@
 #include "Compatibility/SceneOutlinerTedsBridge.h"
 
 #include "ActorTreeItem.h"
+#include "Columns/TedsOutlinerColumns.h"
 #include "Elements/Columns/TypedElementLabelColumns.h"
 #include "Elements/Columns/TypedElementMiscColumns.h"
 #include "Elements/Columns/TypedElementPackageColumns.h"
 #include "Elements/Columns/TypedElementTransformColumns.h"
 #include "Elements/Columns/TypedElementTypeInfoColumns.h"
+#include "Elements/Columns/TypedElementVisibilityColumns.h"
 #include "Elements/Common/EditorDataStorageFeatures.h"
 #include "Elements/Framework/TypedElementQueryBuilder.h"
 #include "Elements/Interfaces/TypedElementDataStorageInterface.h"
@@ -265,20 +267,20 @@ public:
 		using namespace UE::Editor::DataStorage;
 		using namespace UE::Editor::Outliner;
 		
-		RowHandle RowHandle = InvalidRowHandle;
+		RowHandle TargetRowHandle = InvalidRowHandle;
 
 		TSharedPtr<SWidget> RowWidget;
 
 		if(const FTedsOutlinerTreeItem* TedsItem = TreeItem->CastTo<FTedsOutlinerTreeItem>())
 		{
-			RowHandle = TedsItem->GetRowHandle();
+			TargetRowHandle = TedsItem->GetRowHandle();
 			
 		}
 		else if (const FActorTreeItem* ActorItem = TreeItem->CastTo<FActorTreeItem>())
 		{
 			if (AActor* Actor = ActorItem->Actor.Get())
 			{
-				RowHandle = StorageCompatibility.FindRowWithCompatibleObject(Actor);
+				TargetRowHandle = StorageCompatibility.FindRowWithCompatibleObject(Actor);
 			}
 		}
 		else if(FallbackColumn)
@@ -286,9 +288,13 @@ public:
 			RowWidget = FallbackColumn->ConstructRowWidget(TreeItem, Row);
 		}
 
-		if(Storage.IsRowAssigned(RowHandle))
+		if(Storage.IsRowAssigned(TargetRowHandle))
 		{
-			RowWidget = TableViewerColumnImpl->ConstructRowWidget(RowHandle);
+			RowWidget = TableViewerColumnImpl->ConstructRowWidget(TargetRowHandle,
+				[&](IEditorDataStorageProvider& DataStorage, const RowHandle& WidgetRow)
+				{
+					DataStorage.AddColumn(WidgetRow, FTedsOutlinerColumn{ .Outliner = OwningOutliner });
+				});
 		}
 
 		if(RowWidget)
@@ -380,6 +386,7 @@ void FSceneOutlinerTedsQueryBinder::SetupDefaultColumnMapping()
 {
 	// Map the type column from the TEDS to the default Outliner type column, so we can show type info for objects not in TEDS
 	TEDSToOutlinerDefaultColumnMapping.Add(FTypedElementClassTypeInfoColumn::StaticStruct(), FSceneOutlinerBuiltInColumnTypes::ActorInfo());
+	TEDSToOutlinerDefaultColumnMapping.Add(FVisibleInEditorColumn::StaticStruct(), FSceneOutlinerBuiltInColumnTypes::Gutter());
 }
 
 FName FSceneOutlinerTedsQueryBinder::FindOutlinerColumnFromTEDSColumns(TConstArrayView<TWeakObjectPtr<const UScriptStruct>> TEDSColumns) const
diff --git a/Engine/Plugins/Experimental/EditorDataStorageFeatures/Source/TedsOutliner/Private/TedsOutlinerModule.cpp b/Engine/Plugins/Experimental/EditorDataStorageFeatures/Source/TedsOutliner/Private/TedsOutlinerModule.cpp
--- a/Engine/Plugins/Experimental/EditorDataStorageFeatures/Source/TedsOutliner/Private/TedsOutlinerModule.cpp
+++ b/Engine/Plugins/Experimental/EditorDataStorageFeatures/Source/TedsOutliner/Private/TedsOutlinerModule.cpp
@@ -13,6 +13,7 @@
 #include "Elements/Columns/TypedElementMiscColumns.h"
 #include "Elements/Columns/TypedElementPackageColumns.h"
 #include "Elements/Columns/TypedElementTypeInfoColumns.h"
+#include "Elements/Columns/TypedElementVisibilityColumns.h"
 #include "Elements/Common/EditorDataStorageFeatures.h"
 #include "Elements/Interfaces/TypedElementDataStorageInterface.h"
 #include "Modules/ModuleManager.h"
@@ -123,13 +124,13 @@ DataStorage::QueryHandle FTedsOutlinerModule::GetLevelEditorTedsOutlinerColumnQu
 
 	static DataStorage::QueryHandle ColumnQuery = Storage->RegisterQuery(
 		Select()
-			.ReadOnly<FTypedElementClassTypeInfoColumn, FTypedElementAlertColumn, FTypedElementChildAlertColumn>()
+			.ReadOnly<FTypedElementClassTypeInfoColumn, FTypedElementAlertColumn, FTypedElementChildAlertColumn, FVisibleInEditorColumn>()
 		.Compile());
 
 	// Query to also include revision control info
 	static DataStorage::QueryHandle RevisionControlQuery = Storage->RegisterQuery(
 		Select()
-			.ReadOnly<FTypedElementClassTypeInfoColumn, FTypedElementPackageReference, FTypedElementAlertColumn>()
+			.ReadOnly<FTypedElementClassTypeInfoColumn, FTypedElementPackageReference, FTypedElementAlertColumn, FVisibleInEditorColumn>()
 		.Compile());
 
 	return Private::bUseNewRevisionControlWidgets ? RevisionControlQuery : ColumnQuery;
diff --git a/Engine/Plugins/Experimental/EditorDataStorageFeatures/Source/TedsOutliner/Private/Widgets/OutlinerVisibilityWidget.cpp b/Engine/Plugins/Experimental/EditorDataStorageFeatures/Source/TedsOutliner/Private/Widgets/OutlinerVisibilityWidget.cpp
new file mode 100644
--- /dev/null
+++ b/Engine/Plugins/Experimental/EditorDataStorageFeatures/Source/TedsOutliner/Private/Widgets/OutlinerVisibilityWidget.cpp
@@ -0,0 +1,365 @@
+// Copyright Epic Games, Inc. All Rights Reserved.
+
+#include "Widgets/OutlinerVisibilityWidget.h"
+
+#include "Columns/SlateHeaderColumns.h"
+#include "Columns/TedsOutlinerColumns.h"
+#include "Compatibility/SceneOutlinerTedsBridge.h"
+#include "Elements/Columns/TypedElementCompatibilityColumns.h"
+#include "Elements/Columns/TypedElementHiearchyColumns.h"
+#include "Elements/Columns/TypedElementMiscColumns.h"
+#include "Elements/Columns/TypedElementSelectionColumns.h"
+#include "Elements/Columns/TypedElementVisibilityColumns.h"
+#include "Elements/Common/EditorDataStorageFeatures.h"
+#include "Elements/Framework/TypedElementQueryBuilder.h"
+#include "Elements/Interfaces/TypedElementDataStorageCompatibilityInterface.h"
+#include "Elements/Interfaces/TypedElementDataStorageInterface.h"
+#include "Widgets/SBoxPanel.h"
+#include "ISceneOutliner.h"
+#include "ActorTreeItem.h"
+
+#define LOCTEXT_NAMESPACE "OutlinerVisibilityWidget"
+
+FName UOutlinerVisibilityWidgetFactory::VisibilityWidgetPurpose = FName(TEXT("SceneOutliner.Cell"));
+FName UOutlinerVisibilityHeaderFactory::VisibilityHeaderPurpose = FName(TEXT("SceneOutliner.Header"));
+
+//
+// Cell Factory
+//
+
+FOutlinerVisibilityWidgetConstructor::FOutlinerVisibilityWidgetConstructor()
+	: FSimpleWidgetConstructor(StaticStruct())
+{
+}
+
+void UOutlinerVisibilityWidgetFactory::RegisterWidgetConstructors(IEditorDataStorageProvider& DataStorage, IEditorDataStorageUiProvider& DataStorageUi) const
+{
+	using namespace UE::Editor::DataStorage::Queries;
+	DataStorageUi.RegisterWidgetFactory<FOutlinerVisibilityWidgetConstructor>(VisibilityWidgetPurpose, TColumn<FVisibleInEditorColumn>());
+	
+}
+
+TSharedPtr<SWidget> FOutlinerVisibilityWidgetConstructor::CreateWidget(IEditorDataStorageProvider* DataStorage, IEditorDataStorageUiProvider* DataStorageUi, UE::Editor::DataStorage::RowHandle TargetRow, UE::Editor::DataStorage::RowHandle WidgetRow, const UE::Editor::DataStorage::FMetaDataView& Arguments)
+{
+	return SNew(SHorizontalBox)
+		+ SHorizontalBox::Slot()
+		.AutoWidth()
+		.VAlign(VAlign_Center)
+		[
+			SNew(STEDSVisibilityWidget, TargetRow, WidgetRow)
+				.ToolTipText(LOCTEXT("SceneOutlinerVisibilityToggleTooltip", "Toggles the visibility of this object in the level editor."))
+		];
+}
+
+//
+// Header Factory
+//
+
+FOutlinerVisibilityHeaderConstructor::FOutlinerVisibilityHeaderConstructor()
+	: FSimpleWidgetConstructor(StaticStruct())
+{
+}
+
+void UOutlinerVisibilityHeaderFactory::RegisterWidgetConstructors(IEditorDataStorageProvider& DataStorage, IEditorDataStorageUiProvider& DataStorageUi) const
+{
+	using namespace UE::Editor::DataStorage::Queries;
+	DataStorageUi.RegisterWidgetFactory<FOutlinerVisibilityHeaderConstructor>(VisibilityHeaderPurpose, TColumn<FVisibleInEditorColumn>());
+}
+
+TSharedPtr<SWidget> FOutlinerVisibilityHeaderConstructor::CreateWidget(IEditorDataStorageProvider* DataStorage, IEditorDataStorageUiProvider* DataStorageUi, UE::Editor::DataStorage::RowHandle TargetRow, UE::Editor::DataStorage::RowHandle WidgetRow, const UE::Editor::DataStorage::FMetaDataView& Arguments)
+{
+	DataStorage->AddColumn(WidgetRow, FHeaderWidgetSizeColumn
+	{
+		.ColumnSizeMode = EColumnSizeMode::Fixed,
+		.Width = 24.0f
+	});
+
+	return SNew(SImage)
+				.DesiredSizeOverride(FVector2D(16.f, 16.f))
+				.Image(FAppStyle::Get().GetBrush(TEXT("Level.VisibleHighlightIcon16x")))
+				.ToolTipText(LOCTEXT("SceneOutlinerVisibilityToggleTooltip", "Visibility"));
+}
+
+//
+// SVisibilityWidget
+//
+
+class FTEDSVisibilityDragDropOp : public FDragDropOperation
+{
+public:
+
+	DRAG_DROP_OPERATOR_TYPE(FTEDSVisibilityDragDropOp, FDragDropOperation)
+
+	/** Flag which defines whether to hide destination items or not */
+	bool bHidden;
+
+	/** Undo transaction stolen from the gutter which is kept alive for the duration of the drag */
+	TUniquePtr<FScopedTransaction> UndoTransaction;
+
+	/** The widget decorator to use */
+	virtual TSharedPtr<SWidget> GetDefaultDecorator() const override
+	{
+		return SNullWidget::NullWidget;
+	}
+
+	/** Create a new drag and drop operation out of the specified flag */
+	static TSharedRef<FTEDSVisibilityDragDropOp> New(const bool _bHidden, TUniquePtr<FScopedTransaction>& ScopedTransaction)
+	{
+		TSharedRef<FTEDSVisibilityDragDropOp> Operation = MakeShareable(new FTEDSVisibilityDragDropOp);
+
+		Operation->bHidden = _bHidden;
+		Operation->UndoTransaction = MoveTemp(ScopedTransaction);
+
+		Operation->Construct();
+		return Operation;
+	}
+};
+
+void STEDSVisibilityWidget::Construct(const FArguments& InArgs, const RowHandle& InTargetRow, const RowHandle& InWidgetRow)
+{
+	TargetRow = InTargetRow;
+	WidgetRow = InWidgetRow;
+	TreeItem = GetTreeItem();
+
+	SImage::Construct(
+		SImage::FArguments()
+		.IsEnabled(this, &STEDSVisibilityWidget::IsEnabled)
+		.ColorAndOpacity(this, &STEDSVisibilityWidget::GetForegroundColor)
+		.Image(this, &STEDSVisibilityWidget::GetBrush)
+	);
+
+
+	static const FName NAME_VisibleHoveredBrush = TEXT("Level.VisibleHighlightIcon16x");
+	static const FName NAME_VisibleNotHoveredBrush = TEXT("Level.VisibleIcon16x");
+	static const FName NAME_NotVisibleHoveredBrush = TEXT("Level.NotVisibleHighlightIcon16x");
+	static const FName NAME_NotVisibleNotHoveredBrush = TEXT("Level.NotVisibleIcon16x");
+
+	VisibleHoveredBrush = FAppStyle::Get().GetBrush(NAME_VisibleHoveredBrush);
+	VisibleNotHoveredBrush = FAppStyle::Get().GetBrush(NAME_VisibleNotHoveredBrush);
+
+	NotVisibleHoveredBrush = FAppStyle::Get().GetBrush(NAME_NotVisibleHoveredBrush);
+	NotVisibleNotHoveredBrush = FAppStyle::Get().GetBrush(NAME_NotVisibleNotHoveredBrush);
+
+}
+
+FReply STEDSVisibilityWidget::OnDragDetected(const FGeometry& MyGeometry, const FPointerEvent& MouseEvent)
+{
+	if (MouseEvent.IsMouseButtonDown(EKeys::LeftMouseButton))
+	{
+		return FReply::Handled().BeginDragDrop(FTEDSVisibilityDragDropOp::New(!IsVisible(), UndoTransaction));
+	}
+	else
+	{
+		return FReply::Unhandled();
+	}
+}
+
+/** If a visibility drag drop operation has entered this widget, set its item to the new visibility state */
+void STEDSVisibilityWidget::OnDragEnter(const FGeometry& MyGeometry, const FDragDropEvent& DragDropEvent)
+{
+	auto VisibilityOp = DragDropEvent.GetOperationAs<FTEDSVisibilityDragDropOp>();
+	if (VisibilityOp.IsValid())
+	{
+		SetIsVisible(TargetRow, !VisibilityOp->bHidden);
+	}
+}
+
+FReply STEDSVisibilityWidget::HandleClick()
+{
+	if (!IsEnabled())
+	{
+		return FReply::Unhandled();
+	}
+
+	// Open an undo transaction
+	UndoTransaction.Reset(new FScopedTransaction(LOCTEXT("SetOutlinerItemVisibility", "Set Item Visibility")));
+
+	const bool bVisible = !IsVisible();
+
+	SetIsVisible(TargetRow, bVisible);
+
+	return FReply::Handled().DetectDrag(SharedThis(this), EKeys::LeftMouseButton);
+}
+
+FReply STEDSVisibilityWidget::OnMouseButtonDoubleClick(const FGeometry& InMyGeometry, const FPointerEvent& InMouseEvent)
+{
+	return HandleClick();
+}
+
+/** Called when the mouse button is pressed down on this widget */
+FReply STEDSVisibilityWidget::OnMouseButtonDown(const FGeometry& MyGeometry, const FPointerEvent& MouseEvent)
+{
+	if (MouseEvent.GetEffectingButton() != EKeys::LeftMouseButton)
+	{
+		return FReply::Unhandled();
+	}
+
+	return HandleClick();
+}
+
+/** Process a mouse up message */
+FReply STEDSVisibilityWidget::OnMouseButtonUp(const FGeometry& MyGeometry, const FPointerEvent& MouseEvent)
+{
+	if (MouseEvent.GetEffectingButton() == EKeys::LeftMouseButton)
+	{
+		UndoTransaction.Reset();
+		return FReply::Handled();
+	}
+
+	return FReply::Unhandled();
+}
+
+/** Called when this widget had captured the mouse, but that capture has been revoked for some reason. */
+void STEDSVisibilityWidget::OnMouseCaptureLost(const FCaptureLostEvent& CaptureLostEvent)
+{
+	UndoTransaction.Reset();
+}
+
+/** Get the brush for this widget */
+const FSlateBrush* STEDSVisibilityWidget::GetBrush() const
+{
+	if (IsVisible())
+	{
+		return IsHovered() ? VisibleHoveredBrush : VisibleNotHoveredBrush;
+	}
+	else
+	{
+		return IsHovered() ? NotVisibleHoveredBrush : NotVisibleNotHoveredBrush;
+	}
+}
+
+FSlateColor STEDSVisibilityWidget::GetForegroundColor() const
+{
+	const bool bIsSelected = IsSelected();
+	const bool bIsVisible = IsVisible();
+	const bool bIsHovered = IsHovered();
+
+	// make the foreground brush transparent if it is not selected and it is visible
+	if (bIsVisible && !bIsHovered && !bIsSelected)
+	{
+		return FLinearColor::Transparent;
+	}
+	else if (bIsHovered && !bIsSelected)
+	{
+		return FAppStyle::Get().GetSlateColor("Colors.ForegroundHover");
+	}
+
+	return FSlateColor::UseForeground();
+}
+
+/** Check if our wrapped tree item is visible */
+bool STEDSVisibilityWidget::IsVisible() const
+{
+	if (IEditorDataStorageProvider* DataStorage = GetDataStorage())
+	{
+		FVisibleInEditorColumn* HiddenInEditorColumn = DataStorage->GetColumn<FVisibleInEditorColumn>(TargetRow);
+		if (HiddenInEditorColumn)
+		{
+			return HiddenInEditorColumn->bIsVisibleInEditor;
+		}
+	}
+
+	return true;
+}
+
+bool STEDSVisibilityWidget::IsSelected() const
+{
+	if (IEditorDataStorageProvider* DataStorage = GetDataStorage())
+	{
+		return DataStorage->HasColumns<FTypedElementSelectionColumn>(TargetRow);
+	}
+
+	return false;
+}
+
+/** Set the item this widget is responsible for to be hidden or shown */
+void STEDSVisibilityWidget::SetIsVisible(RowHandle InRow, const bool bVisible)
+{
+	if (IEditorDataStorageProvider* DataStorage = GetDataStorage(); TreeItem.IsValid())
+	{
+		SetVisibility_Recursive(*DataStorage, TreeItem.Pin(), bVisible);
+	}
+}
+
+void STEDSVisibilityWidget::CommitVisibility(IEditorDataStorageProvider& DataStorage, RowHandle Row, bool bVisible)
+{
+	// Set new visibility on the TargetRow
+	FVisibleInEditorColumn* VisibileInEditorColumn = DataStorage.GetColumn<FVisibleInEditorColumn>(Row);
+	if (VisibileInEditorColumn)
+	{
+		VisibileInEditorColumn->bIsVisibleInEditor = bVisible;
+	}
+	DataStorage.AddColumn<FTypedElementSyncBackToWorldTag>(Row);
+}
+
+void STEDSVisibilityWidget::SetVisibility_Recursive(IEditorDataStorageProvider& DataStorage, FSceneOutlinerTreeItemPtr TreeItem, bool bVisible)
+{
+	using namespace UE::Editor::DataStorage;
+	IEditorDataStorageCompatibilityProvider* DataStorageComat = GetDataStorageCompatibility();
+	FActorTreeItem* ActorTreeItem = TreeItem->CastTo<FActorTreeItem>();
+	
+	if (ActorTreeItem && DataStorageComat)
+	{
+		RowHandle ActorRow = DataStorageComat->FindRowWithCompatibleObject(ActorTreeItem->Actor);
+		CommitVisibility(DataStorage, ActorRow, bVisible);
+	}
+
+	for (auto& ChildTreeItemPtr : TreeItem->GetChildren())
+	{
+		auto ChildTreeItem = ChildTreeItemPtr.Pin();
+		if (ChildTreeItem.IsValid())
+		{
+			SetVisibility_Recursive(DataStorage, ChildTreeItem, bVisible);
+		}
+	}
+}
+
+FSceneOutlinerTreeItemPtr STEDSVisibilityWidget::GetTreeItem() const
+{
+	IEditorDataStorageProvider* DataStorage = GetDataStorage();
+	if (!DataStorage)
+	{
+		return nullptr;
+	}
+
+	const FTedsOutlinerColumn* TEDSOutlinerColumn = DataStorage->GetColumn<FTedsOutlinerColumn>(WidgetRow);
+	if (!TEDSOutlinerColumn)
+	{
+		return nullptr;
+	}
+
+	TSharedPtr<ISceneOutliner> Outliner = TEDSOutlinerColumn->Outliner.IsValid() ? TEDSOutlinerColumn->Outliner.Pin() : nullptr;
+	if (!Outliner.IsValid())
+	{
+		return nullptr;
+	}
+
+	FSceneOutlinerTedsQueryBinder& Binder = FSceneOutlinerTedsQueryBinder::GetInstance();
+	if (FTreeItemIDDealiaser Dealiaser = Binder.GetTreeItemIDDealiaser(Outliner); Dealiaser.IsBound())
+	{
+		FSceneOutlinerTreeItemID ID = Dealiaser.Execute(TargetRow);
+		return Outliner->GetTreeItem(ID);
+	}
+
+	return nullptr;
+}
+
+IEditorDataStorageProvider* STEDSVisibilityWidget::GetDataStorage()
+{
+	using namespace UE::Editor::DataStorage;
+	return GetMutableDataStorageFeature<IEditorDataStorageProvider>(StorageFeatureName);
+}
+
+IEditorDataStorageProvider* STEDSVisibilityWidget::GetDataStorageUI()
+{
+	using namespace UE::Editor::DataStorage;
+	return GetMutableDataStorageFeature<IEditorDataStorageProvider>(UiFeatureName);
+}
+
+IEditorDataStorageCompatibilityProvider* STEDSVisibilityWidget::GetDataStorageCompatibility()
+{
+	using namespace UE::Editor::DataStorage;
+	return GetMutableDataStorageFeature<IEditorDataStorageCompatibilityProvider>(CompatibilityFeatureName);
+}
+
+#undef LOCTEXT_NAMESPACE
\ No newline at end of file
diff --git a/Engine/Plugins/Experimental/EditorDataStorageFeatures/Source/TedsOutliner/Public/Widgets/OutlinerVisibilityWidget.h b/Engine/Plugins/Experimental/EditorDataStorageFeatures/Source/TedsOutliner/Public/Widgets/OutlinerVisibilityWidget.h
new file mode 100644
--- /dev/null
+++ b/Engine/Plugins/Experimental/EditorDataStorageFeatures/Source/TedsOutliner/Public/Widgets/OutlinerVisibilityWidget.h
@@ -0,0 +1,164 @@
+// Copyright Epic Games, Inc. All Rights Reserved.
+
+#pragma once
+
+#include "Elements/Interfaces/TypedElementDataStorageFactory.h"
+#include "Elements/Interfaces/TypedElementDataStorageUiInterface.h"
+#include "Internationalization/Text.h"
+#include "UObject/ObjectMacros.h"
+#include "Widgets/Images/SImage.h"
+#include "ScopedTransaction.h"
+#include "SceneOutlinerFwd.h"
+
+#include "OutlinerVisibilityWidget.generated.h"
+
+class IEditorDataStorageProvider;
+class UScriptStruct;
+
+UCLASS()
+class UOutlinerVisibilityWidgetFactory : public UEditorDataStorageFactory
+{
+	GENERATED_BODY()
+
+public:
+	~UOutlinerVisibilityWidgetFactory() override = default;
+
+	TEDSOUTLINER_API void RegisterWidgetConstructors(IEditorDataStorageProvider& DataStorage, IEditorDataStorageUiProvider& DataStorageUi) const override;
+
+public:
+
+	static FName VisibilityWidgetPurpose;
+};
+
+UCLASS()
+class UOutlinerVisibilityHeaderFactory : public UEditorDataStorageFactory
+{
+	GENERATED_BODY()
+
+public:
+	~UOutlinerVisibilityHeaderFactory() override = default;
+
+	TEDSOUTLINER_API void RegisterWidgetConstructors(IEditorDataStorageProvider& DataStorage, IEditorDataStorageUiProvider& DataStorageUi) const override;
+
+public:
+
+	static FName VisibilityHeaderPurpose;
+};
+
+USTRUCT()
+struct FOutlinerVisibilityWidgetConstructor : public FSimpleWidgetConstructor
+{
+	GENERATED_BODY()
+
+public:
+	TEDSOUTLINER_API FOutlinerVisibilityWidgetConstructor();
+	~FOutlinerVisibilityWidgetConstructor() override = default;
+
+	TEDSOUTLINER_API virtual TSharedPtr<SWidget> CreateWidget(
+		IEditorDataStorageProvider* DataStorage,
+		IEditorDataStorageUiProvider* DataStorageUi,
+		UE::Editor::DataStorage::RowHandle TargetRow,
+		UE::Editor::DataStorage::RowHandle WidgetRow, 
+		const UE::Editor::DataStorage::FMetaDataView& Arguments) override;
+};
+
+USTRUCT()
+struct FOutlinerVisibilityHeaderConstructor : public FSimpleWidgetConstructor
+{
+	GENERATED_BODY()
+
+public:
+	TEDSOUTLINER_API FOutlinerVisibilityHeaderConstructor();
+	~FOutlinerVisibilityHeaderConstructor() override = default;
+
+	TEDSOUTLINER_API virtual TSharedPtr<SWidget> CreateWidget(
+		IEditorDataStorageProvider* DataStorage,
+		IEditorDataStorageUiProvider* DataStorageUi,
+		UE::Editor::DataStorage::RowHandle TargetRow,
+		UE::Editor::DataStorage::RowHandle WidgetRow, 
+		const UE::Editor::DataStorage::FMetaDataView& Arguments) override;
+};
+
+/** Widget responsible for managing the visibility for a single item */
+class TEDSOUTLINER_API STEDSVisibilityWidget : public SImage
+{
+	using RowHandle = UE::Editor::DataStorage::RowHandle;
+
+public:
+	SLATE_BEGIN_ARGS(STEDSVisibilityWidget) {}
+	SLATE_END_ARGS()
+
+	/** Construct this widget */
+	void Construct(const FArguments& InArgs, const RowHandle& InTargetRow, const RowHandle& InWidgetRow);
+
+protected:
+
+	/** Returns whether the widget is enabled or not */
+	virtual bool IsEnabled() const { return true; }
+
+	/** Get the brush for this widget */
+	virtual const FSlateBrush* GetBrush() const;
+
+	/** Start a new drag/drop operation for this widget */
+	virtual FReply OnDragDetected(const FGeometry& MyGeometry, const FPointerEvent& MouseEvent) override;
+
+	/** If a visibility drag drop operation has entered this widget, set its item to the new visibility state */
+	virtual void OnDragEnter(const FGeometry& MyGeometry, const FDragDropEvent& DragDropEvent) override;
+
+	FReply HandleClick();
+
+	virtual FReply OnMouseButtonDoubleClick(const FGeometry& InMyGeometry, const FPointerEvent& InMouseEvent) override;
+
+	/** Called when the mouse button is pressed down on this widget */
+	virtual FReply OnMouseButtonDown(const FGeometry& MyGeometry, const FPointerEvent& MouseEvent) override;
+
+	/** Process a mouse up message */
+	virtual FReply OnMouseButtonUp(const FGeometry& MyGeometry, const FPointerEvent& MouseEvent) override;
+
+	/** Called when this widget had captured the mouse, but that capture has been revoked for some reason. */
+	virtual void OnMouseCaptureLost(const FCaptureLostEvent& CaptureLostEvent) override;
+
+	/** Whether visibility change should propagate down to children. */
+	virtual bool ShouldPropagateVisibilityChangeOnChildren() const { return true; }
+
+	virtual FSlateColor GetForegroundColor() const;
+
+	/** Check if the TargetRow Object is Visible */
+	bool IsVisible() const;
+
+	/** Check if the TargetRow Object is Selected */
+	bool IsSelected() const;
+
+	/** Set the item this widget is responsible for to be hidden or shown */
+	void SetIsVisible(RowHandle InRow, const bool bVisible);
+
+	RowHandle TargetRow;
+	RowHandle WidgetRow;
+
+	// Pointer to the outliner tree item associated with this widget. 
+	TWeakPtr<ISceneOutlinerTreeItem> TreeItem;
+
+	/** Scoped undo transaction */
+	TUniquePtr<FScopedTransaction> UndoTransaction;
+
+	/** Visibility brushes for the various states */
+	const FSlateBrush* VisibleHoveredBrush;
+	const FSlateBrush* VisibleNotHoveredBrush;
+	const FSlateBrush* NotVisibleHoveredBrush;
+	const FSlateBrush* NotVisibleNotHoveredBrush;
+
+private:
+
+	// Commit the visibility state into the Rows DataStorage. Will trigger a SyncBackToWorld manually but the object is never dirtied.
+	static void CommitVisibility(IEditorDataStorageProvider& DataStorage, RowHandle Row, bool bVisible);
+
+	// Recursive internal method to set visibility on a TreeItem and all of it's children.
+	static void SetVisibility_Recursive(IEditorDataStorageProvider& DataStorage, FSceneOutlinerTreeItemPtr TreeItem, bool bVisible);
+
+	// Returns the tree item associated with this widget.
+	FSceneOutlinerTreeItemPtr GetTreeItem() const;
+
+	static IEditorDataStorageProvider* GetDataStorage();
+	static IEditorDataStorageProvider* GetDataStorageUI();
+	static IEditorDataStorageCompatibilityProvider* GetDataStorageCompatibility();
+};
\ No newline at end of file
diff --git a/Engine/Plugins/Experimental/EditorDataStorageFeatures/Source/TedsOutliner/TedsOutliner.Build.cs b/Engine/Plugins/Experimental/EditorDataStorageFeatures/Source/TedsOutliner/TedsOutliner.Build.cs
--- a/Engine/Plugins/Experimental/EditorDataStorageFeatures/Source/TedsOutliner/TedsOutliner.Build.cs
+++ b/Engine/Plugins/Experimental/EditorDataStorageFeatures/Source/TedsOutliner/TedsOutliner.Build.cs
@@ -33,7 +33,8 @@ public class TedsOutliner : ModuleRules
 					"ToolMenus",
 					"ApplicationCore",
 					"TedsTableViewer",
-					"ToolWidgets"
+					"ToolWidgets",
+					"InputCore"
 				});
 
 			DynamicallyLoadedModuleNames.AddRange(new string[] {});
diff --git a/Engine/Plugins/Experimental/EditorDataStorageFeatures/Source/TedsTableViewer/Private/TedsTableViewerColumn.cpp b/Engine/Plugins/Experimental/EditorDataStorageFeatures/Source/TedsTableViewer/Private/TedsTableViewerColumn.cpp
--- a/Engine/Plugins/Experimental/EditorDataStorageFeatures/Source/TedsTableViewer/Private/TedsTableViewerColumn.cpp
+++ b/Engine/Plugins/Experimental/EditorDataStorageFeatures/Source/TedsTableViewer/Private/TedsTableViewerColumn.cpp
@@ -47,7 +47,7 @@ namespace UE::Editor::DataStorage
 		UnRegisterQueries();
 	}
 
-	TSharedPtr<SWidget> FTedsTableViewerColumn::ConstructRowWidget(RowHandle InRowHandle) const
+	TSharedPtr<SWidget> FTedsTableViewerColumn::ConstructRowWidget(RowHandle InRowHandle, TFunction<void(IEditorDataStorageProvider&, const RowHandle&)> WidgetRowSetupDelegate/* = nullptr*/) const
 	{
 		TSharedPtr<SWidget> RowWidget;
 		
@@ -71,6 +71,11 @@ namespace UE::Editor::DataStorage
 				WidgetReferenceColumn->WidgetConstructor = CellWidgetConstructor;
 			}
 
+			if (WidgetRowSetupDelegate)
+			{
+				WidgetRowSetupDelegate(*Storage, UiRowHandle);
+			}
+
 			RowWidget = StorageUi->ConstructWidget(UiRowHandle, *CellWidgetConstructor, WidgetMetaData);
 		}
 
diff --git a/Engine/Plugins/Experimental/EditorDataStorageFeatures/Source/TedsTableViewer/Public/TedsTableViewerColumn.h b/Engine/Plugins/Experimental/EditorDataStorageFeatures/Source/TedsTableViewer/Public/TedsTableViewerColumn.h
--- a/Engine/Plugins/Experimental/EditorDataStorageFeatures/Source/TedsTableViewer/Public/TedsTableViewerColumn.h
+++ b/Engine/Plugins/Experimental/EditorDataStorageFeatures/Source/TedsTableViewer/Public/TedsTableViewerColumn.h
@@ -36,7 +36,8 @@ namespace UE::Editor::DataStorage
 
 		TEDSTABLEVIEWER_API ~FTedsTableViewerColumn();
 		
-		TEDSTABLEVIEWER_API	TSharedPtr<SWidget> ConstructRowWidget(RowHandle InRowHandle) const;
+		TEDSTABLEVIEWER_API	TSharedPtr<SWidget> ConstructRowWidget(RowHandle InRowHandle, 
+			TFunction<void(IEditorDataStorageProvider&, const RowHandle&)> WidgetRowSetupDelegate = nullptr) const;
 		
 		TEDSTABLEVIEWER_API SHeaderRow::FColumn::FArguments ConstructHeaderRowColumn() const;
 		
diff --git a/Engine/Source/Editor/UnrealEd/Private/EditorActor.cpp b/Engine/Source/Editor/UnrealEd/Private/EditorActor.cpp
--- a/Engine/Source/Editor/UnrealEd/Private/EditorActor.cpp
+++ b/Engine/Source/Editor/UnrealEd/Private/EditorActor.cpp
@@ -77,6 +77,11 @@
 #include "Subsystems/AssetEditorSubsystem.h"
 #include "SceneInterface.h"
 #include "SDeleteReferencedActorDialog.h"
+#include "Elements/Interfaces/TypedElementDataStorageCompatibilityInterface.h"
+#include "Elements/Common/EditorDataStorageFeatures.h"
+#include "Elements/Common/EditorDataStorageFeatures.h"
+#include "Elements/Columns/TypedElementVisibilityColumns.h"
+#include "Elements/Columns/TypedElementMiscColumns.h"
 
 #define LOCTEXT_NAMESPACE "UnrealEd.EditorActor"
 
@@ -1477,6 +1482,30 @@ void UUnrealEdEngine::edactReplaceClassWithClass(UWorld* InWorld, UClass* SrcCla
 	NoteSelectionChange();
 }
 
+static void SetActorVisibility(AActor* Actor, bool bVisible)
+{
+	using namespace UE::Editor::DataStorage;
+
+	// Save the actor to the transaction buffer to support undo/redo, but do
+	// not call Modify, as we do not want to dirty the actor's package and
+	// we're only editing temporary, transient values
+	SaveToTransactionBuffer(Actor, false);
+
+	if (const IEditorDataStorageCompatibilityProvider* DataStoreCompat = GetDataStorageFeature<IEditorDataStorageCompatibilityProvider>(CompatibilityFeatureName))
+	{
+		RowHandle Row = DataStoreCompat->FindRowWithCompatibleObject(Actor);
+		if (IEditorDataStorageProvider* DataStore = GetMutableDataStorageFeature<IEditorDataStorageProvider>(StorageFeatureName))
+		{
+			FVisibleInEditorColumn* VisibilityColumn = DataStore->GetColumn<FVisibleInEditorColumn>(Row);
+			if (VisibilityColumn)
+			{
+				VisibilityColumn->bIsVisibleInEditor = bVisible;
+			}
+			DataStore->AddColumn<FTypedElementSyncBackToWorldTag>(Row);
+		}
+	}
+}
+
 void UUnrealEdEngine::edactHideSelected( UWorld* InWorld )
 {
 	// Assemble a list of actors to hide.
@@ -1503,11 +1532,7 @@ void UUnrealEdEngine::edactHideSelected( UWorld* InWorld )
 		{
 			AActor* Actor = ActorsToHide[ ActorIndex ];
 			
-			// Save the actor to the transaction buffer to support undo/redo, but do
-			// not call Modify, as we do not want to dirty the actor's package and
-			// we're only editing temporary, transient values
-			SaveToTransactionBuffer(Actor, false);
-			Actor->SetIsTemporarilyHiddenInEditor( true );
+			SetActorVisibility(Actor, false);
 			SelectedActors->Deselect( Actor );
 		}
 
@@ -1546,11 +1571,7 @@ void UUnrealEdEngine::edactHideUnselected( UWorld* InWorld )
 		AActor* Actor = *It;
 		if( !FActorEditorUtils::IsABuilderBrush(Actor) && !Actor->IsActorOrSelectionParentSelected() && !Actor->IsHiddenEd() )
 		{
-			// Save the actor to the transaction buffer to support undo/redo, but do
-			// not call Modify, as we do not want to dirty the actor's package and
-			// we're only editing temporary, transient values
-			SaveToTransactionBuffer(Actor, false);
-			Actor->SetIsTemporarilyHiddenInEditor( true );
+			SetActorVisibility(Actor, false);
 		}
 	}
 
@@ -1585,11 +1606,7 @@ void UUnrealEdEngine::edactUnHideAll( UWorld* InWorld )
 		AActor* Actor = *It;
 		if( !FActorEditorUtils::IsABuilderBrush(Actor) && Actor->IsTemporarilyHiddenInEditor() )
 		{
-			// Save the actor to the transaction buffer to support undo/redo, but do
-			// not call Modify, as we do not want to dirty the actor's package and
-			// we're only editing temporary, transient values
-			SaveToTransactionBuffer(Actor, false);
-			Actor->SetIsTemporarilyHiddenInEditor( false );
+			SetActorVisibility(Actor, true);
 		}
 	}
 
@@ -1777,11 +1794,7 @@ void UUnrealEdEngine::edactUnhideSelected( UWorld* InWorld )
 		{
 			AActor* Actor = ActorsToShow[ ActorIndex ];
 			
-			// Save the actor to the transaction buffer to support undo/redo, but do
-			// not call Modify, as we do not want to dirty the actor's package and
-			// we're only editing temporary, transient values
-			SaveToTransactionBuffer(Actor, false);
-			Actor->SetIsTemporarilyHiddenInEditor( false );
+			SetActorVisibility(Actor, true);
 		}
 	}
 
diff --git a/Engine/Source/Runtime/TypedElementFramework/Public/Elements/Columns/TypedElementVisibilityColumns.h b/Engine/Source/Runtime/TypedElementFramework/Public/Elements/Columns/TypedElementVisibilityColumns.h
new file mode 100644
--- /dev/null
+++ b/Engine/Source/Runtime/TypedElementFramework/Public/Elements/Columns/TypedElementVisibilityColumns.h
@@ -0,0 +1,20 @@
+// Copyright Epic Games, Inc. All Rights Reserved.
+
+#pragma once
+
+#include "Elements/Interfaces/TypedElementDataStorageInterface.h"
+
+#include "TypedElementVisibilityColumns.generated.h"
+
+/**
+ * VisibleInEditor column that signifies whether or not this row's object should be visible in view ports
+ */
+USTRUCT(meta = (DisplayName = "Visibility"))
+struct FVisibleInEditorColumn final : public FEditorDataStorageColumn
+{
+	GENERATED_BODY()
+
+public:
+
+	bool bIsVisibleInEditor = true;
+};
\ No newline at end of file
