[TEDSOutliner]
* Fixed conflicting LOCTEXT keys

diff --git a/Engine/Plugins/Experimental/EditorDataStorageFeatures/Source/TedsOutliner/Private/Widgets/OutlinerVisibilityWidget.cpp b/Engine/Plugins/Experimental/EditorDataStorageFeatures/Source/TedsOutliner/Private/Widgets/OutlinerVisibilityWidget.cpp
--- a/Engine/Plugins/Experimental/EditorDataStorageFeatures/Source/TedsOutliner/Private/Widgets/OutlinerVisibilityWidget.cpp
+++ b/Engine/Plugins/Experimental/EditorDataStorageFeatures/Source/TedsOutliner/Private/Widgets/OutlinerVisibilityWidget.cpp
@@ -77,7 +77,7 @@ TSharedPtr<SWidget> FOutlinerVisibilityHeaderConstructor::CreateWidget(IEditorDa
 	return SNew(SImage)
 				.DesiredSizeOverride(FVector2D(16.f, 16.f))
 				.Image(FAppStyle::Get().GetBrush(TEXT("Level.VisibleHighlightIcon16x")))
-				.ToolTipText(LOCTEXT("SceneOutlinerVisibilityToggleTooltip", "Visibility"));
+				.ToolTipText(LOCTEXT("SceneOutlinerVisibilityHeaderTooltip", "Visibility"));
 }
 
 //
