TEDS: Use package filenames that include the extension.

Previously, the package filename had no extension, so the PathOnDisk would be wrong. These paths were given to the FSourceControlFileStatusMonitor in URevisionControlDataStorageFactory::RegisterFetchUpdates which would trigger FUpdateStatus requests on a batch of non-existing files.

#rnx

diff --git a/Engine/Plugins/Developer/PerforceSourceControl/Source/PerforceSourceControl/Private/PerforceSourceControlOperations.cpp b/Engine/Plugins/Developer/PerforceSourceControl/Source/PerforceSourceControl/Private/PerforceSourceControlOperations.cpp
--- a/Engine/Plugins/Developer/PerforceSourceControl/Source/PerforceSourceControl/Private/PerforceSourceControlOperations.cpp
+++ b/Engine/Plugins/Developer/PerforceSourceControl/Source/PerforceSourceControl/Private/PerforceSourceControlOperations.cpp
@@ -2051,7 +2051,7 @@ bool FPerforceUpdateStatusWorker::UpdateStates() const
 
 		auto GetRevisionControlRow = [DataStorage, RevisionControlTable](const FString& InFilename) -> RowHandle
 		{
-			FString Filename = FPaths::SetExtension(InFilename, "");
+			FString Filename = InFilename;
 			FPaths::NormalizeFilename(Filename);
 			Filename = FPaths::ConvertRelativePathToFull(Filename);
 			
diff --git a/Engine/Plugins/Experimental/EditorDataStorageFeatures/Source/TedsRevisionControl/Private/Queries/ObjectPackagePathToColumnQueries.cpp b/Engine/Plugins/Experimental/EditorDataStorageFeatures/Source/TedsRevisionControl/Private/Queries/ObjectPackagePathToColumnQueries.cpp
--- a/Engine/Plugins/Experimental/EditorDataStorageFeatures/Source/TedsRevisionControl/Private/Queries/ObjectPackagePathToColumnQueries.cpp
+++ b/Engine/Plugins/Experimental/EditorDataStorageFeatures/Source/TedsRevisionControl/Private/Queries/ObjectPackagePathToColumnQueries.cpp
@@ -116,9 +116,13 @@ void UTypedElementUObjectPackagePathFactory::RegisterTryAddPackageRef(IEditorDat
 						return;
 					}
 
+					const FString PackageExtension = Target->ContainsMap()
+						? FPackageName::GetMapPackageExtension()
+						: FPackageName::GetAssetPackageExtension();
+
 					FString Path;
 					Target->GetPathName(nullptr, Path);
-					if (FString PackageFilename; FPackageName::TryConvertLongPackageNameToFilename(Path, PackageFilename))
+					if (FString PackageFilename; FPackageName::TryConvertLongPackageNameToFilename(Path, PackageFilename, PackageExtension))
 					{
 						FPaths::NormalizeFilename(PackageFilename);
 						FString FullPackageFilename = FPaths::ConvertRelativePathToFull(PackageFilename);
