TEDS Alerts can now handle multiple alerts on the same row.

Previously alerts would override each other when added to the same row. This would cause issues when the alert from code that didn't place the final alert. The new approach keeps an ordered linked list of all alerts so when one alert is deleted the one with the next highest priority takes its place. To support this alerts are now ordered errors before warnings and then by the one with the largets priority value. In order to still be able to remove individual alerts, each alert now also requires a name. The name doesn't need to strictly be unique but if there are multiple alerts with the same name only the first one can be addressed. The chaining introduced additional complexity making it impractical to simply rely on the action column being added or not, so alerts now need to be added through the new TEDS Alerts module.

reviews: 37940020, 37940193

diff --git a/Engine/Plugins/Experimental/EditorDataStorage/Source/TedsCore/Private/Queries/TypedElementAlertQueries.cpp b/Engine/Plugins/Experimental/EditorDataStorage/Source/TedsCore/Private/Queries/TypedElementAlertQueries.cpp
deleted file mode 100644
diff --git a/Engine/Plugins/Experimental/EditorDataStorage/Source/TedsCore/Private/Queries/TypedElementAlertQueries.h b/Engine/Plugins/Experimental/EditorDataStorage/Source/TedsCore/Private/Queries/TypedElementAlertQueries.h
deleted file mode 100644
diff --git a/Engine/Plugins/Experimental/EditorDataStorageFeatures/EditorDataStorageFeatures.uplugin b/Engine/Plugins/Experimental/EditorDataStorageFeatures/EditorDataStorageFeatures.uplugin
--- a/Engine/Plugins/Experimental/EditorDataStorageFeatures/EditorDataStorageFeatures.uplugin
+++ b/Engine/Plugins/Experimental/EditorDataStorageFeatures/EditorDataStorageFeatures.uplugin
@@ -19,6 +19,10 @@
 			"Name": "TedsRevisionControl",
 			"Type": "Editor"
 		},
+		{
+			"Name": "TedsAlerts",
+			"Type":  "Editor"
+		},
 		{
 			"Name": "TedsAssetData",
 			"Type":  "Editor"
diff --git a/Engine/Plugins/Experimental/EditorDataStorageFeatures/Source/TedsAlerts/Private/TedsAlertWidget.cpp b/Engine/Plugins/Experimental/EditorDataStorageFeatures/Source/TedsAlerts/Private/TedsAlertWidget.cpp
new file mode 100644
--- /dev/null
+++ b/Engine/Plugins/Experimental/EditorDataStorageFeatures/Source/TedsAlerts/Private/TedsAlertWidget.cpp
@@ -0,0 +1,437 @@
+// Copyright Epic Games, Inc. All Rights Reserved.
+
+#include "TedsAlertWidget.h"
+
+#include "Columns/SlateHeaderColumns.h"
+#include "Elements/Columns/TypedElementMiscColumns.h"
+#include "Elements/Columns/TypedElementSlateWidgetColumns.h"
+#include "Elements/Common/EditorDataStorageFeatures.h"
+#include "TedsAlertColumns.h"
+#include "Widgets/Images/SImage.h"
+#include "Widgets/Input/SButton.h"
+#include "Widgets/SOverlay.h"
+#include "Widgets/Text/STextBlock.h"
+
+#define LOCTEXT_NAMESPACE "TedsAlertWidget"
+
+namespace UE::Editor::DataStorage::Widgets::Private
+{
+	void UpdateWidget(const TSharedPtr<SWidget>& Widget, const FText& Alert, bool bIsWarning, uint16 ErrorCount, uint16 WarningCount,
+		RowHandle RowWithAlertAction)
+	{
+		using namespace UE::Editor::DataStorage::Columns;
+
+		if (Widget)
+		{
+			uint32 ChildCount = ErrorCount + WarningCount;
+			if (FChildren* Children = Widget->GetChildren())
+			{
+				SImage& Background =
+					static_cast<SImage&>(*Children->GetSlotAt(FAlertWidgetConstructor::IconBackgroundSlot).GetWidget());
+				SImage& Badge =
+					static_cast<SImage&>(*Children->GetSlotAt(FAlertWidgetConstructor::IconBadgeSlot).GetWidget());
+				STextBlock& CounterText =
+					static_cast<STextBlock&>(*Children->GetSlotAt(FAlertWidgetConstructor::CounterTextSlot).GetWidget());
+				SButton& ActionButton =
+					static_cast<SButton&>(*Children->GetSlotAt(FAlertWidgetConstructor::ActionButtonSlot).GetWidget());
+
+				// Setup the background image
+				if (ChildCount > 0)
+				{
+					if (!Alert.IsEmpty())
+					{
+						Background.SetImage(bIsWarning
+							? FAppStyle::GetBrush("Icons.Warning.Background")
+							: FAppStyle::GetBrush("Icons.Error.Background"));
+					}
+					else
+					{
+						Background.SetImage(FAppStyle::GetBrush("Icons.Alert.Background"));
+					}
+				}
+				else
+				{
+					if (!Alert.IsEmpty())
+					{
+						Background.SetImage(bIsWarning
+							? FAppStyle::GetBrush("Icons.Warning.Solid")
+							: FAppStyle::GetBrush("Icons.Error.Solid"));
+					}
+					else
+					{
+						Background.SetImage(FAppStyle::GetBrush("Icons.Alert.Solid"));
+					}
+				}
+
+				// Set counter if needed, otherwise turn it off.
+				if (ChildCount == 0)
+				{
+					// If there are no children, don't show the badge and don't show a counter.
+					Badge.SetVisibility(EVisibility::Hidden);
+					CounterText.SetVisibility(EVisibility::Hidden);
+				}
+				else
+				{
+					// If there are children, also take into account if there's an alert as well.
+					uint32 TotalChildCount = ChildCount + (Alert.IsEmpty() ? 0 : 1);
+					if (TotalChildCount <= 9)
+					{
+						Badge.SetVisibility(EVisibility::HitTestInvisible);
+						CounterText.SetVisibility(EVisibility::HitTestInvisible);
+						CounterText.SetText(FText::AsNumber(TotalChildCount));
+						CounterText.SetFont(FCoreStyle::GetDefaultFontStyle("Regular", FAlertWidgetConstructor::BadgeFontSize));
+						CounterText.SetMargin(FMargin(
+							FAlertWidgetConstructor::BadgeHorizontalOffset,
+							FAlertWidgetConstructor::BadgeVerticalOffset));
+					}
+					else
+					{
+						Badge.SetVisibility(EVisibility::HitTestInvisible);
+						CounterText.SetVisibility(EVisibility::HitTestInvisible);
+						CounterText.SetText(FText::FromString(TEXT("*")));
+						CounterText.SetFont(FCoreStyle::GetDefaultFontStyle("Regular", 14));
+						CounterText.SetMargin(FMargin(
+							FAlertWidgetConstructor::BadgeHorizontalOffset - 2.0f,
+							FAlertWidgetConstructor::BadgeVerticalOffset - 6.5f));
+					}
+				}
+
+				// Setup the tool tip text
+				if (!Alert.IsEmpty() && ChildCount > 0)
+				{
+					FText ToolTipText = FText::Format(LOCTEXT("ChildAlertCountWithMessage", "Errors: {0}\nWarnings: {1}\n\n{2}"),
+						FText::AsNumber(ErrorCount + (bIsWarning ? 0 : 1)),
+						FText::AsNumber(WarningCount + (bIsWarning ? 1 : 0)), Alert);
+					Background.SetToolTipText(ToolTipText);
+					ActionButton.SetToolTipText(MoveTemp(ToolTipText));
+				}
+				else if (!Alert.IsEmpty())
+				{
+					Background.SetToolTipText(Alert);
+					ActionButton.SetToolTipText(Alert);
+				}
+				else if (ChildCount > 0)
+				{
+					FText ToolTipText = FText::Format(LOCTEXT("ChildAlertCount", "Errors: {0}\nWarnings: {1}"),
+						FText::AsNumber(ErrorCount), FText::AsNumber(WarningCount));
+					Background.SetToolTipText(ToolTipText);
+					ActionButton.SetToolTipText(MoveTemp(ToolTipText));
+				}
+
+				// If there's an action to call, enable the invisible button, otherwise turn it off.
+				if (RowWithAlertAction != InvalidRowHandle)
+				{
+					Background.SetVisibility(EVisibility::HitTestInvisible);
+					ActionButton.SetVisibility(EVisibility::Visible);
+					ActionButton.SetOnClicked(FOnClicked::CreateLambda([RowWithAlertAction]()
+						{
+							const IEditorDataStorageProvider* DataStorage = GetDataStorageFeature<IEditorDataStorageProvider>(StorageFeatureName);
+							if (const FAlertActionColumn* Action =
+								DataStorage->GetColumn<FAlertActionColumn>(RowWithAlertAction))
+							{
+								Action->Action(RowWithAlertAction);
+							}
+							return FReply::Handled();
+						}));
+				}
+				else
+				{
+					Background.SetVisibility(EVisibility::Visible);
+					ActionButton.SetVisibility(EVisibility::Hidden);
+				}
+			}
+		}
+	}
+}
+
+//
+// UAlertWidgetFactory
+//
+
+void UAlertWidgetFactory::RegisterWidgetConstructors(
+	IEditorDataStorageProvider& DataStorage,
+	IEditorDataStorageUiProvider& DataStorageUi) const
+{
+	using namespace UE::Editor::DataStorage::Columns;
+	using namespace UE::Editor::DataStorage::Queries;
+
+	DataStorageUi.RegisterWidgetFactory<FAlertWidgetConstructor>(FName(TEXT("General.Cell")),
+		TColumn<FAlertColumn>() || TColumn<FChildAlertColumn>());
+
+	DataStorageUi.RegisterWidgetFactory<FAlertHeaderWidgetConstructor>(FName(TEXT("General.Header")),
+		TColumn<FAlertColumn>() || TColumn<FChildAlertColumn>());
+}
+
+void UAlertWidgetFactory::RegisterQueries(IEditorDataStorageProvider& DataStorage)
+{
+	RegisterAlertQueries(DataStorage);
+	RegisterAlertHeaderQueries(DataStorage);
+}
+
+void UAlertWidgetFactory::RegisterAlertQueries(IEditorDataStorageProvider& DataStorage)
+{
+	using namespace UE::Editor::DataStorage::Columns;
+	using namespace UE::Editor::DataStorage::Queries;
+
+	QueryHandle UpdateWidget_OnlyAlert = DataStorage.RegisterQuery(
+		Select()
+		.ReadOnly<FAlertColumn>()
+		.Where()
+			.Any<FTypedElementSyncFromWorldTag, FTypedElementSyncBackToWorldTag>()
+			.None<FChildAlertColumn>()
+		.Compile());
+
+	QueryHandle UpdateWidget_OnlyChildAlert = DataStorage.RegisterQuery(
+		Select()
+		.ReadOnly<FChildAlertColumn>()
+		.Where()
+			.Any<FTypedElementSyncFromWorldTag, FTypedElementSyncBackToWorldTag>()
+			.None<FAlertColumn>()
+		.Compile());
+
+	QueryHandle UpdateWidget_Both = DataStorage.RegisterQuery(
+		Select()
+		.ReadOnly<FAlertColumn, FChildAlertColumn>()
+		.Where()
+			.Any<FTypedElementSyncFromWorldTag, FTypedElementSyncBackToWorldTag>()
+		.Compile());
+
+	DataStorage.RegisterQuery(
+		Select(TEXT("Update alert widget state"),
+			FProcessor(EQueryTickPhase::FrameEnd, DataStorage.GetQueryTickGroupName(EQueryTickGroups::SyncWidgets))
+				.SetExecutionMode(EExecutionMode::GameThread),
+			[](IQueryContext& Context,
+				FTypedElementSlateWidgetReferenceColumn& Widget,
+				const FTypedElementRowReferenceColumn& ReferenceColumn)
+			{
+				Context.RunSubquery(0, ReferenceColumn.Row, CreateSubqueryCallbackBinding(
+					[&Widget](ISubqueryContext& Context, RowHandle Row, const FAlertColumn& Alert)
+					{
+						checkf(
+							Alert.AlertType == FAlertColumnType::Warning ||
+							Alert.AlertType == FAlertColumnType::Error,
+							TEXT("Alert column has unsupported type %i"), static_cast<int>(Alert.AlertType));
+						Widgets::Private::UpdateWidget(Widget.Widget.Pin(), Alert.Message,
+							Alert.AlertType == FAlertColumnType::Warning, 0, 0,
+							Context.HasColumn<FAlertActionColumn>() ? Row : InvalidRowHandle);
+					}));
+				Context.RunSubquery(1, ReferenceColumn.Row, CreateSubqueryCallbackBinding(
+					[&Widget](ISubqueryContext& Context, RowHandle Row, const FChildAlertColumn& ChildAlert)
+					{
+						Widgets::Private::UpdateWidget(Widget.Widget.Pin(), FText::GetEmpty(), false,
+							ChildAlert.Counts[static_cast<size_t>(FAlertColumnType::Error)],
+							ChildAlert.Counts[static_cast<size_t>(FAlertColumnType::Warning)],
+							Context.HasColumn<FAlertActionColumn>() ? Row : InvalidRowHandle);
+					}));
+				Context.RunSubquery(2, ReferenceColumn.Row, CreateSubqueryCallbackBinding(
+					[&Widget](
+						ISubqueryContext& Context, RowHandle Row,
+						const FAlertColumn& Alert,
+						const FChildAlertColumn& ChildAlert)
+					{
+						checkf(
+							Alert.AlertType == FAlertColumnType::Warning ||
+							Alert.AlertType == FAlertColumnType::Error,
+							TEXT("Alert column has unsupported type %i"), static_cast<int>(Alert.AlertType));
+						Widgets::Private::UpdateWidget(Widget.Widget.Pin(), Alert.Message,
+							Alert.AlertType == FAlertColumnType::Warning,
+							ChildAlert.Counts[static_cast<size_t>(FAlertColumnType::Error)],
+							ChildAlert.Counts[static_cast<size_t>(FAlertColumnType::Warning)],
+							Context.HasColumn<FAlertActionColumn>() ? Row : InvalidRowHandle);
+					}));
+			}
+		)
+		.Where()
+			.All<FAlertWidgetTag>()
+		.DependsOn()
+			.SubQuery(UpdateWidget_OnlyAlert)
+			.SubQuery(UpdateWidget_OnlyChildAlert)
+			.SubQuery(UpdateWidget_Both)
+		.Compile());
+}
+
+void UAlertWidgetFactory::RegisterAlertHeaderQueries(IEditorDataStorageProvider& DataStorage)
+{
+	using namespace UE::Editor::DataStorage::Columns;
+	using namespace UE::Editor::DataStorage::Queries;
+
+	QueryHandle AlertCount = DataStorage.RegisterQuery(
+		Count()
+		.Where()
+			.Any<FAlertColumn>()
+		.Compile());
+
+	DataStorage.RegisterQuery(
+		Select(TEXT("Update alert header"),
+			FProcessor(EQueryTickPhase::FrameEnd, DataStorage.GetQueryTickGroupName(EQueryTickGroups::SyncWidgets))
+				.SetExecutionMode(EExecutionMode::GameThread),
+			[](IQueryContext& Context, RowHandle Row, FTypedElementSlateWidgetReferenceColumn& Widget)
+			{
+				FQueryResult Result = Context.RunSubquery(0);
+				if (Result.Count > 0)
+				{
+					if (TSharedPtr<SWidget> WidgetPtr = Widget.Widget.Pin())
+					{
+						static_cast<SImage*>(WidgetPtr.Get())->SetImage(FAppStyle::GetBrush("Icons.Warning.Solid"));
+						Context.AddColumns<FAlertHeaderActiveWidgetTag>(Row);
+					}
+				}
+			})
+		.Where()
+			.All<FAlertHeaderWidgetTag>()
+			.None<FAlertHeaderActiveWidgetTag>()
+		.DependsOn()
+			.SubQuery(AlertCount)
+		.Compile());
+
+	DataStorage.RegisterQuery(
+		Select(TEXT("Update active alert header"),
+			FProcessor(EQueryTickPhase::FrameEnd, DataStorage.GetQueryTickGroupName(EQueryTickGroups::SyncWidgets))
+			.SetExecutionMode(EExecutionMode::GameThread),
+			[](IQueryContext& Context, RowHandle Row, FTypedElementSlateWidgetReferenceColumn& Widget)
+			{
+				FQueryResult Result = Context.RunSubquery(0);
+				if (Result.Count == 0)
+				{
+					if (TSharedPtr<SWidget> WidgetPtr = Widget.Widget.Pin())
+					{
+						static_cast<SImage*>(WidgetPtr.Get())->SetImage(FAppStyle::GetBrush("Icons.Alert"));
+						Context.RemoveColumns<FAlertHeaderActiveWidgetTag>(Row);
+					}
+				}
+			}
+		)
+		.Where()
+			.All<FAlertHeaderWidgetTag, FAlertHeaderActiveWidgetTag>()
+		.DependsOn()
+			.SubQuery(AlertCount)
+		.Compile());
+}
+
+
+
+//
+// FAlertWidgetConstructor
+//
+
+FAlertWidgetConstructor::FAlertWidgetConstructor()
+	: Super(StaticStruct())
+{
+}
+
+TSharedPtr<SWidget> FAlertWidgetConstructor::CreateWidget(
+	const UE::Editor::DataStorage::FMetaDataView& Arguments)
+{
+	return SNew(SOverlay)
+		+SOverlay::Slot()
+			.VAlign(VAlign_Fill)
+			.HAlign(HAlign_Fill)
+		[
+			SNew(SButton)
+				.ButtonStyle(FAppStyle::Get(), "FlatButton")
+				.Text(FText::FromString(TEXT("X"))) // There needs to be at least some content otherwise nothing will show.
+				.ForegroundColor(FLinearColor::Transparent) // Then the color needs to be cleared so the X doesn't show.
+				.ContentPadding(FMargin(0.0f))
+		]
+		+SOverlay::Slot()
+			.VAlign(VAlign_Center)
+			.HAlign(HAlign_Center)
+		[
+			SNew(SImage)
+				.Image(FAppStyle::GetBrush("Icons.Alert.Solid"))
+				.DesiredSizeOverride(FVector2D(16.0f, 16.0f))
+		]
+		+SOverlay::Slot()
+			.VAlign(VAlign_Center)
+			.HAlign(HAlign_Center)
+		[
+			SNew(SImage)
+				.Image(FAppStyle::GetBrush("Icons.Alert.Badge"))
+				.DesiredSizeOverride(FVector2D(16.0f, 16.0f))
+		]
+		+SOverlay::Slot()
+			.VAlign(VAlign_Bottom)
+			.HAlign(HAlign_Center)
+		[
+			SNew(STextBlock)
+				.Font(FCoreStyle::GetDefaultFontStyle("Regular", BadgeFontSize))
+				.ColorAndOpacity(FLinearColor::Black)
+				.Margin(FMargin(BadgeHorizontalOffset, BadgeVerticalOffset))
+		];
+}
+
+TConstArrayView<const UScriptStruct*> FAlertWidgetConstructor::GetAdditionalColumnsList() const
+{
+	static TTypedElementColumnTypeList<FAlertWidgetTag> Columns;
+	return Columns;
+}
+
+bool FAlertWidgetConstructor::FinalizeWidget(IEditorDataStorageProvider* DataStorage,
+	IEditorDataStorageUiProvider* DataStorageUi, UE::Editor::DataStorage::RowHandle Row, const TSharedPtr<SWidget>& Widget)
+{
+	using namespace UE::Editor::DataStorage;
+	using namespace UE::Editor::DataStorage::Columns;
+
+	const FTypedElementRowReferenceColumn* RowReferenceColumn = DataStorage->GetColumn<FTypedElementRowReferenceColumn>(Row);
+
+	if (!ensureMsgf(RowReferenceColumn, TEXT("The alert widget requires your widget to have a target row")))
+	{
+		return false;
+	}
+
+	RowHandle TargetRow = RowReferenceColumn->Row;
+	const FAlertColumn* Alert = DataStorage->GetColumn<FAlertColumn>(TargetRow);
+	const FChildAlertColumn* ChildAlert = DataStorage->GetColumn<FChildAlertColumn>(TargetRow);
+
+	uint16 ErrorCount = ChildAlert ? ChildAlert->Counts[static_cast<size_t>(FAlertColumnType::Error)] : 0;
+	uint16 WarningCount = ChildAlert ? ChildAlert->Counts[static_cast<size_t>(FAlertColumnType::Warning)] : 0;
+
+	Widgets::Private::UpdateWidget(
+		Widget,
+		Alert ? Alert->Message : FText::GetEmpty(),
+		Alert ? (Alert->AlertType == FAlertColumnType::Warning) : false,
+		ErrorCount, WarningCount,
+		DataStorage->HasColumns<FAlertActionColumn>(TargetRow) ? TargetRow : InvalidRowHandle);
+
+	return true;
+}
+
+
+
+//
+// FAlertHeaderWidgetConstructor
+//
+
+FAlertHeaderWidgetConstructor::FAlertHeaderWidgetConstructor()
+	: Super(StaticStruct())
+{
+}
+
+TSharedPtr<SWidget> FAlertHeaderWidgetConstructor::CreateWidget(
+	const UE::Editor::DataStorage::FMetaDataView& Arguments)
+{
+	return SNew(SImage)
+		.DesiredSizeOverride(FVector2D(16.f, 16.f))
+		.ColorAndOpacity(FSlateColor::UseForeground())
+		.Image(FAppStyle::GetBrush("Icons.Alert"))
+		.ToolTipText(FText(LOCTEXT("AlertColumnHeader", "Alerts")));
+}
+
+TConstArrayView<const UScriptStruct*> FAlertHeaderWidgetConstructor::GetAdditionalColumnsList() const
+{
+	static TTypedElementColumnTypeList<FAlertHeaderWidgetTag> Columns;
+	return Columns;
+}
+
+bool FAlertHeaderWidgetConstructor::FinalizeWidget(IEditorDataStorageProvider* DataStorage,
+	IEditorDataStorageUiProvider* DataStorageUi, UE::Editor::DataStorage::RowHandle Row, const TSharedPtr<SWidget>& Widget)
+{
+	DataStorage->AddColumn(Row, FHeaderWidgetSizeColumn
+		{
+			.ColumnSizeMode = EColumnSizeMode::Fixed,
+			.Width = 24.0f
+		});
+	return true;
+}
+
+#undef LOCTEXT_NAMESPACE
diff --git a/Engine/Plugins/Experimental/EditorDataStorageFeatures/Source/TedsAlerts/Private/TedsAlertWidget.h b/Engine/Plugins/Experimental/EditorDataStorageFeatures/Source/TedsAlerts/Private/TedsAlertWidget.h
new file mode 100644
--- /dev/null
+++ b/Engine/Plugins/Experimental/EditorDataStorageFeatures/Source/TedsAlerts/Private/TedsAlertWidget.h
@@ -0,0 +1,85 @@
+// Copyright Epic Games, Inc. All Rights Reserved.
+
+#pragma once
+
+#include "Elements/Interfaces/TypedElementDataStorageFactory.h"
+#include "Elements/Common/TypedElementCommonTypes.h"
+#include "Elements/Interfaces/TypedElementDataStorageUiInterface.h"
+
+#include "TedsAlertWidget.generated.h"
+
+UCLASS()
+class UAlertWidgetFactory final : public UEditorDataStorageFactory
+{
+	GENERATED_BODY()
+
+public:
+	~UAlertWidgetFactory() override = default;
+
+	void RegisterWidgetConstructors(IEditorDataStorageProvider& DataStorage,
+		IEditorDataStorageUiProvider& DataStorageUi) const override;
+	void RegisterQueries(IEditorDataStorageProvider& DataStorage) override;
+
+private:
+	void RegisterAlertQueries(IEditorDataStorageProvider& DataStorage);
+	void RegisterAlertHeaderQueries(IEditorDataStorageProvider& DataStorage);
+};
+
+USTRUCT()
+struct FAlertWidgetConstructor final : public FTypedElementWidgetConstructor
+{
+	GENERATED_BODY()
+
+public:
+	static constexpr int32 IconBackgroundSlot = 1;
+	static constexpr int32 IconBadgeSlot = 2;
+	static constexpr int32 CounterTextSlot = 3;
+	static constexpr int32 ActionButtonSlot = 0;
+
+	static constexpr float BadgeFontSize = 7.0f;
+	static constexpr float BadgeHorizontalOffset = 13.0f;
+	static constexpr float BadgeVerticalOffset = 1.0f;
+
+	FAlertWidgetConstructor();
+	~FAlertWidgetConstructor() override = default;
+
+private:
+	TSharedPtr<SWidget> CreateWidget(const UE::Editor::DataStorage::FMetaDataView& Arguments) override;
+	TConstArrayView<const UScriptStruct*> GetAdditionalColumnsList() const override;
+	bool FinalizeWidget(IEditorDataStorageProvider* DataStorage, IEditorDataStorageUiProvider* DataStorageUi,
+		UE::Editor::DataStorage::RowHandle Row, const TSharedPtr<SWidget>& Widget) override;
+};
+
+USTRUCT(meta = (DisplayName = "General purpose alert"))
+struct FAlertWidgetTag : public FEditorDataStorageTag
+{
+	GENERATED_BODY()
+};
+
+USTRUCT()
+struct FAlertHeaderWidgetConstructor final : public FTypedElementWidgetConstructor
+{
+	GENERATED_BODY()
+
+public:
+	FAlertHeaderWidgetConstructor();
+	~FAlertHeaderWidgetConstructor() override = default;
+
+protected:
+	TSharedPtr<SWidget> CreateWidget(const UE::Editor::DataStorage::FMetaDataView& Arguments) override;
+	TConstArrayView<const UScriptStruct*> GetAdditionalColumnsList() const override;
+	bool FinalizeWidget(IEditorDataStorageProvider* DataStorage, IEditorDataStorageUiProvider* DataStorageUi,
+		UE::Editor::DataStorage::RowHandle Row, const TSharedPtr<SWidget>& Widget) override;
+};
+
+USTRUCT(meta = (DisplayName = "General purpose alert header"))
+struct FAlertHeaderWidgetTag : public FEditorDataStorageTag
+{
+	GENERATED_BODY()
+};
+
+USTRUCT(meta = (DisplayName = "Alert header active"))
+struct FAlertHeaderActiveWidgetTag : public FEditorDataStorageTag
+{
+	GENERATED_BODY()
+};
diff --git a/Engine/Plugins/Experimental/EditorDataStorageFeatures/Source/TedsAlerts/Private/TedsAlerts.cpp b/Engine/Plugins/Experimental/EditorDataStorageFeatures/Source/TedsAlerts/Private/TedsAlerts.cpp
new file mode 100644
--- /dev/null
+++ b/Engine/Plugins/Experimental/EditorDataStorageFeatures/Source/TedsAlerts/Private/TedsAlerts.cpp
@@ -0,0 +1,235 @@
+// Copyright Epic Games, Inc. All Rights Reserved.
+
+#include "TedsAlerts.h"
+
+#include "Elements/Columns/TypedElementMiscColumns.h"
+#include "Elements/Common/EditorDataStorageFeatures.h"
+#include "Elements/Framework/TypedElementQueryBuilder.h"
+#include "Elements/Interfaces/TypedElementDataStorageInterface.h"
+#include "TedsAlertsFactory.h"
+
+namespace UE::Editor::DataStorage::Alerts
+{
+	void AddAlert(IEditorDataStorageProvider& DataStorage, RowHandle TargetRow, const FName& Name, FText Message, 
+		Columns::FAlertColumnType Type, uint8 Priority, Columns::FAlertActionCallback Action)
+	{
+		using namespace UE::Editor::DataStorage::Columns;
+
+		// Add the new alert to the unsorted alert chain table and let TEDS sort is out at an appropriate time.
+		if (const UTedsAlertsFactory* AlertsFactory = DataStorage.FindFactory<UTedsAlertsFactory>())
+		{
+			// Add a placeholder alert to be filled in later.
+			if (!DataStorage.HasColumns<FAlertColumn>(TargetRow))
+			{
+				DataStorage.AddColumn<FAlertColumn>(TargetRow);
+			}
+			
+			DataStorage.AddRow(AlertsFactory->GetUnsortedAlertChainTable(),
+				[&DataStorage, TargetRow, Name, Message = MoveTemp(Message), Type, Priority, Action = MoveTemp(Action)]
+				(RowHandle Row) mutable
+				{
+					DataStorage.AddColumn(Row, FAlertColumn
+						{
+							.Message = MoveTemp(Message),
+							.NextAlert = TargetRow,
+							.Name = Name,
+							.AlertType = Type,
+							.Priority = Priority
+						});
+					if (Action)
+					{
+						DataStorage.AddColumn(Row, FAlertActionColumn
+							{
+								.Action = MoveTemp(Action)
+							});
+					}
+				});
+		}
+	}
+
+	void AddAlert(IQueryContext& Context, RowHandle TargetRow, const FName& Name, FText Message,
+		Columns::FAlertColumnType Type, uint8 Priority, Columns::FAlertActionCallback Action)
+	{
+		using namespace UE::Editor::DataStorage;
+		
+		Context.PushCommand([TargetRow, Name, Message = MoveTemp(Message), Type, Priority, Action = MoveTemp(Action)]() mutable
+			{
+				if (IEditorDataStorageProvider* DataStorage = GetMutableDataStorageFeature<IEditorDataStorageProvider>(StorageFeatureName))
+				{
+					AddAlert(*DataStorage, TargetRow, Name, MoveTemp(Message), Type, Priority, MoveTemp(Action));
+				}
+			});
+	}
+
+	void UpdateAlertText(IEditorDataStorageProvider& DataStorage, RowHandle TargetRow, const FName& Name, FText Message)
+	{
+		using namespace UE::Editor::DataStorage::Columns;
+
+		FAlertColumn* Alert = DataStorage.GetColumn<FAlertColumn>(TargetRow);
+		while (Alert)
+		{
+			if (Alert->Name == Name)
+			{
+				Alert->Message = MoveTemp(Message);
+				DataStorage.AddColumns<FTypedElementSyncBackToWorldTag>(TargetRow);
+				return;
+			}
+			TargetRow = Alert->NextAlert;
+			Alert = DataStorage.GetColumn<FAlertColumn>(TargetRow);
+		}
+	}
+
+	void UpdateAlertText(IQueryContext& Context, RowHandle TargetRow, const FName& Name, FText Message)
+	{
+		Context.PushCommand([TargetRow, Name, Message = MoveTemp(Message)]() mutable
+			{
+				if (IEditorDataStorageProvider* DataStorage = GetMutableDataStorageFeature<IEditorDataStorageProvider>(StorageFeatureName))
+				{
+					UpdateAlertText(*DataStorage, TargetRow, Name, MoveTemp(Message));
+				}
+			});
+	}
+
+	void UpdateAlertAction(IEditorDataStorageProvider& DataStorage, RowHandle TargetRow, const FName& Name, 
+		Columns::FAlertActionCallback Action)
+	{
+		using namespace UE::Editor::DataStorage::Columns;
+
+		FAlertColumn* Alert = DataStorage.GetColumn<FAlertColumn>(TargetRow);
+		while (Alert)
+		{
+			if (Alert->Name == Name)
+			{
+				if (Action)
+				{
+					if (FAlertActionColumn* StoredAction = DataStorage.GetColumn<FAlertActionColumn>(TargetRow))
+					{
+						StoredAction->Action = MoveTemp(Action);
+					}
+					else
+					{
+						DataStorage.AddColumn(TargetRow, FAlertActionColumn{ .Action = MoveTemp(Action) });
+					}
+				}
+				else
+				{
+					DataStorage.RemoveColumn<FAlertActionColumn>(TargetRow);
+				}
+				DataStorage.AddColumns<FTypedElementSyncBackToWorldTag>(TargetRow);
+				return;
+			}
+			TargetRow = Alert->NextAlert;
+			Alert = DataStorage.GetColumn<FAlertColumn>(TargetRow);
+		}
+	}
+
+	void UpdateAlertAction(IQueryContext& Context, RowHandle TargetRow, const FName& Name, Columns::FAlertActionCallback Action)
+	{
+		Context.PushCommand([TargetRow, Name, Action = MoveTemp(Action)]() mutable
+			{
+				if (IEditorDataStorageProvider* DataStorage = GetMutableDataStorageFeature<IEditorDataStorageProvider>(StorageFeatureName))
+				{
+					UpdateAlertAction(*DataStorage, TargetRow, Name, MoveTemp(Action));
+				}
+			});
+	}
+
+	void RemoveAlert(IEditorDataStorageProvider& DataStorage, RowHandle TargetRow, const FName& Name)
+	{
+		using namespace UE::Editor::DataStorage::Columns;
+		using namespace UE::Editor::DataStorage::Queries;
+
+		FAlertColumn* PreviousAlert = DataStorage.GetColumn<FAlertColumn>(TargetRow);
+		FAlertColumn* Alert = nullptr;
+		
+		// First check the active alert.
+		if (PreviousAlert)
+		{
+			// This is the currently active alert so replace it with the next one if it exists.
+			if (PreviousAlert->Name == Name)
+			{
+				Alert = DataStorage.GetColumn<FAlertColumn>(PreviousAlert->NextAlert);
+				if (Alert)
+				{
+					// There's another message in the queue so use that message. Remove the copied alert from the alerts table.
+					RowHandle NextAlert = PreviousAlert->NextAlert;
+					*PreviousAlert = MoveTemp(*Alert);
+
+					FAlertActionColumn* Action = DataStorage.GetColumn<FAlertActionColumn>(PreviousAlert->NextAlert);
+					if (Action)
+					{
+						if (FAlertActionColumn* PreviousAction = DataStorage.GetColumn<FAlertActionColumn>(TargetRow))
+						{
+							*PreviousAction = MoveTemp(*Action);
+						}
+						else
+						{
+							DataStorage.AddColumn(TargetRow, MoveTemp(*Action));
+						}
+
+					}
+					else
+					{
+						DataStorage.RemoveColumn<FAlertActionColumn>(TargetRow);
+					}
+
+					DataStorage.RemoveRow(NextAlert);
+				}
+				else
+				{
+					// There are no more messages so delete the current alert.
+					DataStorage.RemoveColumn<FAlertColumn>(TargetRow);
+					DataStorage.RemoveColumn<FAlertActionColumn>(TargetRow);
+				}
+				DataStorage.AddColumns<FTypedElementSyncBackToWorldTag>(TargetRow);
+				return;
+			}
+
+			// The alert is in the alerts table so walk the chain and remove it when found, updating the links as they're found.
+			Alert = DataStorage.GetColumn<FAlertColumn>(PreviousAlert->NextAlert);
+			while (Alert)
+			{
+				if (Alert->Name == Name)
+				{
+					RowHandle CurrentRow = PreviousAlert->NextAlert;
+					PreviousAlert->NextAlert = Alert->NextAlert;
+					DataStorage.RemoveRow(CurrentRow);
+					return;
+				}
+				PreviousAlert = Alert;
+				Alert = DataStorage.GetColumn<FAlertColumn>(PreviousAlert->NextAlert);
+			}
+		}
+
+		// The alert might be queued up for processing so search the unsorted alerts table and remove it from there.
+		if (const UTedsAlertsFactory* AlertsFactory = DataStorage.FindFactory<UTedsAlertsFactory>())
+		{
+			TArray<RowHandle> RemoveRows;
+			DataStorage.RunQuery(AlertsFactory->GetUnsortedAlertChainTable(), CreateDirectQueryCallbackBinding(
+				[&Name, &RemoveRows](RowHandle Row, const FAlertColumn& Alert)
+				{
+					if (Alert.Name == Name)
+					{
+						RemoveRows.Add(Row);
+					}
+				}));
+			for (RowHandle Row : RemoveRows)
+			{
+				DataStorage.RemoveRow(Row);
+			}
+		}
+	}
+
+	void RemoveAlert(IQueryContext& Context, RowHandle TargetRow, const FName& Name)
+	{
+		using namespace UE::Editor::DataStorage;
+
+		Context.PushCommand([TargetRow, Name]() mutable
+			{
+				if (IEditorDataStorageProvider* DataStorage = GetMutableDataStorageFeature<IEditorDataStorageProvider>(StorageFeatureName))
+				{
+					RemoveAlert(*DataStorage, TargetRow, Name);
+				}
+			});
+	}
+} //namespace UE::Editor::DataStorage::Alerts
diff --git a/Engine/Plugins/Experimental/EditorDataStorageFeatures/Source/TedsAlerts/Private/TedsAlertsCommands.cpp b/Engine/Plugins/Experimental/EditorDataStorageFeatures/Source/TedsAlerts/Private/TedsAlertsCommands.cpp
new file mode 100644
--- /dev/null
+++ b/Engine/Plugins/Experimental/EditorDataStorageFeatures/Source/TedsAlerts/Private/TedsAlertsCommands.cpp
@@ -0,0 +1,163 @@
+// Copyright Epic Games, Inc. All Rights Reserved.
+
+#include "Elements/Columns/TypedElementHiearchyColumns.h"
+#include "Elements/Columns/TypedElementMiscColumns.h"
+#include "Elements/Columns/TypedElementSelectionColumns.h"
+#include "Elements/Common/EditorDataStorageFeatures.h"
+#include "Elements/Framework/TypedElementQueryBuilder.h"
+#include "Elements/Interfaces/TypedElementDataStorageInterface.h"
+#include "HAL/IConsoleManager.h"
+#include "TedsAlerts.h"
+#include "TedsAlertColumns.h"
+
+FAutoConsoleCommand AddRandomAlertToRowConsoleCommand(
+	TEXT("TEDS.Debug.AddRandomAlertToSelectedRows"),
+	TEXT("Add random alert to all selected rows."),
+	FConsoleCommandDelegate::CreateLambda(
+		[]()
+		{
+			using namespace UE::Editor::DataStorage::Columns;
+			using namespace UE::Editor::DataStorage::Queries;
+
+			TRACE_CPUPROFILER_EVENT_SCOPE(TEDS.Debug.AddRandomAlertToSelectedRows);
+
+			if (IEditorDataStorageProvider* DataStorage = GetMutableDataStorageFeature<IEditorDataStorageProvider>(StorageFeatureName))
+			{
+				static uint32 Count = 0;
+				static FName TestMessageName = FName("Test message");
+				static QueryHandle Query = [DataStorage]
+					{
+						return DataStorage->RegisterQuery(
+							Select()
+							.Where()
+							.All<FTypedElementSelectionColumn>()
+							.Compile());
+					}();
+
+				TArray<RowHandle> Rows;
+				DataStorage->RunQuery(Query, CreateDirectQueryCallbackBinding(
+					[&Rows](IDirectQueryContext& Context, RowHandle Row)
+					{
+						Rows.Add(Row);
+					}));
+				for (RowHandle Row : Rows)
+				{
+					uint8 Priority = static_cast<uint8>(FMath::RandRange(0, 255));
+					int32 Config = FMath::RandRange(0, 4);
+					FString Message = FString::Printf(TEXT("Test alert %u has priority %u"), Count++, Priority);
+					FAlertColumnType Type = (Config & 0x1) ? FAlertColumnType::Error : FAlertColumnType::Warning;
+					TFunction<void(RowHandle)> Action;
+					if ((Config >> 1) & 0x1)
+					{
+						Action = [Message](RowHandle)
+							{
+								FPlatformMisc::MessageBoxExt(
+									EAppMsgType::Ok,
+									*FString::Printf(TEXT("Example of an alert action for message: \n`%s`"), *Message),
+									TEXT("TEDS.Debug.AddRandomAlertToSelectedRows"));
+							};
+					}
+					Alerts::AddAlert(*DataStorage, Row, TestMessageName, FText::FromString(MoveTemp(Message)), Type, Priority, MoveTemp(Action));
+				}
+			}
+		}
+	));
+
+FAutoConsoleCommand RemoveSelectedAlertsConsoleCommand(
+	TEXT("TEDS.Debug.RemoveSelectedAlerts"),
+	TEXT("Removes the alert from all selected rows. Any queued alerts will replace the removed alert."),
+	FConsoleCommandDelegate::CreateLambda(
+		[]()
+		{
+			using namespace UE::Editor::DataStorage::Columns;
+			using namespace UE::Editor::DataStorage::Queries;
+
+			TRACE_CPUPROFILER_EVENT_SCOPE(TEDS.Debug.ClearSelectedAlerts);
+			static FName TestMessageName = FName("Test message");
+
+			if (IEditorDataStorageProvider* DataStorage = GetMutableDataStorageFeature<IEditorDataStorageProvider>(StorageFeatureName))
+			{
+				static QueryHandle AlertQuery = [DataStorage]
+					{
+						return DataStorage->RegisterQuery(
+							Select()
+							.Where()
+							.All<FAlertColumn, FTypedElementSelectionColumn>()
+							.Compile());
+					}();
+				TArray<RowHandle> Rows;
+				DataStorage->RunQuery(AlertQuery, CreateDirectQueryCallbackBinding(
+					[&Rows](IDirectQueryContext& Context, RowHandle Row)
+					{
+						Rows.Add(Row);
+					}));
+				for (RowHandle Row : Rows)
+				{
+					Alerts::RemoveAlert(*DataStorage, Row, TestMessageName);
+				}
+			}
+		}
+	));
+
+FAutoConsoleCommand ClearAllAlertsConsoleCommand(
+	TEXT("TEDS.Debug.ClearAllAlertInfo"),
+	TEXT("Removes all alerts and child alerts."),
+	FConsoleCommandDelegate::CreateLambda(
+		[]()
+		{
+			using namespace UE::Editor::DataStorage::Columns;
+			using namespace UE::Editor::DataStorage::Queries;
+
+			TRACE_CPUPROFILER_EVENT_SCOPE(TEDS.Debug.ClearAllAlertInfo);
+
+			if (IEditorDataStorageProvider* DataStorage = GetMutableDataStorageFeature<IEditorDataStorageProvider>(StorageFeatureName))
+			{
+				// Remove all alerts on existing rows.
+				{
+					static QueryHandle AlertInfoQuery = [DataStorage]
+						{
+							return DataStorage->RegisterQuery(
+								Select()
+								.Where()
+								.Any<FAlertColumn, FChildAlertColumn>()
+								.Compile());
+						}();
+					TArray<RowHandle> Rows;
+					DataStorage->RunQuery(AlertInfoQuery, CreateDirectQueryCallbackBinding(
+						[&Rows](IDirectQueryContext& Context, RowHandle Row)
+						{
+							Rows.Add(Row);
+						}));
+					for (RowHandle Row : Rows)
+					{
+						DataStorage->RemoveColumn<FAlertColumn>(Row);
+						DataStorage->RemoveColumn<FChildAlertColumn>(Row);
+						DataStorage->RemoveColumn<FAlertActionColumn>(Row);
+						DataStorage->AddColumn<FTypedElementSyncBackToWorldTag>(Row);
+					}
+				}
+
+				// Remove all rows with alert chains.
+				{
+					static QueryHandle AlertChainQuery = [DataStorage]
+						{
+							return DataStorage->RegisterQuery(
+								Select()
+								.Where()
+								.Any<FAlertChainTag>()
+								.Compile());
+						}();
+					TArray<RowHandle> Rows;
+					DataStorage->RunQuery(AlertChainQuery, CreateDirectQueryCallbackBinding(
+						[&Rows](IDirectQueryContext& Context, RowHandle Row)
+						{
+							Rows.Add(Row);
+						}));
+					for (RowHandle Row : Rows)
+					{
+						DataStorage->RemoveRow(Row);
+					}
+				}
+			}
+		}
+	));
diff --git a/Engine/Plugins/Experimental/EditorDataStorageFeatures/Source/TedsAlerts/Private/TedsAlertsFactory.cpp b/Engine/Plugins/Experimental/EditorDataStorageFeatures/Source/TedsAlerts/Private/TedsAlertsFactory.cpp
new file mode 100644
--- /dev/null
+++ b/Engine/Plugins/Experimental/EditorDataStorageFeatures/Source/TedsAlerts/Private/TedsAlertsFactory.cpp
@@ -0,0 +1,502 @@
+// Copyright Epic Games, Inc. All Rights Reserved.
+
+#include "TedsAlertsFactory.h"
+
+#include "Elements/Columns/TypedElementHiearchyColumns.h"
+#include "Elements/Columns/TypedElementMiscColumns.h"
+#include "Elements/Framework/TypedElementQueryBuilder.h"
+#include "Elements/Interfaces/TypedElementDataStorageInterface.h"
+#include "TedsAlerts.h"
+#include "Templates/UnrealTemplate.h"
+
+const FName UTedsAlertsFactory::AlertChainTableName = "Alerts chain";
+const FName UTedsAlertsFactory::UnsortedAlertChainTableName = "Alerts chain (unsorted)";
+const FName UTedsAlertsFactory::AlertConditionName = FName(TEXT("Alerts"));
+
+void UTedsAlertsFactory::RegisterTables(IEditorDataStorageProvider& DataStorage)
+{
+	ChainTable = DataStorage.RegisterTable<FTedsAlertColumn, FTedsAlertChainTag>(AlertChainTableName);
+	UnsortedChainTable = DataStorage.RegisterTable<FTedsUnsortedAlertChainTag>(ChainTable, UnsortedAlertChainTableName);
+}
+
+void UTedsAlertsFactory::RegisterQueries(IEditorDataStorageProvider& DataStorage)
+{
+	RegisterSubQueries(DataStorage);
+	RegisterSortUnsortedAlertsQuery(DataStorage);
+	RegisterOnAddQueries(DataStorage);
+	RegisterOnRemoveQueries(DataStorage);
+	RegisterParentUpdatesQueries(DataStorage);
+	RegisterChildAlertUpdatesQueries(DataStorage);
+}
+
+UE::Editor::DataStorage::TableHandle UTedsAlertsFactory::GetAlertChainTable() const
+{
+	return ChainTable;
+}
+
+UE::Editor::DataStorage::TableHandle UTedsAlertsFactory::GetUnsortedAlertChainTable() const
+{
+	return UnsortedChainTable;
+}
+
+UE::Editor::DataStorage::QueryHandle UTedsAlertsFactory::GetSortedAlertsQuery() const
+{
+	return SortedAlertsQuery;
+}
+
+UE::Editor::DataStorage::QueryHandle UTedsAlertsFactory::GetUnsortedAlertsQuery() const
+{
+	return UnsortedAlertsQuery;
+}
+
+void UTedsAlertsFactory::RegisterSubQueries(IEditorDataStorageProvider& DataStorage)
+{
+	using namespace UE::Editor::DataStorage::Queries;
+
+	SortedAlertsQuery = DataStorage.RegisterQuery(
+		Select()
+			.ReadWrite<FTedsAlertColumn>()
+		.Where()
+			.None<FTedsUnsortedAlertChainTag>()
+		.Compile());
+
+	UnsortedAlertsQuery = DataStorage.RegisterQuery(
+		Select()
+			.ReadWrite<FTedsAlertColumn>()
+		.Where()
+			.All<FTedsUnsortedAlertChainTag>()
+		.Compile());
+	
+	AlertActionQuery = DataStorage.RegisterQuery(
+		Select()
+			.ReadWrite<FTedsAlertActionColumn>()
+		.Compile());
+
+	ChildAlertColumnReadWriteQuery = DataStorage.RegisterQuery(
+		Select()
+			.ReadWrite<FTedsChildAlertColumn>()
+		.Compile());
+
+	ParentReadOnlyQuery = DataStorage.RegisterQuery(
+		Select()
+			.ReadOnly<FTableRowParentColumn>()
+		.Compile());
+}
+
+void UTedsAlertsFactory::RegisterSortUnsortedAlertsQuery(IEditorDataStorageProvider& DataStorage)
+{
+	using namespace UE::Editor::DataStorage::Queries;
+
+	DataStorage.RegisterQuery(
+		Select("Sort unsorted alerts", FProcessor(EQueryTickPhase::PrePhysics, DataStorage.GetQueryTickGroupName(EQueryTickGroups::Update)),
+			[](IQueryContext& NewAlertContext, RowHandle NewAlertRow, FTedsAlertColumn& NewAlert)
+			{
+				RowHandle NextRow = NewAlert.NextAlert;
+				do
+				{
+					// Walk the chain and add link at the appropriate spot. This can include the row with the active alert.
+					NewAlertContext.RunSubquery(0, NextRow, CreateSubqueryCallbackBinding([NewAlertRow, &NewAlert, &NewAlertContext, &NextRow]
+						(ISubqueryContext& Context, RowHandle TargetRow, FTedsAlertColumn& TargetAlert)
+						{
+							if (NewAlert.AlertType > TargetAlert.AlertType ||
+								(NewAlert.AlertType == TargetAlert.AlertType && NewAlert.Priority >= TargetAlert.Priority))
+							{
+								// if the row is not the active alert, then simply link up.
+								// Note that "NextAlert" at this point still contains the row that has the active alert.
+								if (NewAlert.NextAlert == TargetRow)
+								{
+									if (TargetAlert.Message.IsEmpty())
+									{
+										// This is a placeholder so override it.
+										AssignAlert(Context, TargetRow, TargetAlert, NewAlertContext, NewAlertRow, NewAlert, 1);
+									}
+									else
+									{
+										// The active alert needs to be replaced so swap the active alert with the new alert.
+										SwapAlerts(Context, TargetRow, TargetAlert, NewAlertContext, NewAlertRow, NewAlert, 1);
+									}
+									// (Re)calculate child alerts.
+									Context.ActivateQueries(AlertConditionName);
+								}
+								else
+								{
+									// Found a spot in the chain. To avoid having to track the previous entry, just swap the 
+									// alert at that spot with the new one and chain them up.
+									SwapAlerts(Context, TargetRow, TargetAlert, NewAlertContext, NewAlertRow, NewAlert, 1);
+								}
+								NextRow = InvalidRowHandle;
+							}
+							else
+							{
+								if (TargetAlert.NextAlert == InvalidRowHandle)
+								{
+									// End of the chain so append at the end.
+									AppendAlert(TargetAlert, NewAlertContext, NewAlertRow, NewAlert);
+									NextRow = InvalidRowHandle;
+								}
+								else
+								{
+									// Go to the next row in the chain to check that one.
+									NextRow = TargetAlert.NextAlert;
+								}
+							}
+						}));
+				} while (NextRow != InvalidRowHandle);
+			})
+		.Where()
+			.Any<FTedsUnsortedAlertChainTag>()
+		.DependsOn()
+			.SubQuery(SortedAlertsQuery)
+			.SubQuery(AlertActionQuery)
+		.Compile());
+}
+
+void UTedsAlertsFactory::RegisterOnAddQueries(IEditorDataStorageProvider& DataStorage)
+{
+	using namespace UE::Editor::DataStorage::Queries;
+
+	DataStorage.RegisterQuery(
+		Select("Register alert with parent on alert add", FObserver::OnAdd<FTedsAlertColumn>(),
+			[](IQueryContext& Context, RowHandle Row)
+			{
+				Context.ActivateQueries(AlertConditionName);
+			})
+		.Where()
+			.All<FTableRowParentColumn>() // Only need to do an update pass if there are parents.
+			.None<FTedsAlertChainTag>()
+		.Compile());
+
+	DataStorage.RegisterQuery(
+		Select("Register alert with parent on parent add", FObserver::OnAdd<FTableRowParentColumn>(),
+			[](IQueryContext& Context, RowHandle Row)
+			{
+				Context.ActivateQueries(AlertConditionName);
+			})
+		.Where()
+			.Any<FTedsAlertColumn, FTedsChildAlertColumn>()
+			.None<FTedsAlertChainTag>()
+		.Compile());
+}
+
+void UTedsAlertsFactory::RegisterOnRemoveQueries(IEditorDataStorageProvider& DataStorage)
+{
+	using namespace UE::Editor::DataStorage::Queries;
+	
+	DataStorage.RegisterQuery(
+		Select("Remove active alert", FObserver::OnRemove<FTedsAlertColumn>(),
+			[](IQueryContext& Context, RowHandle Row, const FTedsAlertColumn& Alert)
+			{
+				// Delete all entries in the alert chain.
+				RowHandle NextRow = Alert.NextAlert;
+				while (Context.RunSubquery(0, NextRow, CreateSubqueryCallbackBinding(
+					[&NextRow](ISubqueryContext& Context, RowHandle Row, FTedsAlertColumn& Alert)
+					{
+						NextRow = Alert.NextAlert;
+						Context.RemoveRow(Row);
+					})).Count > 0);
+
+
+				// Remove any alerts that haven't been processed yet.
+				Context.RunSubquery(1, CreateSubqueryCallbackBinding(
+					[&Alert](ISubqueryContext& Context, RowHandle Row, FTedsAlertColumn& PendingAlert)
+					{
+						if (Alert.Name == PendingAlert.Name)
+						{
+							Context.RemoveRow(Row);
+						}
+					}));
+
+				// Update any alert parents.
+				Context.ActivateQueries(AlertConditionName);
+			})
+		.Where()
+			.None<FTedsAlertChainTag>()
+		.DependsOn()
+			.SubQuery(SortedAlertsQuery)
+			.SubQuery(UnsortedAlertsQuery)
+		.Compile());
+
+	DataStorage.RegisterQuery(
+		Select("Update alert upon parent removal", FObserver::OnRemove<FTableRowParentColumn>(),
+			[](IQueryContext& Context, RowHandle Row)
+			{
+				Context.ActivateQueries(AlertConditionName);
+			})
+		.Where()
+			.Any<FTedsAlertColumn, FTedsChildAlertColumn>()
+			.None<FTedsAlertChainTag>()
+		.Compile());
+}
+
+void UTedsAlertsFactory::RegisterParentUpdatesQueries(IEditorDataStorageProvider& DataStorage)
+{
+	using namespace UE::Editor::DataStorage::Queries;
+
+	DataStorage.RegisterQuery(
+		Select(
+			"Trigger alert update if alert's parent changed", 
+			FProcessor(EQueryTickPhase::PostPhysics, DataStorage.GetQueryTickGroupName(EQueryTickGroups::Default)),
+			[](IQueryContext& Context, FTedsAlertColumn& Alert, const FTableRowParentColumn& Parent)
+			{
+				if (Alert.CachedParent != Parent.Parent)
+				{
+					Alert.CachedParent = Parent.Parent;
+					Context.ActivateQueries(AlertConditionName);
+				}
+			})
+		.Where()
+			.Any<FTypedElementSyncBackToWorldTag, FTypedElementSyncFromWorldTag>()
+			.None<FTedsAlertChainTag>()
+		.Compile());
+
+	DataStorage.RegisterQuery(
+		Select(
+			"Trigger alert update if child alert's parent changed",
+			FProcessor(EQueryTickPhase::PostPhysics, DataStorage.GetQueryTickGroupName(EQueryTickGroups::Default)),
+			[](IQueryContext& Context, FTedsChildAlertColumn& ChildAlert, const FTableRowParentColumn& Parent)
+			{
+				if (ChildAlert.CachedParent != Parent.Parent)
+				{
+					ChildAlert.CachedParent = Parent.Parent;
+					Context.ActivateQueries(AlertConditionName);
+				}
+			})
+		.Where()
+			.Any<FTypedElementSyncBackToWorldTag, FTypedElementSyncFromWorldTag>()
+			.None<FTedsAlertChainTag>()
+		.Compile());
+}
+
+void UTedsAlertsFactory::RegisterChildAlertUpdatesQueries(IEditorDataStorageProvider& DataStorage)
+{
+	using namespace UE::Editor::DataStorage::Queries;
+
+	DataStorage.RegisterQuery(
+		Select(
+			"Add missing child alerts",
+			FPhaseAmble(FPhaseAmble::ELocation::Preamble, EQueryTickPhase::PostPhysics)
+				.MakeActivatable(AlertConditionName),
+			[](IQueryContext& Context, RowHandle Row, FTedsAlertColumn& Alert, const FTableRowParentColumn& Parent)
+			{
+				AddChildAlertsToHierarchy(Context, Parent.Parent, 0);
+			})
+		.Where()
+			.None<FTedsAlertChainTag>()
+		.DependsOn()
+			.SubQuery(ParentReadOnlyQuery)
+		.Compile());
+
+	DataStorage.RegisterQuery(
+		Select(
+			"Clear child alerts",
+			FProcessor(EQueryTickPhase::PostPhysics, DataStorage.GetQueryTickGroupName(EQueryTickGroups::PreUpdate))
+				.MakeActivatable(AlertConditionName),
+			[](IQueryContext& Context, RowHandle Row, FTedsChildAlertColumn& ChildAlert)
+			{
+				ResetChildAlertCounters(ChildAlert);
+			})
+		.Where()
+			.None<FTedsAlertChainTag>()
+		.Compile());
+
+	DataStorage.RegisterQuery(
+		Select(
+			"Increment child alerts",
+			FProcessor(EQueryTickPhase::PostPhysics, DataStorage.GetQueryTickGroupName(EQueryTickGroups::Update))
+				.MakeActivatable(AlertConditionName),
+			[](IQueryContext& Context, RowHandle Row, FTedsAlertColumn& Alert)
+			{
+				IncrementParents(Context, Alert.CachedParent, Alert.AlertType, 0);
+			})
+		.Where()
+			.None<FTedsAlertChainTag>()
+		.DependsOn()
+			.SubQuery(ChildAlertColumnReadWriteQuery)
+		.Compile());
+
+	DataStorage.RegisterQuery(
+		Select(
+			"Remove unused child alerts",
+			FProcessor(EQueryTickPhase::PostPhysics, DataStorage.GetQueryTickGroupName(EQueryTickGroups::PostUpdate))
+				.MakeActivatable(AlertConditionName),
+			[](IQueryContext& Context, RowHandle Row, FTedsChildAlertColumn& ChildAlert)
+			{
+				for (size_t It = 0; It < static_cast<size_t>(FTedsAlertColumnType::MAX); ++It)
+				{
+					if (ChildAlert.Counts[It] != 0)
+					{
+						return;
+					}
+				}
+				Context.RemoveColumns<FTedsChildAlertColumn>(Row);
+				Context.AddColumns<FTypedElementSyncBackToWorldTag>(Row);
+			})
+		.Compile());
+}
+
+void UTedsAlertsFactory::AssignAlert(
+	UE::Editor::DataStorage::ISubqueryContext& TargetContext, UE::Editor::DataStorage::RowHandle TargetRow, FTedsAlertColumn& Target,
+	UE::Editor::DataStorage::IQueryContext& SourceContext, UE::Editor::DataStorage::RowHandle SourceRow, FTedsAlertColumn& Source,
+	int32 AlertActionQueryIndex)
+{
+	Target = MoveTemp(Source);
+	Target.NextAlert = UE::Editor::DataStorage::InvalidRowHandle;
+	
+	FTedsAlertActionColumn* TargetAction = GetAction(SourceContext, TargetRow, AlertActionQueryIndex);
+	FTedsAlertActionColumn* SourceAction = GetAction(SourceContext, SourceRow, AlertActionQueryIndex);
+	if (SourceAction)
+	{
+		if (TargetAction)
+		{
+			*TargetAction = MoveTemp(*SourceAction);
+		}
+		else
+		{
+			SourceContext.AddColumn(TargetRow, MoveTemp(*SourceAction));
+		}
+	}
+	else
+	{
+		if (TargetAction)
+		{
+			TargetContext.RemoveColumns<FTedsAlertActionColumn>(TargetRow);
+		}
+	}
+
+	// Notify UI.
+	TargetContext.AddColumns<FTypedElementSyncBackToWorldTag>(TargetRow);
+	// Fully absorbed into the placeholder so the addition to the chain is no longer needed.
+	SourceContext.RemoveRow(SourceRow);
+}
+
+void UTedsAlertsFactory::SwapAlerts(
+	UE::Editor::DataStorage::ISubqueryContext& OriginalContext, UE::Editor::DataStorage::RowHandle OriginalRow, FTedsAlertColumn& Original,
+	UE::Editor::DataStorage::IQueryContext& NewContext, UE::Editor::DataStorage::RowHandle NewRow, FTedsAlertColumn& New,
+	int32 AlertActionQueryIndex)
+{
+	Swap(Original, New);
+	
+	// Swap the alert action columns if they exist.
+	{
+		FTedsAlertActionColumn* OriginalAction = GetAction(NewContext, OriginalRow, AlertActionQueryIndex);
+		FTedsAlertActionColumn* NewAction = GetAction(NewContext, NewRow, AlertActionQueryIndex);
+		if (NewAction)
+		{
+			if (OriginalAction)
+			{
+				Swap(*OriginalAction, *NewAction);
+			}
+			else
+			{
+				OriginalContext.AddColumn(OriginalRow, MoveTemp(*NewAction));
+				NewContext.RemoveColumns<FTedsAlertActionColumn>(NewRow);
+			}
+		}
+		else
+		{
+			if (OriginalAction)
+			{
+				NewContext.AddColumn(NewRow, MoveTemp(*OriginalAction));
+				OriginalContext.RemoveColumns<FTedsAlertActionColumn>(OriginalRow);
+			}
+		}
+	}
+
+	Original.NextAlert = NewRow;
+	// Notify UI. Most of the time this is not needed as only the active alert should be used in the UI,
+	// but the TEDS Debugger might be showing the other alerts, so make sure they're updated to prevent
+	// present invalid data or worse in the case of a action which can cause a crash.
+	OriginalContext.AddColumns<FTypedElementSyncBackToWorldTag>(OriginalRow);
+	NewContext.AddColumns<FTypedElementSyncBackToWorldTag>(NewRow);
+	NewContext.RemoveColumns<FTedsUnsortedAlertChainTag>(NewRow);
+}
+
+void UTedsAlertsFactory::AppendAlert(
+	FTedsAlertColumn& LastAlert,
+	UE::Editor::DataStorage::IQueryContext& AdditionalAlertContext,
+	UE::Editor::DataStorage::RowHandle AdditionalAlertRow,
+	FTedsAlertColumn& AdditionalAlert)
+{
+	AdditionalAlert.NextAlert = UE::Editor::DataStorage::InvalidRowHandle;
+	LastAlert.NextAlert = AdditionalAlertRow;
+	AdditionalAlertContext.RemoveColumns<FTedsUnsortedAlertChainTag>(AdditionalAlertRow);
+}
+
+FTedsAlertActionColumn* UTedsAlertsFactory::GetAction(UE::Editor::DataStorage::IQueryContext& Context,
+	UE::Editor::DataStorage::RowHandle Row, int32 ChildAlertQueryIndex)
+{
+	using namespace UE::Editor::DataStorage::Queries;
+	
+	FTedsAlertActionColumn* Result = nullptr;
+	Context.RunSubquery(ChildAlertQueryIndex, Row, CreateSubqueryCallbackBinding(
+		[&Result](ISubqueryContext& Context, FTedsAlertActionColumn& Action)
+		{
+			Result = &Action;
+		}));
+	return Result;
+}
+
+void UTedsAlertsFactory::AddChildAlertsToHierarchy(
+	UE::Editor::DataStorage::IQueryContext& Context, UE::Editor::DataStorage::RowHandle Parent, int32 ParentQueryIndex)
+{
+	using namespace UE::Editor::DataStorage;
+
+	bool bHasParent = true;
+	do
+	{
+		RowHandle NextParent = Parent;
+		bHasParent = MoveToNextParent(NextParent, Context, ParentQueryIndex);
+
+		// Check if a child alert column exists and add one if not.
+		if (!Context.HasColumn<FTedsChildAlertColumn>(Parent))
+		{
+			FTedsChildAlertColumn ChildAlert;
+			ResetChildAlertCounters(ChildAlert);
+			ChildAlert.CachedParent = NextParent;
+			Context.AddColumn(Parent, MoveTemp(ChildAlert));
+		}
+		Parent = NextParent;
+	} while (bHasParent);
+}
+
+void UTedsAlertsFactory::IncrementParents(
+	UE::Editor::DataStorage::IQueryContext& Context, UE::Editor::DataStorage::RowHandle Row, FTedsAlertColumnType AlertType,
+	int32 ChildAlertQueryIndex)
+{
+	using namespace UE::Editor::DataStorage::Queries;
+
+	while (Context.IsRowAvailable(Row))
+	{
+		FQueryResult Result = Context.RunSubquery(ChildAlertQueryIndex, Row, CreateSubqueryCallbackBinding(
+			[&ParentRow = Row, AlertType](
+				ISubqueryContext& Context, RowHandle Row, FTedsChildAlertColumn& ChildAlert)
+			{
+				ChildAlert.Counts[static_cast<size_t>(AlertType)]++;
+				Context.AddColumns<FTypedElementSyncBackToWorldTag>(Row);
+				ParentRow = ChildAlert.CachedParent;
+			}));
+		checkf(Result.Count > 0, TEXT("Expected to be able to setup the child alert, but it was missing on the parent column."));
+	}
+}
+
+bool UTedsAlertsFactory::MoveToNextParent(
+	UE::Editor::DataStorage::RowHandle& Parent, UE::Editor::DataStorage::IQueryContext& Context, int32 SubQueryIndex)
+{
+	using namespace UE::Editor::DataStorage::Queries;
+
+	FQueryResult Result = Context.RunSubquery(SubQueryIndex, Parent, CreateSubqueryCallbackBinding(
+		[&Parent](const FTableRowParentColumn& NextParent)
+		{
+			Parent = NextParent.Parent;
+		}));
+	Parent = Result.Count != 0 ? Parent : InvalidRowHandle;
+	return Result.Count != 0;
+}
+
+void UTedsAlertsFactory::ResetChildAlertCounters(FTedsChildAlertColumn& ChildAlert)
+{
+	for (size_t It = 0; It < static_cast<size_t>(FTedsAlertColumnType::MAX); ++It)
+	{
+		ChildAlert.Counts[It] = 0;
+	}
+}
diff --git a/Engine/Plugins/Experimental/EditorDataStorageFeatures/Source/TedsAlerts/Private/TedsAlertsFactory.h b/Engine/Plugins/Experimental/EditorDataStorageFeatures/Source/TedsAlerts/Private/TedsAlertsFactory.h
new file mode 100644
--- /dev/null
+++ b/Engine/Plugins/Experimental/EditorDataStorageFeatures/Source/TedsAlerts/Private/TedsAlertsFactory.h
@@ -0,0 +1,78 @@
+// Copyright Epic Games, Inc. All Rights Reserved.
+
+#pragma once
+
+#include "Elements/Common/TypedElementHandles.h"
+#include "Elements/Interfaces/TypedElementDataStorageFactory.h"
+#include "TedsAlertColumns.h"
+#include "UObject/ObjectMacros.h"
+
+#include "TedsAlertsFactory.generated.h"
+
+/**
+ * Factory that manages tables, queries and any other data for alerts.
+ */
+UCLASS()
+class UTedsAlertsFactory final : public UEditorDataStorageFactory
+{
+	GENERATED_BODY()
+
+public:
+	static const FName AlertChainTableName;
+	static const FName UnsortedAlertChainTableName;
+
+	~UTedsAlertsFactory() override = default;
+	
+	void RegisterTables(IEditorDataStorageProvider& DataStorage) override;
+	void RegisterQueries(IEditorDataStorageProvider& DataStorage) override;
+
+	UE::Editor::DataStorage::TableHandle GetAlertChainTable() const;
+	UE::Editor::DataStorage::TableHandle GetUnsortedAlertChainTable() const;
+
+	UE::Editor::DataStorage::QueryHandle GetSortedAlertsQuery() const;
+	UE::Editor::DataStorage::QueryHandle GetUnsortedAlertsQuery() const;
+
+private:
+	static const FName AlertConditionName;
+
+	UE::Editor::DataStorage::TableHandle ChainTable = UE::Editor::DataStorage::InvalidTableHandle;
+	UE::Editor::DataStorage::TableHandle UnsortedChainTable = UE::Editor::DataStorage::InvalidTableHandle;
+	
+	UE::Editor::DataStorage::QueryHandle SortedAlertsQuery = UE::Editor::DataStorage::InvalidQueryHandle;
+	UE::Editor::DataStorage::QueryHandle UnsortedAlertsQuery = UE::Editor::DataStorage::InvalidQueryHandle;
+	UE::Editor::DataStorage::QueryHandle AlertActionQuery = UE::Editor::DataStorage::InvalidQueryHandle;
+	UE::Editor::DataStorage::QueryHandle ChildAlertColumnReadWriteQuery = UE::Editor::DataStorage::InvalidTableHandle;
+	UE::Editor::DataStorage::QueryHandle ParentReadOnlyQuery = UE::Editor::DataStorage::InvalidTableHandle;
+
+	void RegisterSubQueries(IEditorDataStorageProvider& DataStorage);
+	void RegisterSortUnsortedAlertsQuery(IEditorDataStorageProvider& DataStorage);
+	void RegisterOnAddQueries(IEditorDataStorageProvider& DataStorage);
+	void RegisterOnRemoveQueries(IEditorDataStorageProvider& DataStorage);
+	void RegisterParentUpdatesQueries(IEditorDataStorageProvider& DataStorage);
+	void RegisterChildAlertUpdatesQueries(IEditorDataStorageProvider& DataStorage);
+
+	static void AssignAlert(
+		UE::Editor::DataStorage::ISubqueryContext& TargetContext, UE::Editor::DataStorage::RowHandle TargetRow, FTedsAlertColumn& Target,
+		UE::Editor::DataStorage::IQueryContext& SourceContext, UE::Editor::DataStorage::RowHandle SourceRow, FTedsAlertColumn& Source,
+		int32 AlertActionQueryIndex);
+	static void SwapAlerts(
+		UE::Editor::DataStorage::ISubqueryContext& OriginalContext, UE::Editor::DataStorage::RowHandle OriginalRow, FTedsAlertColumn& Original,
+		UE::Editor::DataStorage::IQueryContext& NewContext, UE::Editor::DataStorage::RowHandle NewRow, FTedsAlertColumn& New,
+		int32 AlertActionQueryIndex);
+	static void AppendAlert(
+		FTedsAlertColumn& LastAlert, 
+		UE::Editor::DataStorage::IQueryContext& AdditionalAlertContext, 
+		UE::Editor::DataStorage::RowHandle AdditionalAlertRow, 
+		FTedsAlertColumn& AdditionalAlert);
+	static FTedsAlertActionColumn* GetAction(UE::Editor::DataStorage::IQueryContext& Context,
+		UE::Editor::DataStorage::RowHandle Row, int32 ChildAlertQueryIndex);
+
+	static void AddChildAlertsToHierarchy(
+		UE::Editor::DataStorage::IQueryContext& Context, UE::Editor::DataStorage::RowHandle Parent, int32 ParentQueryIndex);
+	static void IncrementParents(
+		UE::Editor::DataStorage::IQueryContext& Context, UE::Editor::DataStorage::RowHandle Row, FTedsAlertColumnType AlertType,
+		int32 ChildAlertQueryIndex);
+	static bool MoveToNextParent(
+		UE::Editor::DataStorage::RowHandle& Parent, UE::Editor::DataStorage::IQueryContext& Context, int32 SubQueryIndex);
+	static void ResetChildAlertCounters(FTedsChildAlertColumn& ChildAlert);
+};
diff --git a/Engine/Plugins/Experimental/EditorDataStorageFeatures/Source/TedsAlerts/Private/TedsAlertsModule.cpp b/Engine/Plugins/Experimental/EditorDataStorageFeatures/Source/TedsAlerts/Private/TedsAlertsModule.cpp
new file mode 100644
--- /dev/null
+++ b/Engine/Plugins/Experimental/EditorDataStorageFeatures/Source/TedsAlerts/Private/TedsAlertsModule.cpp
@@ -0,0 +1,22 @@
+// Copyright Epic Games, Inc. All Rights Reserved.
+
+#include "TedsAlertsModule.h"
+
+#include "Modules/ModuleManager.h"
+
+namespace UE::Editor::DataStorage
+{
+	void FTedsAlertsModule::StartupModule()
+	{
+		IModuleInterface::StartupModule();
+
+		FModuleManager::Get().LoadModule(TEXT("TypedElementFramework"));	
+	}
+
+	void FTedsAlertsModule::ShutdownModule()
+	{
+		IModuleInterface::ShutdownModule();
+	}
+} // namespace UE::Editor::DataStorage
+
+IMPLEMENT_MODULE(UE::Editor::DataStorage::FTedsAlertsModule, TedsAlerts);
diff --git a/Engine/Plugins/Experimental/EditorDataStorageFeatures/Source/TedsAlerts/Private/TedsAlertsModule.h b/Engine/Plugins/Experimental/EditorDataStorageFeatures/Source/TedsAlerts/Private/TedsAlertsModule.h
new file mode 100644
--- /dev/null
+++ b/Engine/Plugins/Experimental/EditorDataStorageFeatures/Source/TedsAlerts/Private/TedsAlertsModule.h
@@ -0,0 +1,20 @@
+// Copyright Epic Games, Inc. All Rights Reserved.
+
+#pragma once
+
+#include "Modules/ModuleInterface.h"
+
+namespace UE::Editor::DataStorage
+{
+	class FTedsAlertsModule
+		: public IModuleInterface
+	{
+	public:
+
+		FTedsAlertsModule() = default;
+
+		// IModuleInterface interface
+		virtual void StartupModule() override;
+		virtual void ShutdownModule() override;
+	};
+} // namespace UE::Editor::DataStorage
diff --git a/Engine/Plugins/Experimental/EditorDataStorageFeatures/Source/TedsAlerts/Public/TedsAlertColumns.h b/Engine/Plugins/Experimental/EditorDataStorageFeatures/Source/TedsAlerts/Public/TedsAlertColumns.h
new file mode 100644
--- /dev/null
+++ b/Engine/Plugins/Experimental/EditorDataStorageFeatures/Source/TedsAlerts/Public/TedsAlertColumns.h
@@ -0,0 +1,113 @@
+// Copyright Epic Games, Inc. All Rights Reserved.
+
+#pragma once
+
+#include "Elements/Common/TypedElementHandles.h"
+#include "Elements/Interfaces/TypedElementDataStorageInterface.h"
+#include "Internationalization/Text.h"
+#include "UObject/ObjectMacros.h"
+
+#include "TedsAlertColumns.generated.h"
+
+UENUM()
+enum class FTedsAlertColumnType : uint8
+{
+	Warning,
+	Error,
+
+	MAX
+};
+
+using FTedsAlertActionCallback = TFunction<void(UE::Editor::DataStorage::RowHandle)>;
+
+/**
+ * Column containing information a user needs to be alerted of.
+ */
+USTRUCT(meta = (DisplayName = "Alert"))
+struct FTedsAlertColumn final : public FEditorDataStorageColumn
+{
+	GENERATED_BODY()
+
+	UPROPERTY()
+	FText Message;
+
+	// Store a copy of the parent row so it's possible to detect if a row has been reparented.
+	UE::Editor::DataStorage::RowHandle CachedParent;
+
+	/**
+	 * If valid, points to the next alert in the chain.The next alert will take the place of the current
+	 * alert if this alert is cleared.
+	 */
+	UE::Editor::DataStorage::RowHandle NextAlert;
+
+	/** Unique name to identify the message with. */
+	UPROPERTY()
+	FName Name;
+
+	/** The type of alert. This is used for ordering and to show appropriate icons. */
+	UPROPERTY(meta = (IgnoreForMemberInitializationTest))
+	FTedsAlertColumnType AlertType;
+
+	/**
+	 * A chain of priorities is sorted by errors, then warnings. If either group has multiple alerts, alerts are sorted by priority
+	 * with the lowest value (0) given the later location and higher value (up to 255) the earlier locations in the chain.
+	 */
+	UPROPERTY(meta = (IgnoreForMemberInitializationTest))
+	uint8 Priority;
+};
+
+/**
+ * Column containing a count for the number of alerts any child rows have.
+ */
+USTRUCT(meta = (DisplayName = "Child alert"))
+struct FTedsChildAlertColumn final : public FEditorDataStorageColumn
+{
+	GENERATED_BODY()
+
+	// Store a copy of the parent row so it's possible to detect if a row has been reparented.
+	UE::Editor::DataStorage::RowHandle CachedParent;
+
+	uint16 Counts[static_cast<size_t>(FTedsAlertColumnType::MAX)];
+};
+
+/**
+ * Column that can be added to an alert column to have it trigger an action when the alert is clicked.
+ */
+USTRUCT(meta = (DisplayName = "Alert action"))
+struct FTedsAlertActionColumn final : public FEditorDataStorageColumn
+{
+	GENERATED_BODY()
+
+	FTedsAlertActionCallback Action;
+};
+
+/**
+ * Tag to indicate this row contains alert chain information.
+ */
+USTRUCT(meta = (DisplayName = "Alert chain"))
+struct FTedsAlertChainTag final : public FEditorDataStorageTag
+{
+	GENERATED_BODY()
+};
+
+/**
+ * Tag to indicate this row contains alert chain information that has not been added/sorted yet. This means that the
+ * "NextAlert" in the Alert column points to the target row, not the row for the next alert.
+ */
+USTRUCT(meta = (DisplayName = "Unsorted alert chain"))
+struct FTedsUnsortedAlertChainTag final : public FEditorDataStorageTag
+{
+	GENERATED_BODY()
+};
+
+namespace UE::Editor::DataStorage::Columns
+{
+	using FAlertColumnType = FTedsAlertColumnType;
+	using FAlertActionCallback = FTedsAlertActionCallback;
+	
+	using FAlertColumn = FTedsAlertColumn;
+	using FChildAlertColumn = FTedsChildAlertColumn;
+	using FAlertActionColumn = FTedsAlertActionColumn;
+	using FAlertChainTag = FTedsAlertChainTag;
+	using FUnsortedAlertChainTag = FTedsUnsortedAlertChainTag;
+}
diff --git a/Engine/Plugins/Experimental/EditorDataStorageFeatures/Source/TedsAlerts/Public/TedsAlerts.h b/Engine/Plugins/Experimental/EditorDataStorageFeatures/Source/TedsAlerts/Public/TedsAlerts.h
new file mode 100644
--- /dev/null
+++ b/Engine/Plugins/Experimental/EditorDataStorageFeatures/Source/TedsAlerts/Public/TedsAlerts.h
@@ -0,0 +1,107 @@
+// Copyright Epic Games, Inc. All Rights Reserved.
+
+#pragma once
+
+#include "Elements/Common/TypedElementHandles.h"
+#include "Internationalization/Text.h"
+#include "TedsAlertColumns.h"
+#include "UObject/NameTypes.h"
+
+class IEditorDataStorageProvider;
+
+namespace UE::Editor::DataStorage
+{
+	struct IQueryContext;
+}
+
+namespace UE::Editor::DataStorage::Alerts
+{
+	/**
+	 * Add an alert to the provided row. Various tools like the Outliner monitor for alerts and can show them for instance as
+	 * a colored triangle with a exclamation count. If a row also has a parent row handle, the alerts will propagate their count
+	 * to the parent.
+	 * @DataStorage A reference to the data storage instance to store the alert in.
+	 * @TargetRow The target row to store the alert on. If the row already has an alert, an ordered chain of alerts will be created.
+	 * @Name A name to uniquely identify the alert. If there's a chain of alerts this is used to find the correct alert to remove.
+	 * @Message The message shown in the alert, e.g. as a tool tip on the alert icon.
+	 * @Type The type of alerts. Widgets will use separate child counts per type and use different colors such as red for errors and 
+	 *		yellow for warnings.
+	 * @Priority If there are multiple alerts for the same type, the priority is used to order the alerts so the most important
+	 *		one is shown when the active alert is removed. Higher values indicate higher priority.
+	 * @Action An optional callback that gets triggered when the user presses the alert icon.
+	 */
+	TEDSALERTS_API void AddAlert(IEditorDataStorageProvider& DataStorage, RowHandle TargetRow, const FName& Name, FText Message, 
+		Columns::FAlertColumnType Type, uint8 Priority = 127, Columns::FAlertActionCallback Action = {});
+	/**
+	 * Add an alert to the provided row. Various tools like the Outliner monitor for alerts and can show them for instance as
+	 * a colored triangle with a exclamation count. If a row also has a parent row handle, the alerts will propagate their count
+	 * to the parent.
+	 * @Context A reference to the query callback context.
+	 * @TargetRow The target row to store the alert on. If the row already has an alert, an ordered chain of alerts will be created.
+	 * @Name A name to uniquely identify the alert. If there's a chain of alerts this is used to find the correct alert to remove.
+	 * @Message The message shown in the alert, e.g. as a tool tip on the alert icon.
+	 * @Type The type of alerts. Widgets will use separate child counts per type and use different colors such as red for errors and
+	 *		yellow for warnings.
+	 * @Priority If there are multiple alerts for the same type, the priority is used to order the alerts so the most important
+	 *		one is shown when the active alert is removed. Higher values indicate higher priority.
+	 * @Action An optional callback that gets triggered when the user presses the alert icon.
+	 */
+	TEDSALERTS_API void AddAlert(IQueryContext& Context, RowHandle TargetRow, const FName& Name, FText Message,
+		Columns::FAlertColumnType Type, uint8 Priority = 127, Columns::FAlertActionCallback Action = {});
+	
+	/**
+	 * Locates the first alerts with the provided name in the alert chain found at the provided row and updates its text.
+	 * @DataStorage A reference to the data storage instance to store the alert in.
+	 * @TargetRow The row the alert is stored on.
+	 * @Name The name that uniquely identifies the alert.
+	 * @Message The new message for the alert.
+	 */
+	TEDSALERTS_API void UpdateAlertText(IEditorDataStorageProvider& DataStorage, RowHandle TargetRow, const FName& Name, FText Message);
+	/**
+	 * Locates the first alerts with the provided name in the alert chain found at the provided row and updates its text.
+	 * @Context A reference to the query callback context.
+	 * @TargetRow The row the alert is stored on.
+	 * @Name The name that uniquely identifies the alert.
+	 * @Message The new message for the alert.
+	 */
+	TEDSALERTS_API void UpdateAlertText(IQueryContext& Context, RowHandle TargetRow, const FName& Name, FText Message);
+
+	/**
+	 * Locates the first alerts with the provided name in the alert chain found at the provided row and updates its action.
+	 * If there's no action it will be added. If the provided action isn't bound, the alert action will be removed.
+	 * @DataStorage A reference to the data storage instance to store the alert in.
+	 * @TargetRow The row the alert is stored on.
+	 * @Name The name that uniquely identifies the alert.
+	 * @Action The new action to trigger.
+	 */
+	TEDSALERTS_API void UpdateAlertAction(
+		IEditorDataStorageProvider& DataStorage, RowHandle TargetRow, const FName& Name, Columns::FAlertActionCallback Action);
+	/**
+	 * Locates the first alerts with the provided name in the alert chain found at the provided row and updates its action.
+	 * If there's no action it will be added. If the provided action isn't bound, the alert action will be removed.
+	 * @Context A reference to the query callback context.
+	 * @TargetRow The row the alert is stored on.
+	 * @Name The name that uniquely identifies the alert.
+	 * @Action The new action to trigger.
+	 */
+	TEDSALERTS_API void UpdateAlertAction(IQueryContext& Context, RowHandle TargetRow, const FName& Name, Columns::FAlertActionCallback Action);
+	
+	/**
+	 * Searches the alert chain found at the target row for the first alert with the provided name and removes it. The chain will
+	 * automatically be patched, keeping the chain in order of priority. If the root alert is removed, the UI will be automatically
+	 * updated with the next alert in line or be cleared if there are no more alerts.
+	 * @DataStorage A reference to the data storage instance to store the alert in.
+	 * @TargetRow The row the alert is stored on.
+	 * @Name The name that uniquely identifies the alert.
+	 */
+	TEDSALERTS_API void RemoveAlert(IEditorDataStorageProvider& DataStorage, RowHandle TargetRow, const FName& Name);
+	/**
+	 * Searches the alert chain found at the target row for the first alert with the provided name and removes it. The chain will
+	 * automatically be patched, keeping the chain in order of priority. If the root alert is removed, the UI will be automatically
+	 * updated with the next alert in line or be cleared if there are no more alerts.
+	 * @Context A reference to the query callback context.
+	 * @TargetRow The row the alert is stored on.
+	 * @Name The name that uniquely identifies the alert.
+	 */
+	TEDSALERTS_API void RemoveAlert(IQueryContext& Context, RowHandle TargetRow, const FName& Name);
+} // namespace UE::Editor::DataStorage::Alerts
diff --git a/Engine/Plugins/Experimental/EditorDataStorageFeatures/Source/TedsAlerts/TedsAlerts.build.cs b/Engine/Plugins/Experimental/EditorDataStorageFeatures/Source/TedsAlerts/TedsAlerts.build.cs
new file mode 100644
--- /dev/null
+++ b/Engine/Plugins/Experimental/EditorDataStorageFeatures/Source/TedsAlerts/TedsAlerts.build.cs
@@ -0,0 +1,39 @@
+// Copyright Epic Games, Inc. All Rights Reserved.
+
+using UnrealBuildTool;
+
+// Experimental test module. please refrain from depending on it until this warning is removed
+public class TedsAlerts : ModuleRules
+{
+	public TedsAlerts(ReadOnlyTargetRules Target) : base(Target)
+	{
+		ShortName = "TedsAlerts";
+
+		PCHUsage = ModuleRules.PCHUsageMode.UseExplicitOrSharedPCHs;
+
+		if (Target.bBuildEditor)
+		{
+			PublicIncludePaths.AddRange(new string[] { });
+			PrivateIncludePaths.AddRange(new string[] { });
+
+			PrivateDependencyModuleNames.AddRange(
+				new string[]
+				{
+					"Core",
+					"CoreUObject",
+					"Engine",
+					"Slate",
+					"SlateCore",
+					"ToolWidgets",
+					"TypedElementFramework"
+				});
+
+			PublicDependencyModuleNames.AddRange(
+				new string[]
+				{
+					"Projects",
+				});
+
+		}
+	}
+}
diff --git a/Engine/Plugins/Experimental/EditorDataStorageFeatures/Source/TedsAssetData/Private/Processors/AssetProcessors.cpp b/Engine/Plugins/Experimental/EditorDataStorageFeatures/Source/TedsAssetData/Private/Processors/AssetProcessors.cpp
--- a/Engine/Plugins/Experimental/EditorDataStorageFeatures/Source/TedsAssetData/Private/Processors/AssetProcessors.cpp
+++ b/Engine/Plugins/Experimental/EditorDataStorageFeatures/Source/TedsAssetData/Private/Processors/AssetProcessors.cpp
@@ -14,10 +14,10 @@
 #include "Elements/Interfaces/TypedElementQueryStorageInterfaces.h"
 #include "Experimental/ContentBrowserExtensionUtils.h"
 #include "TedsAssetDataColumns.h"
-#include "Elements/Columns/TypedElementAlertColumns.h"
 #include "Elements/Columns/TypedElementMiscColumns.h"
 #include "Elements/Framework/TypedElementIndexHasher.h"
 #include "String/LexFromString.h"
+#include "TedsAlerts.h"
 
 #define LOCTEXT_NAMESPACE "UTedsAssetDataFactory"
 
@@ -86,6 +86,8 @@ namespace UE::TedsAssetDataFactory::Private
 
 void UTedsAssetDataFactory::RegisterQueries(IEditorDataStorageProvider& DataStorage)
 {
+	using namespace UE::Editor::DataStorage;
+	using namespace UE::Editor::DataStorage::Columns;
 	using namespace UE::Editor::DataStorage::Queries;
 
 	DataStorage.RegisterQuery(
@@ -160,6 +162,7 @@ void UTedsAssetDataFactory::RegisterQueries(IEditorDataStorageProvider& DataStor
 			[](IQueryContext& Context, const RowHandle* Rows)
 			{
 				const int32 NumRows = Context.GetRowCount();
+				static const FName AlertName = "NonSquareTextureAlert";
 
 				TConstArrayView<FItemStringAttributeColumn_Experimental> ColumnView = MakeConstArrayView(
 								Context.GetColumn<FItemStringAttributeColumn_Experimental>("Dimensions"),
@@ -169,14 +172,13 @@ void UTedsAssetDataFactory::RegisterQueries(IEditorDataStorageProvider& DataStor
 				{
 					if (UE::TedsAssetDataFactory::Private::IsTextureNonSquare(ColumnView[Index].Value))
 					{
-						FTypedElementAlertColumn Alert;
-						Alert.AlertType = FTypedElementAlertColumnType::Error;
-						Alert.Message = LOCTEXT("NonSquareTextureAlert", "Texture has a non-square aspect ratio");
-						Context.AddColumn(Rows[Index], MoveTemp(Alert));
+						Alerts::AddAlert(Context, Rows[Index], AlertName,
+							LOCTEXT("NonSquareTextureAlert", "Texture has a non-square aspect ratio."),
+							FAlertColumnType::Error);
 					}
 					else
 					{
-						Context.RemoveColumns<FTypedElementAlertColumn>(Rows[Index]);
+						Alerts::RemoveAlert(Context, Rows[Index], AlertName);
 					}
 				}
 			}
diff --git a/Engine/Plugins/Experimental/EditorDataStorageFeatures/Source/TedsAssetData/TedsAssetData.build.cs b/Engine/Plugins/Experimental/EditorDataStorageFeatures/Source/TedsAssetData/TedsAssetData.build.cs
--- a/Engine/Plugins/Experimental/EditorDataStorageFeatures/Source/TedsAssetData/TedsAssetData.build.cs
+++ b/Engine/Plugins/Experimental/EditorDataStorageFeatures/Source/TedsAssetData/TedsAssetData.build.cs
@@ -27,6 +27,7 @@ public class TedsAssetData : ModuleRules
 					"Engine",
 					"Slate",
 					"SlateCore",
+					"TedsAlerts",
 					"TypedElementFramework",
 					"UnrealEd",
 					"UnrealEd",
diff --git a/Engine/Plugins/Experimental/EditorDataStorageFeatures/Source/TedsContentBrowser/Private/TedsContentBrowserModule.cpp b/Engine/Plugins/Experimental/EditorDataStorageFeatures/Source/TedsContentBrowser/Private/TedsContentBrowserModule.cpp
--- a/Engine/Plugins/Experimental/EditorDataStorageFeatures/Source/TedsContentBrowser/Private/TedsContentBrowserModule.cpp
+++ b/Engine/Plugins/Experimental/EditorDataStorageFeatures/Source/TedsContentBrowser/Private/TedsContentBrowserModule.cpp
@@ -12,8 +12,8 @@
 #include "Experimental/ContentBrowserViewExtender.h"
 #include "Modules/ModuleManager.h"
 #include "QueryStack/FQueryStackNode_RowView.h"
+#include "TedsAlertColumns.h"
 #include "TedsAssetDataColumns.h"
-#include "Elements/Columns/TypedElementAlertColumns.h"
 #include "Elements/Columns/TypedElementMiscColumns.h"
 #include "Widgets/STedsTableViewer.h"
 #include "Widgets/Views/SListView.h"
@@ -149,7 +149,7 @@ namespace UE::Editor::ContentBrowser
 					.QueryStack(RowQueryStack)
 					.CellWidgetPurposes({TEXT("General.RowLabel"), TEXT("General.Cell")})
 					// Default list of columns to display
-					.Columns({ FNameColumn::StaticStruct(), FTypedElementAlertColumn::StaticStruct(),
+					.Columns({ FNameColumn::StaticStruct(), FTedsAlertColumn::StaticStruct(),
 						FAssetClassColumn::StaticStruct(), FAssetTag::StaticStruct(), FAssetPathColumn_Experimental::StaticStruct(),
 						FDiskSizeColumn::StaticStruct(), FVirtualPathColumn_Experimental::StaticStruct(), DynamicSkeletalMeshSkeletonColumn })
 					.ListSelectionMode(ESelectionMode::Multi)
diff --git a/Engine/Plugins/Experimental/EditorDataStorageFeatures/Source/TedsContentBrowser/TedsContentBrowser.Build.cs b/Engine/Plugins/Experimental/EditorDataStorageFeatures/Source/TedsContentBrowser/TedsContentBrowser.Build.cs
--- a/Engine/Plugins/Experimental/EditorDataStorageFeatures/Source/TedsContentBrowser/TedsContentBrowser.Build.cs
+++ b/Engine/Plugins/Experimental/EditorDataStorageFeatures/Source/TedsContentBrowser/TedsContentBrowser.Build.cs
@@ -30,6 +30,7 @@ public class TedsContentBrowser : ModuleRules
 		            "ContentBrowser",
 		            "ContentBrowserData",
 		            "InputCore",
+					"TedsAlerts",
 		            "TedsAssetData",
 		            "ToolWidgets"
             	});
diff --git a/Engine/Plugins/Experimental/EditorDataStorageFeatures/Source/TedsOutliner/Private/TedsOutlinerModule.cpp b/Engine/Plugins/Experimental/EditorDataStorageFeatures/Source/TedsOutliner/Private/TedsOutlinerModule.cpp
--- a/Engine/Plugins/Experimental/EditorDataStorageFeatures/Source/TedsOutliner/Private/TedsOutlinerModule.cpp
+++ b/Engine/Plugins/Experimental/EditorDataStorageFeatures/Source/TedsOutliner/Private/TedsOutlinerModule.cpp
@@ -7,7 +7,6 @@
 #include "TedsOutlinerMode.h"
 #include "WorkspaceMenuStructureModule.h"
 #include "WorkspaceMenuStructure.h"
-#include "Elements/Columns/TypedElementAlertColumns.h"
 #include "Elements/Columns/TypedElementCompatibilityColumns.h"
 #include "Elements/Columns/TypedElementLabelColumns.h"
 #include "Elements/Columns/TypedElementMiscColumns.h"
@@ -19,6 +18,7 @@
 #include "Modules/ModuleManager.h"
 #include "Compatibility/SceneOutlinerTedsBridge.h"
 #include "Compatibility/SceneOutlinerRowHandleColumn.h"
+#include "TedsAlertColumns.h"
 #include "Widgets/Docking/SDockTab.h"
 
 #define LOCTEXT_NAMESPACE "TedsOutlinerModule"
@@ -118,19 +118,20 @@ void FTedsOutlinerModule::ShutdownModule()
 DataStorage::QueryHandle FTedsOutlinerModule::GetLevelEditorTedsOutlinerColumnQuery()
 {
 	using namespace UE::Editor::DataStorage;
+	using namespace UE::Editor::DataStorage::Columns;
 	IEditorDataStorageProvider* Storage = GetMutableDataStorageFeature<IEditorDataStorageProvider>(StorageFeatureName);
 		
 	using namespace DataStorage::Queries;
 
 	static DataStorage::QueryHandle ColumnQuery = Storage->RegisterQuery(
 		Select()
-			.ReadOnly<FTypedElementClassTypeInfoColumn, FTypedElementAlertColumn, FTypedElementChildAlertColumn, FVisibleInEditorColumn>()
+			.ReadOnly<FTypedElementClassTypeInfoColumn, FAlertColumn, FChildAlertColumn, FVisibleInEditorColumn>()
 		.Compile());
 
 	// Query to also include revision control info
 	static DataStorage::QueryHandle RevisionControlQuery = Storage->RegisterQuery(
 		Select()
-			.ReadOnly<FTypedElementClassTypeInfoColumn, FTypedElementPackageReference, FTypedElementAlertColumn, FVisibleInEditorColumn>()
+			.ReadOnly<FTypedElementClassTypeInfoColumn, FTypedElementPackageReference, FAlertColumn, FVisibleInEditorColumn>()
 		.Compile());
 
 	return Private::bUseNewRevisionControlWidgets ? RevisionControlQuery : ColumnQuery;
diff --git a/Engine/Plugins/Experimental/EditorDataStorageFeatures/Source/TedsOutliner/Private/Widgets/AlertWidget.cpp b/Engine/Plugins/Experimental/EditorDataStorageFeatures/Source/TedsOutliner/Private/Widgets/AlertWidget.cpp
deleted file mode 100644
diff --git a/Engine/Plugins/Experimental/EditorDataStorageFeatures/Source/TedsOutliner/Public/Widgets/AlertWidget.h b/Engine/Plugins/Experimental/EditorDataStorageFeatures/Source/TedsOutliner/Public/Widgets/AlertWidget.h
deleted file mode 100644
diff --git a/Engine/Plugins/Experimental/EditorDataStorageFeatures/Source/TedsOutliner/TedsOutliner.Build.cs b/Engine/Plugins/Experimental/EditorDataStorageFeatures/Source/TedsOutliner/TedsOutliner.Build.cs
--- a/Engine/Plugins/Experimental/EditorDataStorageFeatures/Source/TedsOutliner/TedsOutliner.Build.cs
+++ b/Engine/Plugins/Experimental/EditorDataStorageFeatures/Source/TedsOutliner/TedsOutliner.Build.cs
@@ -32,6 +32,7 @@ public class TedsOutliner : ModuleRules
 					"UnrealEd", // FEditorUndoClient used by SSceneOutliner
 					"ToolMenus",
 					"ApplicationCore",
+					"TedsAlerts",
 					"TedsTableViewer",
 					"ToolWidgets",
 					"InputCore"
diff --git a/Engine/Plugins/Experimental/InstanceDataObjectFixupTool/InstanceDataObjectFixupTool.uplugin b/Engine/Plugins/Experimental/InstanceDataObjectFixupTool/InstanceDataObjectFixupTool.uplugin
--- a/Engine/Plugins/Experimental/InstanceDataObjectFixupTool/InstanceDataObjectFixupTool.uplugin
+++ b/Engine/Plugins/Experimental/InstanceDataObjectFixupTool/InstanceDataObjectFixupTool.uplugin
@@ -15,8 +15,11 @@
 	"IsBetaVersion" : false,
 	"IsExperimentalVersion" : true,
 	"Installed" : false,
-	"Plugins": 
-	[
+	"Plugins": [
+		{
+			"Name": "EditorDataStorageFeatures",
+			"Enabled": true
+		}
 	],
 	"Modules" :
 	[
diff --git a/Engine/Plugins/Experimental/InstanceDataObjectFixupTool/Source/InstanceDataObjectFixupTool.Build.cs b/Engine/Plugins/Experimental/InstanceDataObjectFixupTool/Source/InstanceDataObjectFixupTool.Build.cs
--- a/Engine/Plugins/Experimental/InstanceDataObjectFixupTool/Source/InstanceDataObjectFixupTool.Build.cs
+++ b/Engine/Plugins/Experimental/InstanceDataObjectFixupTool/Source/InstanceDataObjectFixupTool.Build.cs
@@ -21,6 +21,7 @@ namespace UnrealBuildTool.Rules
 					"PropertyEditor", 
 					"WorkspaceMenuStructure", 
 					"UnrealEd",
+					"TedsAlerts",
 					"TypedElementFramework"
 				}
 			);
diff --git a/Engine/Plugins/Experimental/InstanceDataObjectFixupTool/Source/Private/InstanceDataObjectFixupPanel.cpp b/Engine/Plugins/Experimental/InstanceDataObjectFixupTool/Source/Private/InstanceDataObjectFixupPanel.cpp
--- a/Engine/Plugins/Experimental/InstanceDataObjectFixupTool/Source/Private/InstanceDataObjectFixupPanel.cpp
+++ b/Engine/Plugins/Experimental/InstanceDataObjectFixupTool/Source/Private/InstanceDataObjectFixupPanel.cpp
@@ -10,10 +10,10 @@
 #include "Modules/ModuleManager.h"
 #include "Editor.h"
 #include "Editor/PropertyEditor/Private/PropertyNode.h"
+#include "TedsAlerts.h"
 #include "UObject/PropertyBagRepository.h"
 #include "UObject/PropertyStateTracking.h"
 
-#include "Elements/Columns/TypedElementAlertColumns.h"
 #include "Elements/Common/EditorDataStorageFeatures.h"
 #include "Elements/Interfaces/TypedElementDataStorageInterface.h"
 #include "Elements/Interfaces/TypedElementDataStorageCompatibilityInterface.h"
@@ -260,13 +260,15 @@ FInstanceDataObjectFixupPanel::~FInstanceDataObjectFixupPanel()
 		{
 			if (!ObjectHasLoosePropertiesThatNeedFixup(Instance))
 			{
+				static FName AlertName = FName("EntityLoosePropertiesErrorAlert");
+
 				UE::FPropertyBagRepository& Repository = UE::FPropertyBagRepository::Get();
 				Repository.MarkAsFixedUp(Repository.FindInstanceForDataObject(Instance));
 
 				// If a UObject isn't registered with TEDS, there's a chance its parent is registered and is the one
 				// with the alert column on it, so search upward until the nearest registered parent is found.
 				RowHandle Row = DataStorageCompatibility->FindRowWithCompatibleObject(InstancesOwner ? InstancesOwner.Get() : Instance);
-				DataStorage->RemoveColumns<FTypedElementAlertColumn>(Row);
+				Alerts::RemoveAlert(*DataStorage, Row, AlertName);
 			}
 		}
 	}
diff --git a/Engine/Plugins/Experimental/InstanceDataObjectFixupTool/Source/Private/InstanceDataObjectFixupToolTedsQueries.cpp b/Engine/Plugins/Experimental/InstanceDataObjectFixupTool/Source/Private/InstanceDataObjectFixupToolTedsQueries.cpp
--- a/Engine/Plugins/Experimental/InstanceDataObjectFixupTool/Source/Private/InstanceDataObjectFixupToolTedsQueries.cpp
+++ b/Engine/Plugins/Experimental/InstanceDataObjectFixupTool/Source/Private/InstanceDataObjectFixupToolTedsQueries.cpp
@@ -2,7 +2,6 @@
 
 #include "InstanceDataObjectFixupToolTedsQueries.h"
 
-#include "Elements/Columns/TypedElementAlertColumns.h"
 #include "Elements/Columns/TypedElementCompatibilityColumns.h"
 #include "Elements/Columns/TypedElementMiscColumns.h"
 #include "Elements/Common/EditorDataStorageFeatures.h"
@@ -12,7 +11,9 @@
 #include "UObject/PropertyBagRepository.h"
 #include "InstanceDataObjectFixupToolModule.h"
 #include "Framework/Notifications/NotificationManager.h"
+#include "TedsAlerts.h"
 #include "Widgets/Notifications/SNotificationList.h"
+
 #define LOCTEXT_NAMESPACE "FixupToolTedsQueries"
 
 void UInstanceDataObjectFixupToolTedsQueryFactory::RegisterQueries(IEditorDataStorageProvider& DataStorage)
@@ -23,26 +24,26 @@ void UInstanceDataObjectFixupToolTedsQueryFactory::RegisterQueries(IEditorDataSt
 		Select(
 			TEXT("Add fix-up tool to serialization placeholder alerts"),
 			FProcessor(EQueryTickPhase::PrePhysics, DataStorage.GetQueryTickGroupName(EQueryTickGroups::SyncExternalToDataStorage)),
-			[](IQueryContext& Context, RowHandle Row, const FTypedElementUObjectColumn& Object)
+			[](IQueryContext& Context, RowHandle Row)
 			{
-				Context.AddColumn(Row, FTypedElementAlertActionColumn{ .Action = ShowFixUpToolForPlaceholders });
+				static FName AlertName("PlaceholderAlert");
+				Alerts::UpdateAlertAction(Context, Row, AlertName, ShowFixUpToolForPlaceholders);
 			})
 		.Where()
-			.All<FTypedElementSyncFromWorldTag, FTypedElementAlertColumn, FTypedElementPropertyBagPlaceholderTag>()
-			.None<FTypedElementAlertActionColumn>()
+			.All<FTypedElementSyncFromWorldTag, FTypedElementUObjectColumn, FTypedElementPropertyBagPlaceholderTag>()
 		.Compile());
 
 	DataStorage.RegisterQuery(
 		Select(
 			TEXT("Add fix-up tool to serialization loose property alerts"),
 			FProcessor(EQueryTickPhase::PrePhysics, DataStorage.GetQueryTickGroupName(EQueryTickGroups::SyncExternalToDataStorage)),
-			[](IQueryContext& Context, RowHandle Row, const FTypedElementUObjectColumn& Object)
+			[](IQueryContext& Context, RowHandle Row)
 			{
-				Context.AddColumn(Row, FTypedElementAlertActionColumn{ .Action = ShowFixUpToolForLooseProperties });
+				static FName AlertName = FName("EntityLoosePropertiesErrorAlert");
+				Alerts::UpdateAlertAction(Context, Row, AlertName, ShowFixUpToolForLooseProperties);
 			})
 		.Where()
-			.All<FTypedElementSyncFromWorldTag, FTypedElementAlertColumn, FTypedElementLoosePropertyTag>()
-			.None<FTypedElementAlertActionColumn>()
+			.All<FTypedElementSyncFromWorldTag, FTypedElementUObjectColumn, FTypedElementLoosePropertyTag>()
 		.Compile());
 }
 
diff --git a/Engine/Plugins/Runtime/HairStrands/Source/HairStrandsDataflow/Private/GroomDataflowVisualization.cpp b/Engine/Plugins/Runtime/HairStrands/Source/HairStrandsDataflow/Private/GroomDataflowVisualization.cpp
--- a/Engine/Plugins/Runtime/HairStrands/Source/HairStrandsDataflow/Private/GroomDataflowVisualization.cpp
+++ b/Engine/Plugins/Runtime/HairStrands/Source/HairStrandsDataflow/Private/GroomDataflowVisualization.cpp
@@ -1,4 +1,4 @@
-﻿// Copyright Epic Games, Inc. All Rights Reserved.
+// Copyright Epic Games, Inc. All Rights Reserved.
 
 #include "GroomDataflowVisualization.h"
 #include "GroomComponent.h"
@@ -6,7 +6,6 @@
 #include "Dataflow/DataflowEditorToolkit.h"
 #include "Dataflow/DataflowSimulationViewportClient.h"
 #include "Dataflow/DataflowSimulationScene.h"
-#include "Elements/Columns/TypedElementAlertColumns.h"
 
 #define LOCTEXT_NAMESPACE "ChaosClothAssetDataflowSimulationVisualization"
 
diff --git a/Engine/Source/Runtime/TypedElementFramework/Public/Elements/Columns/TypedElementAlertColumns.h b/Engine/Source/Runtime/TypedElementFramework/Public/Elements/Columns/TypedElementAlertColumns.h
deleted file mode 100644
