[TEDSOutliner]
* Implementing support for multiselect visibility, in parity with the old outliner

diff --git a/Engine/Plugins/Experimental/EditorDataStorageFeatures/Source/TedsOutliner/Private/Widgets/OutlinerVisibilityWidget.cpp b/Engine/Plugins/Experimental/EditorDataStorageFeatures/Source/TedsOutliner/Private/Widgets/OutlinerVisibilityWidget.cpp
--- a/Engine/Plugins/Experimental/EditorDataStorageFeatures/Source/TedsOutliner/Private/Widgets/OutlinerVisibilityWidget.cpp
+++ b/Engine/Plugins/Experimental/EditorDataStorageFeatures/Source/TedsOutliner/Private/Widgets/OutlinerVisibilityWidget.cpp
@@ -119,7 +119,6 @@ void STEDSVisibilityWidget::Construct(const FArguments& InArgs, const RowHandle&
 {
 	TargetRow = InTargetRow;
 	WidgetRow = InWidgetRow;
-	TreeItem = GetTreeItem();
 
 	SImage::Construct(
 		SImage::FArguments()
@@ -176,7 +175,20 @@ FReply STEDSVisibilityWidget::HandleClick()
 
 	const bool bVisible = !IsVisible();
 
-	SetIsVisible(TargetRow, bVisible);
+	TArray<RowHandle> SelectedRows;
+	GetSelectedRows(SelectedRows);
+
+	if (SelectedRows.Contains(TargetRow))
+	{
+		for (RowHandle Row : SelectedRows)
+		{
+			SetIsVisible(Row, bVisible);
+		}
+	}
+	else
+	{
+		SetIsVisible(TargetRow, bVisible);
+	}
 
 	return FReply::Handled().DetectDrag(SharedThis(this), EKeys::LeftMouseButton);
 }
@@ -275,9 +287,10 @@ bool STEDSVisibilityWidget::IsSelected() const
 /** Set the item this widget is responsible for to be hidden or shown */
 void STEDSVisibilityWidget::SetIsVisible(RowHandle InRow, const bool bVisible)
 {
-	if (IEditorDataStorageProvider* DataStorage = GetDataStorage(); TreeItem.IsValid())
+	TWeakPtr<ISceneOutlinerTreeItem> RowTreeItem = GetTreeItem(InRow);
+	if (IEditorDataStorageProvider* DataStorage = GetDataStorage(); RowTreeItem.IsValid())
 	{
-		SetVisibility_Recursive(*DataStorage, TreeItem.Pin(), bVisible);
+		SetVisibility_Recursive(*DataStorage, RowTreeItem.Pin(), bVisible);
 	}
 }
 
@@ -314,7 +327,7 @@ void STEDSVisibilityWidget::SetVisibility_Recursive(IEditorDataStorageProvider&
 	}
 }
 
-FSceneOutlinerTreeItemPtr STEDSVisibilityWidget::GetTreeItem() const
+FSceneOutlinerTreeItemPtr STEDSVisibilityWidget::GetTreeItem(RowHandle InRow) const
 {
 	IEditorDataStorageProvider* DataStorage = GetDataStorage();
 	if (!DataStorage)
@@ -337,7 +350,7 @@ FSceneOutlinerTreeItemPtr STEDSVisibilityWidget::GetTreeItem() const
 	FSceneOutlinerTedsQueryBinder& Binder = FSceneOutlinerTedsQueryBinder::GetInstance();
 	if (FTreeItemIDDealiaser Dealiaser = Binder.GetTreeItemIDDealiaser(Outliner); Dealiaser.IsBound())
 	{
-		FSceneOutlinerTreeItemID ID = Dealiaser.Execute(TargetRow);
+		FSceneOutlinerTreeItemID ID = Dealiaser.Execute(InRow);
 		return Outliner->GetTreeItem(ID);
 	}
 
@@ -362,4 +375,28 @@ IEditorDataStorageCompatibilityProvider* STEDSVisibilityWidget::GetDataStorageCo
 	return GetMutableDataStorageFeature<IEditorDataStorageCompatibilityProvider>(CompatibilityFeatureName);
 }
 
+void STEDSVisibilityWidget::GetSelectedRows(TArray<RowHandle>& OutSelectedRows)
+{
+	using namespace UE::Editor::DataStorage;
+	using namespace UE::Editor::DataStorage::Queries;
+	if (IEditorDataStorageProvider* DataStorage = GetDataStorage())
+	{
+		static QueryHandle AllSelectedItemsQuery = DataStorage->RegisterQuery(Select().Where().All<FTypedElementSelectionColumn>().Compile());
+		DataStorage->RunQuery(AllSelectedItemsQuery,
+			CreateDirectQueryCallbackBinding([&OutSelectedRows, &DataStorage](const IEditorDataStorageProvider::IDirectQueryContext& Context)
+			{
+				// Only add selections from the level editor - for now.
+				// UE-231184: Support Visibility Column in multiple editor contexts
+				for (RowHandle Row : Context.GetRowHandles())
+				{
+					FTypedElementSelectionColumn* Selection = DataStorage->GetColumn<FTypedElementSelectionColumn>(Row);
+					if (Selection && Selection->SelectionSet.IsNone()) // SelectionSet.IsNone() == Level Editor
+					{
+						OutSelectedRows.Add(Row);
+					}
+				}
+			}));
+	}
+}
+
 #undef LOCTEXT_NAMESPACE
\ No newline at end of file
diff --git a/Engine/Plugins/Experimental/EditorDataStorageFeatures/Source/TedsOutliner/Public/Widgets/OutlinerVisibilityWidget.h b/Engine/Plugins/Experimental/EditorDataStorageFeatures/Source/TedsOutliner/Public/Widgets/OutlinerVisibilityWidget.h
--- a/Engine/Plugins/Experimental/EditorDataStorageFeatures/Source/TedsOutliner/Public/Widgets/OutlinerVisibilityWidget.h
+++ b/Engine/Plugins/Experimental/EditorDataStorageFeatures/Source/TedsOutliner/Public/Widgets/OutlinerVisibilityWidget.h
@@ -135,9 +135,6 @@ protected:
 	RowHandle TargetRow;
 	RowHandle WidgetRow;
 
-	// Pointer to the outliner tree item associated with this widget. 
-	TWeakPtr<ISceneOutlinerTreeItem> TreeItem;
-
 	/** Scoped undo transaction */
 	TUniquePtr<FScopedTransaction> UndoTransaction;
 
@@ -156,9 +153,11 @@ private:
 	static void SetVisibility_Recursive(IEditorDataStorageProvider& DataStorage, FSceneOutlinerTreeItemPtr TreeItem, bool bVisible);
 
 	// Returns the tree item associated with this widget.
-	FSceneOutlinerTreeItemPtr GetTreeItem() const;
+	FSceneOutlinerTreeItemPtr GetTreeItem(RowHandle InRow) const;
 
 	static IEditorDataStorageProvider* GetDataStorage();
 	static IEditorDataStorageProvider* GetDataStorageUI();
 	static IEditorDataStorageCompatibilityProvider* GetDataStorageCompatibility();
+
+	static void GetSelectedRows(TArray<RowHandle>& OutSelectedRows);
 };
\ No newline at end of file
