[TEDSOutliner][ActorEditor]
* Broadcast a PostEditChangeProperty event for bHiddenEdTemporary so that listeners can react, specifically UEditorDataStorageComatibility::OnPostEditChangeProperty
- This fixes de-sync between the world object and TEDS when editor visibility is changed directly on objects.
* Added call to Tick for FallbackColumns in FOutlinerColumn
- Needed due to SceneOutlinerGutter being a FallbackColumn that must reset visibility cache for Folders that have not been fully "TED-ified" yet.

diff --git a/Engine/Plugins/Experimental/EditorDataStorageFeatures/Source/TedsOutliner/Private/Compatibility/SceneOutlinerTedsBridge.cpp b/Engine/Plugins/Experimental/EditorDataStorageFeatures/Source/TedsOutliner/Private/Compatibility/SceneOutlinerTedsBridge.cpp
--- a/Engine/Plugins/Experimental/EditorDataStorageFeatures/Source/TedsOutliner/Private/Compatibility/SceneOutlinerTedsBridge.cpp
+++ b/Engine/Plugins/Experimental/EditorDataStorageFeatures/Source/TedsOutliner/Private/Compatibility/SceneOutlinerTedsBridge.cpp
@@ -188,6 +188,10 @@ public:
 	virtual void Tick(double InCurrentTime, float InDeltaTime)
 	{
 		TableViewerColumnImpl->Tick();
+		if (FallbackColumn)
+		{
+			FallbackColumn->Tick(InCurrentTime, InDeltaTime);
+		}
 	}
 
 	bool IsRowVisible(const UE::Editor::DataStorage::RowHandle InRowHandle) const
diff --git a/Engine/Source/Runtime/Engine/Private/ActorEditor.cpp b/Engine/Source/Runtime/Engine/Private/ActorEditor.cpp
--- a/Engine/Source/Runtime/Engine/Private/ActorEditor.cpp
+++ b/Engine/Source/Runtime/Engine/Private/ActorEditor.cpp
@@ -983,6 +983,9 @@ void AActor::SetIsTemporarilyHiddenInEditor( bool bIsHidden )
 	{
 		bHiddenEdTemporary = bIsHidden;
 		MarkComponentsRenderStateDirty();
+		
+		FPropertyChangedEvent HiddenEdTemporaryChanged(AActor::StaticClass()->FindPropertyByName(GET_MEMBER_NAME_CHECKED(AActor, bHiddenEdTemporary)), EPropertyChangeType::ValueSet);
+		PostEditChangeProperty(HiddenEdTemporaryChanged);
 	}
 }
 
