[TEDS] Changes required to support the Chaos visual Debugger requirements

- Added a way to pre-compute all children of any given row. This allow us to avoid having to walk the hierarchy upwards, which in the current implementation is exponential and was causing the outliner takes minutes to populate on a recording of City Sample.
- Fixed an issue where actors with no label set, were showing an empty label instead of a default label.
- Added support for ustructs to the teds visibility widget (we just needed to allow them in the query).
- Added CVD standalone to the supported programs list of the Editor Data Storage Features plugin.

diff --git a/Engine/Plugins/Experimental/EditorDataStorage/Source/TedsCore/Private/Compatibility/TypedElementActorLabelQueries.cpp b/Engine/Plugins/Experimental/EditorDataStorage/Source/TedsCore/Private/Compatibility/TypedElementActorLabelQueries.cpp
--- a/Engine/Plugins/Experimental/EditorDataStorage/Source/TedsCore/Private/Compatibility/TypedElementActorLabelQueries.cpp
+++ b/Engine/Plugins/Experimental/EditorDataStorage/Source/TedsCore/Private/Compatibility/TypedElementActorLabelQueries.cpp
@@ -32,7 +32,7 @@ void UActorLabelDataStorageFactory::RegisterActorLabelToColumnQuery(IEditorDataS
 			{
 				if (const AActor* ActorInstance = Cast<AActor>(Actor.Object); ActorInstance != nullptr)
 				{
-					const FString& ActorLabel = ActorInstance->GetActorLabel(false);
+					const FString& ActorLabel = ActorInstance->GetActorLabel();
 					uint64 ActorLabelHash = CityHash64(reinterpret_cast<const char*>(*ActorLabel), ActorLabel.Len() * sizeof(**ActorLabel));
 					if (LabelHash.LabelHash != ActorLabelHash)
 					{
diff --git a/Engine/Plugins/Experimental/EditorDataStorageFeatures/EditorDataStorageFeatures.uplugin b/Engine/Plugins/Experimental/EditorDataStorageFeatures/EditorDataStorageFeatures.uplugin
--- a/Engine/Plugins/Experimental/EditorDataStorageFeatures/EditorDataStorageFeatures.uplugin
+++ b/Engine/Plugins/Experimental/EditorDataStorageFeatures/EditorDataStorageFeatures.uplugin
@@ -57,5 +57,6 @@
 		}
 	],
 	"Plugins": [
-	]
+	],
+	"SupportedPrograms": ["ChaosVisualDebugger"]
 }
\ No newline at end of file
diff --git a/Engine/Plugins/Experimental/EditorDataStorageFeatures/Source/TedsOutliner/Private/TedsOutlinerImpl.cpp b/Engine/Plugins/Experimental/EditorDataStorageFeatures/Source/TedsOutliner/Private/TedsOutlinerImpl.cpp
--- a/Engine/Plugins/Experimental/EditorDataStorageFeatures/Source/TedsOutliner/Private/TedsOutlinerImpl.cpp
+++ b/Engine/Plugins/Experimental/EditorDataStorageFeatures/Source/TedsOutliner/Private/TedsOutlinerImpl.cpp
@@ -407,48 +407,53 @@ void FTedsOutlinerImpl::CreateChildren(const FSceneOutlinerTreeItemPtr& Item, TA
 		return;
 	}
 
+	TArray<RowHandle> ChildItems;
 
-	TArray<RowHandle> MatchedRowsWithParentColumn;
-	
-	// Collect all entities that are owned by our entity
-	DirectQueryCallback ChildRowCollector = CreateDirectQueryCallbackBinding(
-	[&MatchedRowsWithParentColumn] (const IDirectQueryContext& Context, const RowHandle*)
+	if (HierarchyData->GetChildren.IsBound())
 	{
-		MatchedRowsWithParentColumn.Append(Context.GetRowHandles());
-	});
+		void* ParentColumnData = Storage->GetColumnData(ItemRowHandle, HierarchyData->HierarchyColumn);
+		ChildItems = HierarchyData->GetChildren.Execute(ParentColumnData);
+	}
+	else
+	{
+		TArray<RowHandle> MatchedRowsWithParentColumn;
 
-	Storage->RunQuery(ChildRowHandleQuery, ChildRowCollector);
+		// Collect all entities that are owned by our entity
+		DirectQueryCallback ChildRowCollector = CreateDirectQueryCallbackBinding(
+		[&MatchedRowsWithParentColumn] (const IDirectQueryContext& Context, const RowHandle*)
+		{
+			MatchedRowsWithParentColumn.Append(Context.GetRowHandles());
+		});
 
-	TArray<RowHandle> ChildItems;
+		Storage->RunQuery(ChildRowHandleQuery, ChildRowCollector);
 
-	// Recursively get the children for each entity
-	TFunction<void(RowHandle)> GetChildrenRecursive = [&ChildItems, &MatchedRowsWithParentColumn, DataStorage = Storage, &GetChildrenRecursive, InHierarchyData = HierarchyData]
-	(RowHandle EntityRowHandle) -> void
-	{
-		for(RowHandle ChildEntityRowHandle : MatchedRowsWithParentColumn)
+		// Recursively get the children for each entity
+		TFunction<void(RowHandle)> GetChildrenRecursive = [&ChildItems, &MatchedRowsWithParentColumn, DataStorage = Storage, &GetChildrenRecursive, InHierarchyData = HierarchyData]
+		(RowHandle EntityRowHandle) -> void
 		{
-			void* ParentColumnData = DataStorage->GetColumnData(ChildEntityRowHandle, InHierarchyData.GetValue().HierarchyColumn);
-
-			if (ensureMsgf(ParentColumnData, TEXT("We should always the a parent column since we only grabbed rows with those ")))
+			for(RowHandle ChildEntityRowHandle : MatchedRowsWithParentColumn)
 			{
-				// Get the parent row handle
-				const RowHandle ParentRowHandle = InHierarchyData.GetValue().GetParent.Execute(ParentColumnData);
-				
-				// Check if this entity is owned by the entity we are looking children for
-				if (ParentRowHandle == EntityRowHandle)
+				void* ParentColumnData = DataStorage->GetColumnData(ChildEntityRowHandle, InHierarchyData.GetValue().HierarchyColumn);
+
+				if (ensureMsgf(ParentColumnData, TEXT("We should always the a parent column since we only grabbed rows with those ")))
 				{
-					ChildItems.Add(ChildEntityRowHandle);
+					// Get the parent row handle
+					const RowHandle ParentRowHandle = InHierarchyData.GetValue().GetParent.Execute(ParentColumnData);
+				
+					// Check if this entity is owned by the entity we are looking children for
+					if (ParentRowHandle == EntityRowHandle)
+					{
+						ChildItems.Add(ChildEntityRowHandle);
 
-					// Recursively look for children of this item
-					GetChildrenRecursive(ChildEntityRowHandle);
+						// Recursively look for children of this item
+						GetChildrenRecursive(ChildEntityRowHandle);
+					}
 				}
-
 			}
-			
-		}
-	};
-		
-	GetChildrenRecursive(ItemRowHandle);
+		};
+
+		GetChildrenRecursive(ItemRowHandle);
+	}
 
 	// Actually create the items for the child entities 
 	for (RowHandle ChildItemRowHandle : ChildItems)
diff --git a/Engine/Plugins/Experimental/EditorDataStorageFeatures/Source/TedsOutliner/Private/Widgets/OutlinerLabelWidget.cpp b/Engine/Plugins/Experimental/EditorDataStorageFeatures/Source/TedsOutliner/Private/Widgets/OutlinerLabelWidget.cpp
--- a/Engine/Plugins/Experimental/EditorDataStorageFeatures/Source/TedsOutliner/Private/Widgets/OutlinerLabelWidget.cpp
+++ b/Engine/Plugins/Experimental/EditorDataStorageFeatures/Source/TedsOutliner/Private/Widgets/OutlinerLabelWidget.cpp
@@ -110,7 +110,7 @@ void UOutlinerLabelWidgetFactory::RegisterWidgetConstructors(IEditorDataStorageP
 	using namespace UE::Editor::DataStorage::Queries;
 	DataStorageUi.RegisterWidgetFactory<FOutlinerLabelWidgetConstructor>(
 		TEXT("SceneOutliner.RowLabel"),
-		TColumn<FTypedElementLabelColumn>() && TColumn<FTypedElementClassTypeInfoColumn>());
+		TColumn<FTypedElementLabelColumn>() && (TColumn<FTypedElementClassTypeInfoColumn>() || TColumn<FTypedElementScriptStructTypeInfoColumn>()));
 }
 
 FOutlinerLabelWidgetConstructor::FOutlinerLabelWidgetConstructor()
diff --git a/Engine/Plugins/Experimental/EditorDataStorageFeatures/Source/TedsOutliner/Private/Widgets/OutlinerVisibilityWidget.cpp b/Engine/Plugins/Experimental/EditorDataStorageFeatures/Source/TedsOutliner/Private/Widgets/OutlinerVisibilityWidget.cpp
--- a/Engine/Plugins/Experimental/EditorDataStorageFeatures/Source/TedsOutliner/Private/Widgets/OutlinerVisibilityWidget.cpp
+++ b/Engine/Plugins/Experimental/EditorDataStorageFeatures/Source/TedsOutliner/Private/Widgets/OutlinerVisibilityWidget.cpp
@@ -17,6 +17,7 @@
 #include "Widgets/SBoxPanel.h"
 #include "ISceneOutliner.h"
 #include "ActorTreeItem.h"
+#include "TedsOutlinerItem.h"
 
 #define LOCTEXT_NAMESPACE "OutlinerVisibilityWidget"
 
@@ -308,13 +309,19 @@ void STEDSVisibilityWidget::CommitVisibility(IEditorDataStorageProvider& DataSto
 void STEDSVisibilityWidget::SetVisibility_Recursive(IEditorDataStorageProvider& DataStorage, FSceneOutlinerTreeItemPtr TreeItem, bool bVisible)
 {
 	using namespace UE::Editor::DataStorage;
-	IEditorDataStorageCompatibilityProvider* DataStorageComat = GetDataStorageCompatibility();
-	FActorTreeItem* ActorTreeItem = TreeItem->CastTo<FActorTreeItem>();
-	
-	if (ActorTreeItem && DataStorageComat)
+	using namespace UE::Editor::Outliner;
+
+	if (IEditorDataStorageCompatibilityProvider* DataStorageCompat = GetDataStorageCompatibility())
 	{
-		RowHandle ActorRow = DataStorageComat->FindRowWithCompatibleObject(ActorTreeItem->Actor);
-		CommitVisibility(DataStorage, ActorRow, bVisible);
+		if (FActorTreeItem* ActorTreeItem = TreeItem->CastTo<FActorTreeItem>())
+		{
+			RowHandle ActorRow = DataStorageCompat->FindRowWithCompatibleObject(ActorTreeItem->Actor);
+			CommitVisibility(DataStorage, ActorRow, bVisible);
+		}
+		else if (FTedsOutlinerTreeItem* TedsItem = TreeItem->CastTo<FTedsOutlinerTreeItem>())
+		{
+			CommitVisibility(DataStorage, TedsItem->GetRowHandle(), bVisible);
+		}
 	}
 
 	for (auto& ChildTreeItemPtr : TreeItem->GetChildren())
diff --git a/Engine/Plugins/Experimental/EditorDataStorageFeatures/Source/TedsOutliner/Public/TedsOutlinerImpl.h b/Engine/Plugins/Experimental/EditorDataStorageFeatures/Source/TedsOutliner/Public/TedsOutlinerImpl.h
--- a/Engine/Plugins/Experimental/EditorDataStorageFeatures/Source/TedsOutliner/Public/TedsOutlinerImpl.h
+++ b/Engine/Plugins/Experimental/EditorDataStorageFeatures/Source/TedsOutliner/Public/TedsOutlinerImpl.h
@@ -19,13 +19,17 @@ struct FTedsOutlinerHierarchyData
 {
 	/** A delegate used to get the parent row handle for a given row */
 	DECLARE_DELEGATE_RetVal_OneParam(DataStorage::RowHandle, FGetParentRowHandle, void* /* InColumnData */);
+
+	/** A delegate used to get the children row handles for a given row */
+	DECLARE_DELEGATE_RetVal_OneParam(TArrayView<DataStorage::RowHandle>, FGetChildrenRowsHandles, void* /* InColumnData */);
 	
 	/** A delegate used to set the parent row handle for a given row */
 	DECLARE_DELEGATE_TwoParams(FSetParentRowHandle, void* /* InColumnData */, DataStorage::RowHandle /* InParentRowHandle */);
 
-	FTedsOutlinerHierarchyData(const UScriptStruct* InHierarchyColumn, const FGetParentRowHandle& InGetParent, const FSetParentRowHandle& InSetParent)
+	FTedsOutlinerHierarchyData(const UScriptStruct* InHierarchyColumn, const FGetParentRowHandle& InGetParent, const FSetParentRowHandle& InSetParent, const FGetChildrenRowsHandles& InGetChildren)
 		: HierarchyColumn(InHierarchyColumn)
 		, GetParent(InGetParent)
+		, GetChildren(InGetChildren)
 		, SetParent(InSetParent)
 	{
 	
@@ -37,6 +41,9 @@ struct FTedsOutlinerHierarchyData
 	// Function to get parent row handle
 	FGetParentRowHandle GetParent;
 
+	// Function to get the children row handle
+	FGetChildrenRowsHandles GetChildren;
+
 	// Function to set the parent row handle
 	FSetParentRowHandle SetParent;
 	
@@ -60,10 +67,9 @@ struct FTedsOutlinerHierarchyData
 				{
 					ParentColumn->Parent = InRowHandle;
 				}
-
 			});
 		
-		return FTedsOutlinerHierarchyData(FTableRowParentColumn::StaticStruct(), RowHandleGetter, RowHandleSetter);
+		return FTedsOutlinerHierarchyData(FTableRowParentColumn::StaticStruct(), RowHandleGetter, RowHandleSetter, FGetChildrenRowsHandles());
 	}
 };
 
diff --git a/Engine/Source/Editor/SceneOutliner/Public/SSceneOutliner.h b/Engine/Source/Editor/SceneOutliner/Public/SSceneOutliner.h
--- a/Engine/Source/Editor/SceneOutliner/Public/SSceneOutliner.h
+++ b/Engine/Source/Editor/SceneOutliner/Public/SSceneOutliner.h
@@ -1024,7 +1024,7 @@ private:
 	FName OutlinerIdentifier;
 
 	/** true if the hierarchy of items is pinned at the top of the outliner */
-	bool bShouldStackHierarchyHeaders;
+	bool bShouldStackHierarchyHeaders = true;
 
 	void ToggleStackHierarchyHeaders();
 
