[EditorDataStorageFeature] Teds Column View Widgets

diff --git a/Engine/Plugins/Experimental/EditorDataStorageFeatures/Source/TedsAssetData/Private/TedsAssetData.cpp b/Engine/Plugins/Experimental/EditorDataStorageFeatures/Source/TedsAssetData/Private/TedsAssetData.cpp
--- a/Engine/Plugins/Experimental/EditorDataStorageFeatures/Source/TedsAssetData/Private/TedsAssetData.cpp
+++ b/Engine/Plugins/Experimental/EditorDataStorageFeatures/Source/TedsAssetData/Private/TedsAssetData.cpp
@@ -5,6 +5,8 @@
 #include "TedsAssetDataColumns.h"
 
 #include "AssetRegistry/IAssetRegistry.h"
+#include "AssetViewTypes.h"
+#include "AssetViewUtils.h"
 #include "Async/ParallelFor.h"
 #include "Containers/Array.h"
 #include "Containers/ChunkedArray.h"
@@ -160,6 +162,21 @@ void PopulatePathDataTableRow(FPopulatePathRowArgs&& InPopulatePathRowArgs, IEdi
 		}
 	}
 
+	if (AssetViewUtils::IsDevelopersFolder(InPopulatePathRowArgs.AssetRegistryPath.ToString()))
+	{
+		Database.GetColumn<FFolderTypeColumn_Experimental>(InRowHandle)->FolderType = EFolderType::Developer;
+	}
+	else if (AssetViewUtils::IsPluginFolder(InPopulatePathRowArgs.AssetRegistryPath.ToString()))
+	{
+		Database.GetColumn<FFolderTypeColumn_Experimental>(InRowHandle)->FolderType = EFolderType::PluginRoot;
+	}
+	// TODO: Missing CPP Folders, need further change and conversion from the old CBClassDataSource (see FNativeClassHierarchy)
+	// TODO: Missing Virtual Folders, need further change and conversion from the old CBAssetDataCore (see GetItemAttribute)
+	else
+	{
+		Database.GetColumn<FFolderTypeColumn_Experimental>(InRowHandle)->FolderType = EFolderType::Normal;
+	}
+
 	Database.GetColumn<FAssetPathColumn_Experimental>(InRowHandle)->Path = InPopulatePathRowArgs.AssetRegistryPath;
 	Database.GetColumn<FNameColumn>(InRowHandle)->Name = InPopulatePathRowArgs.AssetName;
 }
@@ -190,7 +207,7 @@ FTedsAssetData::FTedsAssetData(IEditorDataStorageProvider& InDatabase)
 	PathsTable = Database.FindTable(FName(TEXT("Editor_AssetRegistryPathsTable")));
 	if (PathsTable == InvalidTableHandle)
 	{
-		PathsTable = Database.RegisterTable<FFolderTag, FAssetPathColumn_Experimental, FNameColumn, FUpdatedPathTag, FSlateColorColumn, FChildrenAssetPathColumn_Experimental, FParentAssetPathColumn_Experimental, FAssetsInPathColumn_Experimental>(FName(TEXT("Editor_AssetRegistryPathsTable")));
+		PathsTable = Database.RegisterTable<FFolderTag, FFolderTypeColumn_Experimental, FAssetPathColumn_Experimental, FNameColumn, FUpdatedPathTag, FSlateColorColumn, FChildrenAssetPathColumn_Experimental, FParentAssetPathColumn_Experimental, FAssetsInPathColumn_Experimental>(FName(TEXT("Editor_AssetRegistryPathsTable")));
 	}
 
 	AssetsDataTable = Database.FindTable(FName(TEXT("Editor_AssetRegistryAssetDataTable")));
diff --git a/Engine/Plugins/Experimental/EditorDataStorageFeatures/Source/TedsAssetData/Private/TedsAssetDataHelper.cpp b/Engine/Plugins/Experimental/EditorDataStorageFeatures/Source/TedsAssetData/Private/TedsAssetDataHelper.cpp
new file mode 100644
--- /dev/null
+++ b/Engine/Plugins/Experimental/EditorDataStorageFeatures/Source/TedsAssetData/Private/TedsAssetDataHelper.cpp
@@ -0,0 +1,32 @@
+﻿// Copyright Epic Games, Inc. All Rights Reserved.
+
+#include "TedsAssetDataHelper.h"
+
+namespace TedsAssetDataHelper
+{
+	FString RemoveSlashFromStart(const FString& InStringToModify)
+	{
+		FString StringToReturn = InStringToModify;
+		int32 StartingSlashToRemoveIndex = INDEX_NONE;
+
+		for (int32 Index = 0; Index < StringToReturn.Len(); ++Index)
+		{
+			if (StringToReturn[Index] == TEXT('/'))
+			{
+				// Add 1 since we later use RightChop which use the count and not the index
+				StartingSlashToRemoveIndex = Index + 1;
+			}
+			else
+			{
+				break;
+			}
+		}
+
+		if (StartingSlashToRemoveIndex != INDEX_NONE)
+		{
+			StringToReturn = StringToReturn.RightChop(StartingSlashToRemoveIndex);
+		}
+
+		return StringToReturn;
+	}
+}
diff --git a/Engine/Plugins/Experimental/EditorDataStorageFeatures/Source/TedsAssetData/Private/TedsAssetDataHelper.h b/Engine/Plugins/Experimental/EditorDataStorageFeatures/Source/TedsAssetData/Private/TedsAssetDataHelper.h
new file mode 100644
--- /dev/null
+++ b/Engine/Plugins/Experimental/EditorDataStorageFeatures/Source/TedsAssetData/Private/TedsAssetDataHelper.h
@@ -0,0 +1,11 @@
+﻿// Copyright Epic Games, Inc. All Rights Reserved.
+
+#pragma once
+
+#include "Containers/UnrealString.h"
+
+namespace TedsAssetDataHelper
+{
+	/** Remove all / from the start of InStringToModify */
+	FString RemoveSlashFromStart(const FString& InStringToModify);
+}
diff --git a/Engine/Plugins/Experimental/EditorDataStorageFeatures/Source/TedsAssetData/Private/Widgets/AssetDataDiskSizeWidget.cpp b/Engine/Plugins/Experimental/EditorDataStorageFeatures/Source/TedsAssetData/Private/Widgets/AssetDataDiskSizeWidget.cpp
--- a/Engine/Plugins/Experimental/EditorDataStorageFeatures/Source/TedsAssetData/Private/Widgets/AssetDataDiskSizeWidget.cpp
+++ b/Engine/Plugins/Experimental/EditorDataStorageFeatures/Source/TedsAssetData/Private/Widgets/AssetDataDiskSizeWidget.cpp
@@ -6,6 +6,7 @@
 #include "Internationalization/Text.h"
 #include "TedsAssetDataColumns.h"
 #include "UObject/ObjectMacros.h"
+#include "Widgets/Layout/SBox.h"
 #include "Widgets/Text/STextBlock.h"
 
 void UDiskSizeWidgetFactory::RegisterWidgetConstructors(IEditorDataStorageProvider& DataStorage,
@@ -25,11 +26,16 @@ TSharedPtr<SWidget> FDiskSizeWidgetConstructor::CreateWidget(IEditorDataStorageP
 	const UE::Editor::DataStorage::FMetaDataView& Arguments)
 {
 	UE::Editor::DataStorage::FAttributeBinder Binder(TargetRow, DataStorage);
-	
-	return SNew(STextBlock)
-			.Text(Binder.BindData(&FDiskSizeColumn::DiskSize, [](int64 DiskSize)
-				{
-					return FText::AsMemory(DiskSize);
-				}));
-			
+
+	static const FMargin ColumnItemPadding(5, 0, 5, 0);
+
+	return SNew(SBox)
+			.Padding(ColumnItemPadding)
+			[
+				SNew(STextBlock)
+				.Text(Binder.BindData(&FDiskSizeColumn::DiskSize, [](int64 DiskSize)
+					{
+						return FText::AsMemory(DiskSize);
+					}))
+			];
 }
diff --git a/Engine/Plugins/Experimental/EditorDataStorageFeatures/Source/TedsAssetData/Private/Widgets/AssetDataItemTypeWidget.cpp b/Engine/Plugins/Experimental/EditorDataStorageFeatures/Source/TedsAssetData/Private/Widgets/AssetDataItemTypeWidget.cpp
new file mode 100644
--- /dev/null
+++ b/Engine/Plugins/Experimental/EditorDataStorageFeatures/Source/TedsAssetData/Private/Widgets/AssetDataItemTypeWidget.cpp
@@ -0,0 +1,61 @@
+﻿// Copyright Epic Games, Inc. All Rights Reserved.
+
+#include "Widgets/AssetDataItemTypeWidget.h"
+
+#include "Elements/Framework/TypedElementAttributeBinding.h"
+#include "Internationalization/Text.h"
+#include "TedsAssetDataColumns.h"
+#include "Elements/Columns/TypedElementFolderColumns.h"
+#include "UObject/ObjectMacros.h"
+#include "Widgets/Layout/SBox.h"
+#include "Widgets/Text/STextBlock.h"
+
+#define LOCTEXT_NAMESPACE "FAssetDataItemTypeWidgetConstructor"
+
+void UAssetDataItemTypeWidgetFactory::RegisterWidgetConstructors(IEditorDataStorageProvider& DataStorage,
+														IEditorDataStorageUiProvider& DataStorageUi) const
+{
+	using namespace UE::Editor::DataStorage::Queries;
+	DataStorageUi.RegisterWidgetFactory<FAssetDataItemTypeWidgetConstructor>(TEXT("General.Cell"), TColumn<FAssetClassColumn>() || TColumn<FFolderTag>());
+}
+
+FAssetDataItemTypeWidgetConstructor::FAssetDataItemTypeWidgetConstructor()
+	: FSimpleWidgetConstructor(StaticStruct())
+{
+}
+
+TSharedPtr<SWidget> FAssetDataItemTypeWidgetConstructor::CreateWidget(IEditorDataStorageProvider* DataStorage,
+	IEditorDataStorageUiProvider* DataStorageUi, UE::Editor::DataStorage::RowHandle TargetRow, UE::Editor::DataStorage::RowHandle WidgetRow,
+	const UE::Editor::DataStorage::FMetaDataView& Arguments)
+{
+	UE::Editor::DataStorage::FAttributeBinder Binder(TargetRow, DataStorage);
+
+	static const FMargin ColumnItemPadding(5, 0, 5, 0);
+
+	bool bIsFolder = DataStorage->HasColumns<FFolderTag>(TargetRow);
+	TAttribute<FText> ItemTypeAttribute;
+	if (bIsFolder)
+	{
+		ItemTypeAttribute = LOCTEXT("AssetClassFolder", "Folder");
+	}
+	else
+	{
+		ItemTypeAttribute = Binder.BindData(&FAssetClassColumn::ClassPath, [](FTopLevelAssetPath InClassPath)
+			{
+				const FName AssetClass = InClassPath.GetAssetName();
+				if (!AssetClass.IsNone())
+				{
+					return FText::FromName(InClassPath.GetAssetName());
+				}
+				return LOCTEXT("AssetClass_Invalid", "Invalid or not found");
+			});
+	}
+	return SNew(SBox)
+			.Padding(ColumnItemPadding)
+			[
+				SNew(STextBlock)
+				.Text(ItemTypeAttribute)
+			];
+}
+
+#undef LOCTEXT_NAMESPACE
diff --git a/Engine/Plugins/Experimental/EditorDataStorageFeatures/Source/TedsAssetData/Private/Widgets/AssetDataItemTypeWidget.h b/Engine/Plugins/Experimental/EditorDataStorageFeatures/Source/TedsAssetData/Private/Widgets/AssetDataItemTypeWidget.h
new file mode 100644
--- /dev/null
+++ b/Engine/Plugins/Experimental/EditorDataStorageFeatures/Source/TedsAssetData/Private/Widgets/AssetDataItemTypeWidget.h
@@ -0,0 +1,41 @@
+﻿// Copyright Epic Games, Inc. All Rights Reserved.
+
+#pragma once
+
+#include "Elements/Interfaces/TypedElementDataStorageFactory.h"
+#include "Elements/Interfaces/TypedElementDataStorageUiInterface.h"
+
+#include "AssetDataItemTypeWidget.generated.h"
+
+class IEditorDataStorageProvider;
+class UScriptStruct;
+
+UCLASS()
+class UAssetDataItemTypeWidgetFactory : public UEditorDataStorageFactory
+{
+	GENERATED_BODY()
+
+public:
+	~UAssetDataItemTypeWidgetFactory() override = default;
+
+	TEDSASSETDATA_API void RegisterWidgetConstructors(IEditorDataStorageProvider& DataStorage,
+		IEditorDataStorageUiProvider& DataStorageUi) const override;
+};
+
+// Widget to show Item type (folder/material/mesh etc...)
+USTRUCT()
+struct FAssetDataItemTypeWidgetConstructor : public FSimpleWidgetConstructor
+{
+	GENERATED_BODY()
+
+public:
+	TEDSASSETDATA_API FAssetDataItemTypeWidgetConstructor();
+	~FAssetDataItemTypeWidgetConstructor() override = default;
+
+	TEDSASSETDATA_API virtual TSharedPtr<SWidget> CreateWidget(
+		IEditorDataStorageProvider* DataStorage,
+		IEditorDataStorageUiProvider* DataStorageUi,
+		UE::Editor::DataStorage::RowHandle TargetRow,
+		UE::Editor::DataStorage::RowHandle WidgetRow,
+		const UE::Editor::DataStorage::FMetaDataView& Arguments) override;
+};
diff --git a/Engine/Plugins/Experimental/EditorDataStorageFeatures/Source/TedsAssetData/Private/Widgets/AssetDataLabelWidget.cpp b/Engine/Plugins/Experimental/EditorDataStorageFeatures/Source/TedsAssetData/Private/Widgets/AssetDataLabelWidget.cpp
--- a/Engine/Plugins/Experimental/EditorDataStorageFeatures/Source/TedsAssetData/Private/Widgets/AssetDataLabelWidget.cpp
+++ b/Engine/Plugins/Experimental/EditorDataStorageFeatures/Source/TedsAssetData/Private/Widgets/AssetDataLabelWidget.cpp
@@ -4,11 +4,14 @@
 
 #include "AssetDefinition.h"
 #include "AssetDefinitionRegistry.h"
+#include "AssetViewTypes.h"
 #include "Elements/Framework/TypedElementAttributeBinding.h"
 #include "TedsAssetDataColumns.h"
+#include "TedsAssetDataHelper.h"
 #include "Elements/Columns/TypedElementMiscColumns.h"
 #include "Elements/Columns/TypedElementSlateWidgetColumns.h"
 #include "Widgets/Images/SImage.h"
+#include "Widgets/Layout/SBox.h"
 #include "Widgets/Layout/SSpacer.h"
 #include "Widgets/Text/STextBlock.h"
 #include "Widgets/SBoxPanel.h"
@@ -45,6 +48,14 @@ TSharedPtr<SWidget> FAssetDataLabelWidgetConstructor::CreateWidget(IEditorDataSt
 		bool bIsAsset = DataStorage->HasColumns<FAssetDataColumn_Experimental>(TargetRow);
 
 		TAttribute<FSlateColor> ColorAndOpacityAttribute;
+		TAttribute<const FSlateBrush*> FolderImage;
+		TAttribute<FText> AssetName;
+
+		if (const FNameColumn* AssetNameColumn = DataStorage->GetColumn<FNameColumn>(TargetRow))
+		{
+			const FString AssetNameString = TedsAssetDataHelper::RemoveSlashFromStart(AssetNameColumn->Name.ToString());
+			AssetName = FText::FromString(AssetNameString);
+		}
 
 		// For assets, grab the color from the asset definition
 		if(bIsAsset)
@@ -59,48 +70,79 @@ TSharedPtr<SWidget> FAssetDataLabelWidgetConstructor::CreateWidget(IEditorDataSt
 				return FSlateColor::UseForeground();
 			});
 		}
-		// For folders, use the color column directly
+		// For folders, use the color and folder type column directly
 		else
 		{
+			FolderImage = Binder.BindData(&FFolderTypeColumn_Experimental::FolderType, [] (EFolderType InFolderType)
+			{
+				// Default values
+				switch (InFolderType)
+				{
+					case EFolderType::Developer:
+						return FAppStyle::GetBrush(FName(TEXT("ContentBrowser.AssetTreeFolderClosedDeveloper")));
+
+					case EFolderType::PluginRoot:
+						return FAppStyle::GetBrush(FName(TEXT("ContentBrowser.AssetTreeFolderClosedPluginRoot")));
+
+					// TODO: Cpp and Virtual are not currently checked, see TedsAssetData PopulatePathDataTableRow for more info
+					case EFolderType::Code:
+						return FAppStyle::GetBrush(FName(TEXT("ContentBrowser.AssetTreeFolderClosedCode")));
+
+					case EFolderType::CustomVirtual:
+						return FAppStyle::GetBrush(FName(TEXT("ContentBrowser.AssetTreeFolderClosedVirtual")));
+
+					case EFolderType::Normal:
+					default:
+						return FAppStyle::GetBrush(FName(TEXT("ContentBrowser.ColumnViewFolderIcon")));
+				}
+			});
+
 			ColorAndOpacityAttribute = Binder.BindData(&FSlateColorColumn::Color, FSlateColor::UseForeground());
 		}
 
-		return SNew(SHorizontalBox)
-			+SHorizontalBox::Slot()
-				.AutoWidth()
-			[
-				SNew(SImage)
-					.Image(bIsAsset ? FAppStyle::GetBrush(FName(TEXT("ContentBrowser.ColumnViewAssetIcon"))) : FAppStyle::GetBrush(FName(TEXT("ContentBrowser.ColumnViewFolderIcon"))))
-					.ColorAndOpacity(ColorAndOpacityAttribute)
-			]
-			+SHorizontalBox::Slot()
-				.AutoWidth()
-			[
-				SNew(SSpacer)
-					.Size(FVector2D(5.0f, 0.0f))
-			]
-			+SHorizontalBox::Slot()
-				.FillWidth(1.0f)
-			[
-				SNew(STextBlock)
-				.Text(Binder.BindText(&FNameColumn::Name))
-				.ToolTipText(Binder.BindTextFormat(
-								LOCTEXT("AssetLabelTooltip", 
-									"{Name}\n\nVirtual path: {VirtualPath}\n  Asset path: {AssetPath}\n  Verse path: {VersePath}"))
-								.Arg(TEXT("Name"), &FNameColumn::Name)
-								.Arg(TEXT("VirtualPath"), &FVirtualPathColumn_Experimental::VirtualPath, LOCTEXT("PathNotSet", "<not set>"))
-								.Arg(TEXT("AssetPath"), &FAssetPathColumn_Experimental::Path, LOCTEXT("PathNotSet", "<not set>"))
-								.Arg(TEXT("VersePath"), &FVersePathColumn::VersePath, 
-									[](const UE::Core::FVersePath& Path) 
-									{
-										return FText::FromStringView(Path.AsStringView());
-									}, 
-									LOCTEXT("PathNotSet", "<not set>")))
-			];
+		static const FMargin ColumnItemPadding(5, 0, 5, 0);
+
+		return SNew(SBox)
+				.Padding(ColumnItemPadding)
+				[
+					SNew(SHorizontalBox)
+					+SHorizontalBox::Slot()
+					.AutoWidth()
+					[
+						SNew(SImage)
+							.Image(bIsAsset ? FAppStyle::GetBrush(FName(TEXT("ContentBrowser.ColumnViewAssetIcon"))) : FolderImage)
+							.ColorAndOpacity(ColorAndOpacityAttribute)
+					]
+					+SHorizontalBox::Slot()
+					.AutoWidth()
+					[
+						SNew(SSpacer)
+							.Size(FVector2D(5.0f, 0.0f))
+					]
+					+SHorizontalBox::Slot()
+					.FillWidth(1.0f)
+					[
+						SNew(STextBlock)
+						.Text(AssetName)
+						.ToolTipText(Binder.BindTextFormat(
+										LOCTEXT("AssetLabelTooltip", 
+											"{Name}\n\nVirtual path: {VirtualPath}\n  Asset path: {AssetPath}\n  Verse path: {VersePath}"))
+										.Arg(TEXT("Name"), &FNameColumn::Name)
+										.Arg(TEXT("VirtualPath"), &FVirtualPathColumn_Experimental::VirtualPath, LOCTEXT("PathNotSet", "<not set>"))
+										.Arg(TEXT("AssetPath"), &FAssetPathColumn_Experimental::Path, LOCTEXT("PathNotSet", "<not set>"))
+										.Arg(TEXT("VersePath"), &FVersePathColumn::VersePath, 
+											[](const UE::Core::FVersePath& Path) 
+											{
+												return FText::FromStringView(Path.AsStringView());
+											}, 
+											LOCTEXT("PathNotSet", "<not set>")))
+					]
+				];
 	}
 	else
 	{
 		return SNullWidget::NullWidget;
 	}
 }
-#undef LOCTEXT_NAMESPACE
\ No newline at end of file
+
+#undef LOCTEXT_NAMESPACE
diff --git a/Engine/Plugins/Experimental/EditorDataStorageFeatures/Source/TedsAssetData/Private/Widgets/AssetDataVirtualPathWidget.cpp b/Engine/Plugins/Experimental/EditorDataStorageFeatures/Source/TedsAssetData/Private/Widgets/AssetDataVirtualPathWidget.cpp
new file mode 100644
--- /dev/null
+++ b/Engine/Plugins/Experimental/EditorDataStorageFeatures/Source/TedsAssetData/Private/Widgets/AssetDataVirtualPathWidget.cpp
@@ -0,0 +1,47 @@
+﻿// Copyright Epic Games, Inc. All Rights Reserved.
+
+#include "Widgets/AssetDataVirtualPathWidget.h"
+
+#include "Elements/Framework/TypedElementAttributeBinding.h"
+#include "Internationalization/Text.h"
+#include "TedsAssetDataColumns.h"
+#include "TedsAssetDataHelper.h"
+#include "UObject/ObjectMacros.h"
+#include "Widgets/Layout/SBox.h"
+#include "Widgets/Text/STextBlock.h"
+
+void UAssetDataVirtualPathWidgetFactory::RegisterWidgetConstructors(IEditorDataStorageProvider& DataStorage,
+														IEditorDataStorageUiProvider& DataStorageUi) const
+{
+	using namespace UE::Editor::DataStorage::Queries;
+	DataStorageUi.RegisterWidgetFactory<FAssetDataVirtualPathWidgetConstructor>(TEXT("General.Cell"), TColumn<FVirtualPathColumn_Experimental>());
+}
+
+FAssetDataVirtualPathWidgetConstructor::FAssetDataVirtualPathWidgetConstructor()
+	: FSimpleWidgetConstructor(StaticStruct())
+{
+}
+
+TSharedPtr<SWidget> FAssetDataVirtualPathWidgetConstructor::CreateWidget(IEditorDataStorageProvider* DataStorage,
+	IEditorDataStorageUiProvider* DataStorageUi, UE::Editor::DataStorage::RowHandle TargetRow, UE::Editor::DataStorage::RowHandle WidgetRow,
+	const UE::Editor::DataStorage::FMetaDataView& Arguments)
+{
+	UE::Editor::DataStorage::FAttributeBinder Binder(TargetRow, DataStorage);
+	static const FMargin ColumnItemPadding(5, 0, 5, 0);
+
+	const FVirtualPathColumn_Experimental* VirtualPathColumn = DataStorage->GetColumn<FVirtualPathColumn_Experimental>(TargetRow);
+	TAttribute<FText> VirtualPathView = FText::GetEmpty();
+
+	if (VirtualPathColumn)
+	{
+		const FString VirtualPathString = TedsAssetDataHelper::RemoveSlashFromStart(VirtualPathColumn->VirtualPath.ToString());
+		VirtualPathView = FText::FromString(VirtualPathString);
+	}
+
+	return SNew(SBox)
+			.Padding(ColumnItemPadding)
+			[
+				SNew(STextBlock)
+				.Text(VirtualPathView)
+			];
+}
diff --git a/Engine/Plugins/Experimental/EditorDataStorageFeatures/Source/TedsAssetData/Private/Widgets/AssetDataVirtualPathWidget.h b/Engine/Plugins/Experimental/EditorDataStorageFeatures/Source/TedsAssetData/Private/Widgets/AssetDataVirtualPathWidget.h
new file mode 100644
--- /dev/null
+++ b/Engine/Plugins/Experimental/EditorDataStorageFeatures/Source/TedsAssetData/Private/Widgets/AssetDataVirtualPathWidget.h
@@ -0,0 +1,41 @@
+﻿// Copyright Epic Games, Inc. All Rights Reserved.
+
+#pragma once
+
+#include "Elements/Interfaces/TypedElementDataStorageFactory.h"
+#include "Elements/Interfaces/TypedElementDataStorageUiInterface.h"
+
+#include "AssetDataVirtualPathWidget.generated.h"
+
+class IEditorDataStorageProvider;
+class UScriptStruct;
+
+UCLASS()
+class UAssetDataVirtualPathWidgetFactory : public UEditorDataStorageFactory
+{
+	GENERATED_BODY()
+
+public:
+	~UAssetDataVirtualPathWidgetFactory() override = default;
+
+	TEDSASSETDATA_API void RegisterWidgetConstructors(IEditorDataStorageProvider& DataStorage,
+		IEditorDataStorageUiProvider& DataStorageUi) const override;
+};
+
+// Widget to show disk size in bytes
+USTRUCT()
+struct FAssetDataVirtualPathWidgetConstructor : public FSimpleWidgetConstructor
+{
+	GENERATED_BODY()
+
+public:
+	TEDSASSETDATA_API FAssetDataVirtualPathWidgetConstructor();
+	~FAssetDataVirtualPathWidgetConstructor() override = default;
+
+	TEDSASSETDATA_API virtual TSharedPtr<SWidget> CreateWidget(
+		IEditorDataStorageProvider* DataStorage,
+		IEditorDataStorageUiProvider* DataStorageUi,
+		UE::Editor::DataStorage::RowHandle TargetRow,
+		UE::Editor::DataStorage::RowHandle WidgetRow,
+		const UE::Editor::DataStorage::FMetaDataView& Arguments) override;
+};
diff --git a/Engine/Plugins/Experimental/EditorDataStorageFeatures/Source/TedsAssetData/Public/TedsAssetDataColumns.h b/Engine/Plugins/Experimental/EditorDataStorageFeatures/Source/TedsAssetData/Public/TedsAssetDataColumns.h
--- a/Engine/Plugins/Experimental/EditorDataStorageFeatures/Source/TedsAssetData/Public/TedsAssetDataColumns.h
+++ b/Engine/Plugins/Experimental/EditorDataStorageFeatures/Source/TedsAssetData/Public/TedsAssetDataColumns.h
@@ -5,6 +5,7 @@
 #include "TedsAssetDataStructs.h"
 
 #include "AssetRegistry/AssetData.h"
+#include "AssetViewTypes.h" 
 #include "Containers/UnrealString.h"
 #include "Containers/VersePath.h"
 #include "Elements/Common/TypedElementCommonTypes.h"
@@ -48,6 +49,14 @@ struct FAssetsInPathColumn_Experimental : public FEditorDataStorageColumn
 	TSet<UE::Editor::DataStorage::RowHandle> AssetsRow;
 };
 
+USTRUCT()
+struct FFolderTypeColumn_Experimental : public FEditorDataStorageColumn
+{
+	GENERATED_BODY()
+	
+	EFolderType FolderType;
+};
+
 USTRUCT()
 struct FAssetDataColumn_Experimental : public FEditorDataStorageColumn
 {
@@ -95,7 +104,7 @@ struct FPublicAssetTag final : public FEditorDataStorageTag
 };
 
 // Column to store the type of an asset
-USTRUCT(meta = (DisplayName = "Type Path"))
+USTRUCT(meta = (DisplayName = "Type"))
 struct FAssetClassColumn final : public FEditorDataStorageColumn
 {
 	GENERATED_BODY()
diff --git a/Engine/Plugins/Experimental/EditorDataStorageFeatures/Source/TedsAssetData/TedsAssetData.build.cs b/Engine/Plugins/Experimental/EditorDataStorageFeatures/Source/TedsAssetData/TedsAssetData.build.cs
--- a/Engine/Plugins/Experimental/EditorDataStorageFeatures/Source/TedsAssetData/TedsAssetData.build.cs
+++ b/Engine/Plugins/Experimental/EditorDataStorageFeatures/Source/TedsAssetData/TedsAssetData.build.cs
@@ -19,6 +19,7 @@ public class TedsAssetData : ModuleRules
 			PrivateDependencyModuleNames.AddRange(
 				new string[]
 				{
+					"AssetTools",
 					"AssetDefinition",
 					"AssetRegistry",
 					"ContentBrowserData",
@@ -32,13 +33,13 @@ public class TedsAssetData : ModuleRules
 					"UnrealEd",
 					"UnrealEd",
 					"AssetDefinition",
-					"ContentBrowser"
 				});
 
 			PublicDependencyModuleNames.AddRange(
 				new string[]
 				{
 					"Projects",
+					"ContentBrowser"
 				});
 
 		}
diff --git a/Engine/Source/Editor/ContentBrowser/Private/SourcesViewWidgets.h b/Engine/Source/Editor/ContentBrowser/Private/SourcesViewWidgets.h
--- a/Engine/Source/Editor/ContentBrowser/Private/SourcesViewWidgets.h
+++ b/Engine/Source/Editor/ContentBrowser/Private/SourcesViewWidgets.h
@@ -3,6 +3,7 @@
 #pragma once
 
 #include "AssetTagItemTypes.h"
+#include "AssetViewTypes.h"
 #include "Delegates/Delegate.h"
 #include "Framework/SlateDelegates.h"
 #include "HAL/Platform.h"
@@ -103,15 +104,6 @@ private:
 	TSharedRef<SWidget> GenerateStateIcons();
 	
 private:
-	enum class EFolderType : uint8
-	{
-		Normal,
-		CustomVirtual, // No corresponding on-disk path, used for organization in the content browser
-		PluginRoot,    // Root content folder of a plugin
-		Code,
-		Developer,
-	};
-
 	/** The data for this item */
 	TWeakPtr<FTreeItem> TreeItem;
 
diff --git a/Engine/Source/Editor/ContentBrowser/Public/AssetViewTypes.h b/Engine/Source/Editor/ContentBrowser/Public/AssetViewTypes.h
--- a/Engine/Source/Editor/ContentBrowser/Public/AssetViewTypes.h
+++ b/Engine/Source/Editor/ContentBrowser/Public/AssetViewTypes.h
@@ -18,6 +18,17 @@ class FString;
 class FText;
 struct FAssetViewCustomColumn;
 
+enum class EFolderType : uint8
+{
+	Normal,
+	CustomVirtual, // No corresponding on-disk path, used for organization in the content browser
+	PluginRoot,    // Root content folder of a plugin
+	Code,
+	Developer,
+};
+
+ENUM_CLASS_FLAGS(EFolderType)
+
 /** An item (folder or file) displayed in the asset view */
 class FAssetViewItem
 {
