[EditorDataStorageFeature] Teds Thumbnails

diff --git a/Engine/Plugins/Experimental/EditorDataStorageFeatures/Source/TedsAssetData/Private/TedsAssetDataHelper.cpp b/Engine/Plugins/Experimental/EditorDataStorageFeatures/Source/TedsAssetData/Private/TedsAssetDataHelper.cpp
--- a/Engine/Plugins/Experimental/EditorDataStorageFeatures/Source/TedsAssetData/Private/TedsAssetDataHelper.cpp
+++ b/Engine/Plugins/Experimental/EditorDataStorageFeatures/Source/TedsAssetData/Private/TedsAssetDataHelper.cpp
@@ -29,4 +29,17 @@ namespace TedsAssetDataHelper
 
 		return StringToReturn;
 	}
+
+	FString RemoveAllFromLastSlash(const FString& InStringToModify)
+	{
+		int32 LastSlashToRemoveIndex = INDEX_NONE;
+		InStringToModify.FindLastChar(TEXT('/'), LastSlashToRemoveIndex);
+
+		if (LastSlashToRemoveIndex != INDEX_NONE)
+		{
+			return InStringToModify.Left(LastSlashToRemoveIndex);
+		}
+
+		return InStringToModify;
+	}
 }
diff --git a/Engine/Plugins/Experimental/EditorDataStorageFeatures/Source/TedsAssetData/Private/TedsAssetDataHelper.h b/Engine/Plugins/Experimental/EditorDataStorageFeatures/Source/TedsAssetData/Private/TedsAssetDataHelper.h
--- a/Engine/Plugins/Experimental/EditorDataStorageFeatures/Source/TedsAssetData/Private/TedsAssetDataHelper.h
+++ b/Engine/Plugins/Experimental/EditorDataStorageFeatures/Source/TedsAssetData/Private/TedsAssetDataHelper.h
@@ -8,4 +8,7 @@ namespace TedsAssetDataHelper
 {
 	/** Remove all / from the start of InStringToModify */
 	FString RemoveSlashFromStart(const FString& InStringToModify);
+
+	/** Remove Last / and successive char from InStringToModify */
+	FString RemoveAllFromLastSlash(const FString& InStringToModify);
 }
diff --git a/Engine/Plugins/Experimental/EditorDataStorageFeatures/Source/TedsAssetData/Private/Widgets/AssetDataLabelWidget.cpp b/Engine/Plugins/Experimental/EditorDataStorageFeatures/Source/TedsAssetData/Private/Widgets/AssetDataLabelWidget.cpp
--- a/Engine/Plugins/Experimental/EditorDataStorageFeatures/Source/TedsAssetData/Private/Widgets/AssetDataLabelWidget.cpp
+++ b/Engine/Plugins/Experimental/EditorDataStorageFeatures/Source/TedsAssetData/Private/Widgets/AssetDataLabelWidget.cpp
@@ -49,12 +49,30 @@ TSharedPtr<SWidget> FAssetDataLabelWidgetConstructor::CreateWidget(IEditorDataSt
 
 		TAttribute<FSlateColor> ColorAndOpacityAttribute;
 		TAttribute<const FSlateBrush*> FolderImage;
-		TAttribute<FText> AssetName;
 
-		if (const FNameColumn* AssetNameColumn = DataStorage->GetColumn<FNameColumn>(TargetRow))
+		// Create the LabelWidget through its General Widget
+		static FName LabelPurposeName = TEXT("General.Cell");
+
+		TSharedPtr<FTypedElementWidgetConstructor> OutLabelWidgetConstructorPtr;
+
+		auto AssignLabelWidgetToColumn = [&OutLabelWidgetConstructorPtr] (TUniquePtr<FTypedElementWidgetConstructor> Constructor, TConstArrayView<TWeakObjectPtr<const UScriptStruct>>)
+		{
+			OutLabelWidgetConstructorPtr = TSharedPtr<FTypedElementWidgetConstructor>(Constructor.Release());
+			return false;
+		};
+
+		TArray<TWeakObjectPtr<const UScriptStruct>> LabelColumns = GetLabelColumns();
+		DataStorageUi->CreateWidgetConstructors(LabelPurposeName, IEditorDataStorageUiProvider::EMatchApproach::ExactMatch, LabelColumns, Arguments, AssignLabelWidgetToColumn);
+
+		TSharedPtr<SWidget> LabelWidget = SNullWidget::NullWidget;
+		if (OutLabelWidgetConstructorPtr)
 		{
-			const FString AssetNameString = TedsAssetDataHelper::RemoveSlashFromStart(AssetNameColumn->Name.ToString());
-			AssetName = FText::FromString(AssetNameString);
+			LabelWidget = DataStorageUi->ConstructWidget(WidgetRow, *OutLabelWidgetConstructorPtr, Arguments);
+		}
+
+		if (!LabelWidget.IsValid())
+		{
+			LabelWidget = SNullWidget::NullWidget;
 		}
 
 		// For assets, grab the color from the asset definition
@@ -122,20 +140,7 @@ TSharedPtr<SWidget> FAssetDataLabelWidgetConstructor::CreateWidget(IEditorDataSt
 					+SHorizontalBox::Slot()
 					.FillWidth(1.0f)
 					[
-						SNew(STextBlock)
-						.Text(AssetName)
-						.ToolTipText(Binder.BindTextFormat(
-										LOCTEXT("AssetLabelTooltip", 
-											"{Name}\n\nVirtual path: {VirtualPath}\n  Asset path: {AssetPath}\n  Verse path: {VersePath}"))
-										.Arg(TEXT("Name"), &FNameColumn::Name)
-										.Arg(TEXT("VirtualPath"), &FVirtualPathColumn_Experimental::VirtualPath, LOCTEXT("PathNotSet", "<not set>"))
-										.Arg(TEXT("AssetPath"), &FAssetPathColumn_Experimental::Path, LOCTEXT("PathNotSet", "<not set>"))
-										.Arg(TEXT("VersePath"), &FVersePathColumn::VersePath, 
-											[](const UE::Core::FVersePath& Path) 
-											{
-												return FText::FromStringView(Path.AsStringView());
-											}, 
-											LOCTEXT("PathNotSet", "<not set>")))
+						LabelWidget.ToSharedRef()
 					]
 				];
 	}
@@ -145,4 +150,10 @@ TSharedPtr<SWidget> FAssetDataLabelWidgetConstructor::CreateWidget(IEditorDataSt
 	}
 }
 
+TArray<TWeakObjectPtr<const UScriptStruct>> FAssetDataLabelWidgetConstructor::GetLabelColumns()
+{
+	static TArray<TWeakObjectPtr<const UScriptStruct>> LabelColumns({ TWeakObjectPtr(FNameColumn::StaticStruct()) });
+	return LabelColumns;
+}
+
 #undef LOCTEXT_NAMESPACE
diff --git a/Engine/Plugins/Experimental/EditorDataStorageFeatures/Source/TedsAssetData/Private/Widgets/AssetDataThumbnailWidget.cpp b/Engine/Plugins/Experimental/EditorDataStorageFeatures/Source/TedsAssetData/Private/Widgets/AssetDataThumbnailWidget.cpp
new file mode 100644
--- /dev/null
+++ b/Engine/Plugins/Experimental/EditorDataStorageFeatures/Source/TedsAssetData/Private/Widgets/AssetDataThumbnailWidget.cpp
@@ -0,0 +1,126 @@
+// Copyright Epic Games, Inc. All Rights Reserved.
+
+#include "Widgets/AssetDataThumbnailWidget.h"
+
+#include "AssetThumbnail.h"
+#include "AssetViewTypes.h"
+#include "Elements/Framework/TypedElementAttributeBinding.h"
+#include "TedsAssetDataColumns.h"
+#include "Elements/Columns/TypedElementFolderColumns.h"
+#include "Elements/Columns/TypedElementSlateWidgetColumns.h"
+#include "ThumbnailRendering/ThumbnailManager.h"
+#include "Widgets/Images/SImage.h"
+#include "Widgets/Layout/SBox.h"
+
+#define LOCTEXT_NAMESPACE "FAssetThumbnailWidgetConstructor"
+
+void UAssetThumbnailWidgetFactory::RegisterWidgetConstructors(IEditorDataStorageProvider& DataStorage,
+	IEditorDataStorageUiProvider& DataStorageUi) const
+{
+	using namespace UE::Editor::DataStorage::Queries;
+	DataStorageUi.RegisterWidgetFactory<FAssetThumbnailWidgetConstructor>(
+		"General.Cell",
+		TColumn<FAssetTag>() || TColumn<FFolderTag>());
+}
+
+FAssetThumbnailWidgetConstructor::FAssetThumbnailWidgetConstructor()
+	: FSimpleWidgetConstructor(StaticStruct())
+{
+}
+
+FAssetThumbnailWidgetConstructor::FAssetThumbnailWidgetConstructor(const UScriptStruct* TypeInfo)
+	: FSimpleWidgetConstructor(TypeInfo)
+{
+}
+
+TSharedPtr<SWidget> FAssetThumbnailWidgetConstructor::CreateWidget(IEditorDataStorageProvider* DataStorage,
+	IEditorDataStorageUiProvider* DataStorageUi, UE::Editor::DataStorage::RowHandle TargetRow, UE::Editor::DataStorage::RowHandle WidgetRow,
+	const UE::Editor::DataStorage::FMetaDataView& Arguments)
+{
+	UE::Editor::DataStorage::FAttributeBinder Binder(TargetRow, DataStorage);
+
+	bool bIsAsset = DataStorage->HasColumns<FAssetTag>(TargetRow);
+
+	constexpr float ThumbnailSize = 64.f;
+	constexpr float ThumbnailVerticalPadding = 4.f;
+
+	// Thumbnail Box container
+	const TSharedRef<SBox> ThumbnailBox = SNew(SBox)
+		.WidthOverride(ThumbnailSize)
+		.HeightOverride(ThumbnailSize + (ThumbnailVerticalPadding * 2))
+		.Padding(0.f, ThumbnailVerticalPadding);
+
+	// For assets, grab the color from the asset definition
+	if (bIsAsset)
+	{
+		// Retrieve the AssetData, used to create the thumbnail
+		FAssetData AssetDataToUse = FAssetData();
+		if (const FAssetDataColumn_Experimental* AssetDataColumn = DataStorage->GetColumn<FAssetDataColumn_Experimental>(TargetRow))
+		{
+			AssetDataToUse = AssetDataColumn->AssetData;
+		}
+
+		// Thumbnail Configuration Args
+		FAssetThumbnailConfig ThumbnailConfig;
+		ThumbnailConfig.bAllowFadeIn = true;
+		ThumbnailConfig.bAllowHintText = false;
+		ThumbnailConfig.bAllowRealTimeOnHovered = false; // we use our own OnMouseEnter/Leave for logical asset item
+		// TODO: AssetItem->GetItem().GetItemTemporaryReason() == EContentBrowserItemFlags::Temporary_Creation;
+		// This was the previous check to assign bForceGenericThumbnail, for now it will use the Generic if the AssetData is not valid
+		ThumbnailConfig.bForceGenericThumbnail = !AssetDataToUse.IsValid();
+		ThumbnailConfig.bAllowAssetSpecificThumbnailOverlay = !ThumbnailConfig.bForceGenericThumbnail;
+		ThumbnailConfig.ThumbnailLabel = EThumbnailLabel::ClassName;
+		// TODO: If Teds View items will be bound to the CB size and can change their size, this would need to be bound to that as well
+		// ThumbnailConfig.GenericThumbnailSize = 64.f;
+		// TODO: This data will need to be integrated in TEDS and later on remove this/update it to use TEDS instead to retrieve the data needed
+		// ThumbnailConfig.AssetSystemInfoProvider = nullptr;
+		ThumbnailConfig.bShowAssetChip = true;
+		// TODO: If we are able to change the thumbnail in Teds view we would need to bind this to make it work
+		// ThumbnailConfig.bCanDisplayEditModePrimitiveTools = false;
+		// ThumbnailConfig.IsEditModeVisible = false;
+
+		// TODO: Consider caching them instead, to avoid creating them everytime
+		const TSharedRef<FAssetThumbnail> AssetThumbnail = MakeShared<FAssetThumbnail>(AssetDataToUse, ThumbnailSize, ThumbnailSize, UThumbnailManager::Get().GetSharedThumbnailPool());
+
+		ThumbnailBox->SetContent(AssetThumbnail->MakeThumbnailWidget(ThumbnailConfig));
+	}
+	// For folders, use the color and folder type column directly
+	else
+	{
+		TAttribute<const FSlateBrush*> FolderImage;
+		TAttribute<FSlateColor> ColorAndOpacityAttribute;
+
+		FolderImage = Binder.BindData(&FFolderTypeColumn_Experimental::FolderType, [] (EFolderType InFolderType)
+		{
+			// Default values
+			switch (InFolderType)
+			{
+				case EFolderType::Developer:
+					return FAppStyle::GetBrush(FName(TEXT("ContentBrowser.ListViewDeveloperFolderIcon")));
+
+				case EFolderType::PluginRoot:
+					return FAppStyle::GetBrush(FName(TEXT("ContentBrowser.ListViewPluginFolderIcon")));
+
+				// TODO: Cpp and Virtual are not currently checked, see TedsAssetData PopulatePathDataTableRow for more info
+				case EFolderType::Code:
+					return FAppStyle::GetBrush(FName(TEXT("ContentBrowser.ListViewCodeFolderIcon")));
+
+				case EFolderType::CustomVirtual:
+					return FAppStyle::GetBrush(FName(TEXT("ContentBrowser.ListViewVirtualFolderIcon")));
+
+				case EFolderType::Normal:
+				default:
+					return FAppStyle::GetBrush(FName(TEXT("ContentBrowser.ListViewFolderIcon")));
+			}
+		});
+
+		ColorAndOpacityAttribute = Binder.BindData(&FSlateColorColumn::Color, FSlateColor::UseForeground());
+		ThumbnailBox->SetContent(SNew(SImage)
+						.Image(FolderImage)
+						.ColorAndOpacity(ColorAndOpacityAttribute));
+	}
+
+	return ThumbnailBox;
+}
+
+#undef LOCTEXT_NAMESPACE
diff --git a/Engine/Plugins/Experimental/EditorDataStorageFeatures/Source/TedsAssetData/Private/Widgets/AssetDataThumbnailWidget.h b/Engine/Plugins/Experimental/EditorDataStorageFeatures/Source/TedsAssetData/Private/Widgets/AssetDataThumbnailWidget.h
new file mode 100644
--- /dev/null
+++ b/Engine/Plugins/Experimental/EditorDataStorageFeatures/Source/TedsAssetData/Private/Widgets/AssetDataThumbnailWidget.h
@@ -0,0 +1,44 @@
+// Copyright Epic Games, Inc. All Rights Reserved.
+
+#pragma once
+
+#include "Elements/Interfaces/TypedElementDataStorageFactory.h"
+#include "Elements/Interfaces/TypedElementDataStorageUiInterface.h"
+#include "Internationalization/Text.h"
+#include "UObject/ObjectMacros.h"
+
+#include "AssetDataThumbnailWidget.generated.h"
+
+class IEditorDataStorageProvider;
+class UScriptStruct;
+
+UCLASS()
+class UAssetThumbnailWidgetFactory : public UEditorDataStorageFactory
+{
+	GENERATED_BODY()
+
+public:
+	virtual ~UAssetThumbnailWidgetFactory() override = default;
+
+	TEDSASSETDATA_API virtual void RegisterWidgetConstructors(IEditorDataStorageProvider& DataStorage,
+		IEditorDataStorageUiProvider& DataStorageUi) const override;
+};
+
+// Thumbnail widget for assets in TEDS
+USTRUCT()
+struct FAssetThumbnailWidgetConstructor : public FSimpleWidgetConstructor
+{
+	GENERATED_BODY()
+
+public:
+	TEDSASSETDATA_API FAssetThumbnailWidgetConstructor();
+	TEDSASSETDATA_API explicit FAssetThumbnailWidgetConstructor(const UScriptStruct* TypeInfo);
+	virtual ~FAssetThumbnailWidgetConstructor() override = default;
+
+	TEDSASSETDATA_API virtual TSharedPtr<SWidget> CreateWidget(
+		IEditorDataStorageProvider* DataStorage,
+		IEditorDataStorageUiProvider* DataStorageUi,
+		UE::Editor::DataStorage::RowHandle TargetRow,
+		UE::Editor::DataStorage::RowHandle WidgetRow, 
+		const UE::Editor::DataStorage::FMetaDataView& Arguments) override;
+};
diff --git a/Engine/Plugins/Experimental/EditorDataStorageFeatures/Source/TedsAssetData/Private/Widgets/AssetDataVirtualPathWidget.cpp b/Engine/Plugins/Experimental/EditorDataStorageFeatures/Source/TedsAssetData/Private/Widgets/AssetDataVirtualPathWidget.cpp
--- a/Engine/Plugins/Experimental/EditorDataStorageFeatures/Source/TedsAssetData/Private/Widgets/AssetDataVirtualPathWidget.cpp
+++ b/Engine/Plugins/Experimental/EditorDataStorageFeatures/Source/TedsAssetData/Private/Widgets/AssetDataVirtualPathWidget.cpp
@@ -11,10 +11,12 @@
 #include "Widgets/Text/STextBlock.h"
 
 void UAssetDataVirtualPathWidgetFactory::RegisterWidgetConstructors(IEditorDataStorageProvider& DataStorage,
-														IEditorDataStorageUiProvider& DataStorageUi) const
+	IEditorDataStorageUiProvider& DataStorageUi) const
 {
 	using namespace UE::Editor::DataStorage::Queries;
-	DataStorageUi.RegisterWidgetFactory<FAssetDataVirtualPathWidgetConstructor>(TEXT("General.Cell"), TColumn<FVirtualPathColumn_Experimental>());
+	DataStorageUi.RegisterWidgetFactory<FAssetDataVirtualPathWidgetConstructor>(
+		TEXT("General.Cell"),
+		TColumn<FVirtualPathColumn_Experimental>() || TColumn<FAssetPathColumn_Experimental>());
 }
 
 FAssetDataVirtualPathWidgetConstructor::FAssetDataVirtualPathWidgetConstructor()
@@ -26,22 +28,35 @@ TSharedPtr<SWidget> FAssetDataVirtualPathWidgetConstructor::CreateWidget(IEditor
 	IEditorDataStorageUiProvider* DataStorageUi, UE::Editor::DataStorage::RowHandle TargetRow, UE::Editor::DataStorage::RowHandle WidgetRow,
 	const UE::Editor::DataStorage::FMetaDataView& Arguments)
 {
-	UE::Editor::DataStorage::FAttributeBinder Binder(TargetRow, DataStorage);
 	static const FMargin ColumnItemPadding(5, 0, 5, 0);
 
-	const FVirtualPathColumn_Experimental* VirtualPathColumn = DataStorage->GetColumn<FVirtualPathColumn_Experimental>(TargetRow);
-	TAttribute<FText> VirtualPathView = FText::GetEmpty();
+	bool bIsAsset = DataStorage->HasColumns<FAssetTag>(TargetRow);
+	TAttribute<FText> PathText = FText::GetEmpty();
 
-	if (VirtualPathColumn)
+	if (bIsAsset)
 	{
-		const FString VirtualPathString = TedsAssetDataHelper::RemoveSlashFromStart(VirtualPathColumn->VirtualPath.ToString());
-		VirtualPathView = FText::FromString(VirtualPathString);
+		if (const FVirtualPathColumn_Experimental* VirtualPathColumn = DataStorage->GetColumn<FVirtualPathColumn_Experimental>(TargetRow))
+		{
+			FString VirtualPathString = VirtualPathColumn->VirtualPath.ToString();
+			VirtualPathString = TedsAssetDataHelper::RemoveSlashFromStart(VirtualPathString);
+			PathText = FText::FromString(VirtualPathString);
+		}
+	}
+	else
+	{
+		if (const FAssetPathColumn_Experimental* AssetPathColumn = DataStorage->GetColumn<FAssetPathColumn_Experimental>(TargetRow))
+		{
+			FString AssetPathString = AssetPathColumn->Path.ToString();
+			AssetPathString = TedsAssetDataHelper::RemoveSlashFromStart(AssetPathString);
+			AssetPathString = TedsAssetDataHelper::RemoveAllFromLastSlash(AssetPathString);
+			PathText = FText::FromString(AssetPathString);
+		}
 	}
 
 	return SNew(SBox)
 			.Padding(ColumnItemPadding)
 			[
 				SNew(STextBlock)
-				.Text(VirtualPathView)
+				.Text(PathText)
 			];
 }
diff --git a/Engine/Plugins/Experimental/EditorDataStorageFeatures/Source/TedsAssetData/Private/Widgets/AssetLabelWidget.cpp b/Engine/Plugins/Experimental/EditorDataStorageFeatures/Source/TedsAssetData/Private/Widgets/AssetLabelWidget.cpp
new file mode 100644
--- /dev/null
+++ b/Engine/Plugins/Experimental/EditorDataStorageFeatures/Source/TedsAssetData/Private/Widgets/AssetLabelWidget.cpp
@@ -0,0 +1,63 @@
+// Copyright Epic Games, Inc. All Rights Reserved.
+
+#include "Widgets/AssetLabelWidget.h"
+
+#include "Elements/Framework/TypedElementAttributeBinding.h"
+#include "TedsAssetDataColumns.h"
+#include "TedsAssetDataHelper.h"
+#include "Elements/Columns/TypedElementMiscColumns.h"
+#include "Widgets/Text/STextBlock.h"
+
+#define LOCTEXT_NAMESPACE "FAssetLabelWidgetConstructor"
+
+void UAssetLabelWidgetFactory::RegisterWidgetConstructors(IEditorDataStorageProvider& DataStorage,
+	IEditorDataStorageUiProvider& DataStorageUi) const
+{
+	using namespace UE::Editor::DataStorage::Queries;
+	DataStorageUi.RegisterWidgetFactory<FAssetLabelWidgetConstructor>(
+		"General.Cell",
+		TColumn<FNameColumn>());
+}
+
+FAssetLabelWidgetConstructor::FAssetLabelWidgetConstructor()
+	: FSimpleWidgetConstructor(StaticStruct())
+{
+}
+
+FAssetLabelWidgetConstructor::FAssetLabelWidgetConstructor(const UScriptStruct* TypeInfo)
+	: FSimpleWidgetConstructor(TypeInfo)
+{
+}
+
+TSharedPtr<SWidget> FAssetLabelWidgetConstructor::CreateWidget(IEditorDataStorageProvider* DataStorage,
+	IEditorDataStorageUiProvider* DataStorageUi, UE::Editor::DataStorage::RowHandle TargetRow, UE::Editor::DataStorage::RowHandle WidgetRow,
+	const UE::Editor::DataStorage::FMetaDataView& Arguments)
+{
+	UE::Editor::DataStorage::FAttributeBinder Binder(TargetRow, DataStorage);
+
+	TAttribute<FText> AssetName;
+
+	if (const FNameColumn* AssetNameColumn = DataStorage->GetColumn<FNameColumn>(TargetRow))
+	{
+		const FString AssetNameString = TedsAssetDataHelper::RemoveSlashFromStart(AssetNameColumn->Name.ToString());
+		AssetName = FText::FromString(AssetNameString);
+	}
+
+	return SNew(STextBlock)
+				.Text(AssetName)
+				.OverflowPolicy(ETextOverflowPolicy::MiddleEllipsis)
+				.ToolTipText(Binder.BindTextFormat(
+								LOCTEXT("AssetLabelTooltip", 
+									"{Name}\n\nVirtual path: {VirtualPath}\n  Asset path: {AssetPath}\n  Verse path: {VersePath}"))
+								.Arg(TEXT("Name"), &FNameColumn::Name)
+								.Arg(TEXT("VirtualPath"), &FVirtualPathColumn_Experimental::VirtualPath, LOCTEXT("PathNotSet", "<not set>"))
+								.Arg(TEXT("AssetPath"), &FAssetPathColumn_Experimental::Path, LOCTEXT("PathNotSet", "<not set>"))
+								.Arg(TEXT("VersePath"), &FVersePathColumn::VersePath, 
+									[](const UE::Core::FVersePath& Path) 
+									{
+										return FText::FromStringView(Path.AsStringView());
+									}, 
+									LOCTEXT("PathNotSet", "<not set>")));
+}
+
+#undef LOCTEXT_NAMESPACE
diff --git a/Engine/Plugins/Experimental/EditorDataStorageFeatures/Source/TedsAssetData/Private/Widgets/AssetLabelWidget.h b/Engine/Plugins/Experimental/EditorDataStorageFeatures/Source/TedsAssetData/Private/Widgets/AssetLabelWidget.h
new file mode 100644
--- /dev/null
+++ b/Engine/Plugins/Experimental/EditorDataStorageFeatures/Source/TedsAssetData/Private/Widgets/AssetLabelWidget.h
@@ -0,0 +1,44 @@
+// Copyright Epic Games, Inc. All Rights Reserved.
+
+#pragma once
+
+#include "Elements/Interfaces/TypedElementDataStorageFactory.h"
+#include "Elements/Interfaces/TypedElementDataStorageUiInterface.h"
+#include "Internationalization/Text.h"
+#include "UObject/ObjectMacros.h"
+
+#include "AssetLabelWidget.generated.h"
+
+class IEditorDataStorageProvider;
+class UScriptStruct;
+
+UCLASS()
+class UAssetLabelWidgetFactory : public UEditorDataStorageFactory
+{
+	GENERATED_BODY()
+
+public:
+	virtual ~UAssetLabelWidgetFactory() override = default;
+
+	TEDSASSETDATA_API virtual void RegisterWidgetConstructors(IEditorDataStorageProvider& DataStorage,
+		IEditorDataStorageUiProvider& DataStorageUi) const override;
+};
+
+// Label for assets in TEDS
+USTRUCT()
+struct FAssetLabelWidgetConstructor : public FSimpleWidgetConstructor
+{
+	GENERATED_BODY()
+
+public:
+	TEDSASSETDATA_API FAssetLabelWidgetConstructor();
+	TEDSASSETDATA_API explicit FAssetLabelWidgetConstructor(const UScriptStruct* TypeInfo);
+	virtual ~FAssetLabelWidgetConstructor() override = default;
+
+	TEDSASSETDATA_API virtual TSharedPtr<SWidget> CreateWidget(
+		IEditorDataStorageProvider* DataStorage,
+		IEditorDataStorageUiProvider* DataStorageUi,
+		UE::Editor::DataStorage::RowHandle TargetRow,
+		UE::Editor::DataStorage::RowHandle WidgetRow, 
+		const UE::Editor::DataStorage::FMetaDataView& Arguments) override;
+};
diff --git a/Engine/Plugins/Experimental/EditorDataStorageFeatures/Source/TedsAssetData/Private/Widgets/ContentBrowserListViewNameWidget.cpp b/Engine/Plugins/Experimental/EditorDataStorageFeatures/Source/TedsAssetData/Private/Widgets/ContentBrowserListViewNameWidget.cpp
new file mode 100644
--- /dev/null
+++ b/Engine/Plugins/Experimental/EditorDataStorageFeatures/Source/TedsAssetData/Private/Widgets/ContentBrowserListViewNameWidget.cpp
@@ -0,0 +1,138 @@
+// Copyright Epic Games, Inc. All Rights Reserved.
+
+#include "Widgets/ContentBrowserListViewNameWidget.h"
+
+#include "AssetThumbnail.h"
+#include "AssetDefinition.h"
+#include "Elements/Framework/TypedElementAttributeBinding.h"
+#include "TedsAssetDataColumns.h"
+#include "Elements/Columns/TypedElementFolderColumns.h"
+#include "Elements/Columns/TypedElementMiscColumns.h"
+#include "Widgets/Images/SImage.h"
+#include "Widgets/Layout/SBox.h"
+#include "Widgets/Layout/SSpacer.h"
+#include "Widgets/Text/STextBlock.h"
+#include "Widgets/SBoxPanel.h"
+
+#define LOCTEXT_NAMESPACE "FContentBrowserLabelAssetWidgetConstructor"
+
+FName UContentBrowserListViewNameWidgetFactory::WidgetPurpose(TEXT("ContentBrowser.RowLabel"));
+
+void UContentBrowserListViewNameWidgetFactory::RegisterWidgetPurposes(IEditorDataStorageUiProvider& DataStorageUi) const
+{
+	using namespace UE::Editor::DataStorage::Queries;
+	DataStorageUi.RegisterWidgetPurpose(WidgetPurpose, IEditorDataStorageUiProvider::EPurposeType::UniqueByNameAndColumn,
+		LOCTEXT("ContentBrowserLabelWidget_PurposeDescription", "Widget that display a Label + Thumbnail in the Content Browser"));
+}
+
+void UContentBrowserListViewNameWidgetFactory::RegisterWidgetConstructors(IEditorDataStorageProvider& DataStorage,
+	IEditorDataStorageUiProvider& DataStorageUi) const
+{
+	using namespace UE::Editor::DataStorage::Queries;
+	DataStorageUi.RegisterWidgetFactory<FContentBrowserListViewNameWidgetConstructor>(
+		WidgetPurpose,
+		TColumn<FNameColumn>() && (TColumn<FAssetTag>() || TColumn<FFolderTag>()));
+}
+
+FContentBrowserListViewNameWidgetConstructor::FContentBrowserListViewNameWidgetConstructor()
+	: FSimpleWidgetConstructor(StaticStruct())
+{
+}
+
+FContentBrowserListViewNameWidgetConstructor::FContentBrowserListViewNameWidgetConstructor(const UScriptStruct* TypeInfo)
+	: FSimpleWidgetConstructor(TypeInfo)
+{
+}
+
+TSharedPtr<SWidget> FContentBrowserListViewNameWidgetConstructor::CreateWidget(IEditorDataStorageProvider* DataStorage,
+	IEditorDataStorageUiProvider* DataStorageUi, UE::Editor::DataStorage::RowHandle TargetRow, UE::Editor::DataStorage::RowHandle WidgetRow,
+	const UE::Editor::DataStorage::FMetaDataView& Arguments)
+{
+	if (DataStorage->IsRowAvailable(TargetRow))
+	{
+		static FName LabelThumbnailPurposeName = TEXT("General.Cell");
+
+		TSharedPtr<FTypedElementWidgetConstructor> OutThumbnailWidgetConstructorPtr;
+		TSharedPtr<FTypedElementWidgetConstructor> OutLabelWidgetConstructorPtr;
+
+		auto AssignThumbnailWidgetToColumn = [&OutThumbnailWidgetConstructorPtr] (TUniquePtr<FTypedElementWidgetConstructor> Constructor, TConstArrayView<TWeakObjectPtr<const UScriptStruct>>)
+		{
+			OutThumbnailWidgetConstructorPtr = TSharedPtr<FTypedElementWidgetConstructor>(Constructor.Release());
+			return false;
+		};
+
+		auto AssignLabelWidgetToColumn = [&OutLabelWidgetConstructorPtr] (TUniquePtr<FTypedElementWidgetConstructor> Constructor, TConstArrayView<TWeakObjectPtr<const UScriptStruct>>)
+		{
+			OutLabelWidgetConstructorPtr = TSharedPtr<FTypedElementWidgetConstructor>(Constructor.Release());
+			return false;
+		};
+
+		// TODO: Arguments need to be developed further to customize the Thumbnail Config if needed
+		TArray<TWeakObjectPtr<const UScriptStruct>> ThumbnailColumns = GetThumbnailColumns();
+		DataStorageUi->CreateWidgetConstructors(LabelThumbnailPurposeName, IEditorDataStorageUiProvider::EMatchApproach::ExactMatch, ThumbnailColumns, Arguments, AssignThumbnailWidgetToColumn);
+
+		TArray<TWeakObjectPtr<const UScriptStruct>> LabelColumns = GetLabelColumns();
+		DataStorageUi->CreateWidgetConstructors(LabelThumbnailPurposeName, IEditorDataStorageUiProvider::EMatchApproach::ExactMatch, LabelColumns, Arguments, AssignLabelWidgetToColumn);
+
+		static const FMargin ColumnItemPadding(8, 0);
+		constexpr float ThumbnailNameHorizontalPadding = 8.f;
+
+		TSharedPtr<SWidget> ThumbnailWidget = SNullWidget::NullWidget;
+		TSharedPtr<SWidget> LabelWidget = SNullWidget::NullWidget;
+
+		if (OutThumbnailWidgetConstructorPtr)
+		{
+			ThumbnailWidget = DataStorageUi->ConstructWidget(WidgetRow, *OutThumbnailWidgetConstructorPtr, Arguments);
+		}
+
+		if (OutLabelWidgetConstructorPtr)
+		{
+			LabelWidget = DataStorageUi->ConstructWidget(WidgetRow, *OutLabelWidgetConstructorPtr, Arguments);
+		}
+
+		if (!ThumbnailWidget.IsValid())
+		{
+			ThumbnailWidget = SNullWidget::NullWidget;
+		}
+
+		if (!LabelWidget.IsValid())
+		{
+			LabelWidget = SNullWidget::NullWidget;
+		}
+
+		return SNew(SBox)
+				.Padding(ColumnItemPadding)
+				[
+					SNew(SHorizontalBox)
+					+SHorizontalBox::Slot()
+					.AutoWidth()
+					[
+						ThumbnailWidget.ToSharedRef()
+					]
+
+					+SHorizontalBox::Slot()
+					.FillWidth(1.0f)
+					.Padding(ThumbnailNameHorizontalPadding, 0.f, 0.f, 0.f)
+					.VAlign(VAlign_Center)
+					[
+						LabelWidget.ToSharedRef()
+					]
+				];
+	}
+
+	return SNullWidget::NullWidget;
+}
+
+TArray<TWeakObjectPtr<const UScriptStruct>> FContentBrowserListViewNameWidgetConstructor::GetThumbnailColumns()
+{
+	static TArray<TWeakObjectPtr<const UScriptStruct>> ThumbnailColumns({ TWeakObjectPtr(FAssetTag::StaticStruct()), TWeakObjectPtr(FFolderTag::StaticStruct()) });
+	return ThumbnailColumns;
+}
+
+TArray<TWeakObjectPtr<const UScriptStruct>> FContentBrowserListViewNameWidgetConstructor::GetLabelColumns()
+{
+	static TArray<TWeakObjectPtr<const UScriptStruct>> LabelColumns({ TWeakObjectPtr(FNameColumn::StaticStruct()) });
+	return LabelColumns;
+}
+
+#undef LOCTEXT_NAMESPACE
diff --git a/Engine/Plugins/Experimental/EditorDataStorageFeatures/Source/TedsAssetData/Private/Widgets/ContentBrowserListViewNameWidget.h b/Engine/Plugins/Experimental/EditorDataStorageFeatures/Source/TedsAssetData/Private/Widgets/ContentBrowserListViewNameWidget.h
new file mode 100644
--- /dev/null
+++ b/Engine/Plugins/Experimental/EditorDataStorageFeatures/Source/TedsAssetData/Private/Widgets/ContentBrowserListViewNameWidget.h
@@ -0,0 +1,52 @@
+// Copyright Epic Games, Inc. All Rights Reserved.
+
+#pragma once
+
+#include "Elements/Interfaces/TypedElementDataStorageFactory.h"
+#include "Elements/Interfaces/TypedElementDataStorageUiInterface.h"
+#include "Internationalization/Text.h"
+#include "UObject/ObjectMacros.h"
+
+#include "ContentBrowserListViewNameWidget.generated.h"
+
+class IEditorDataStorageProvider;
+class UScriptStruct;
+
+UCLASS()
+class UContentBrowserListViewNameWidgetFactory : public UEditorDataStorageFactory
+{
+	GENERATED_BODY()
+
+public:
+	static FName WidgetPurpose;
+
+	~UContentBrowserListViewNameWidgetFactory() override = default;
+
+	TEDSASSETDATA_API virtual void RegisterWidgetPurposes(IEditorDataStorageUiProvider& DataStorageUi) const override;
+
+	TEDSASSETDATA_API void RegisterWidgetConstructors(IEditorDataStorageProvider& DataStorage,
+		IEditorDataStorageUiProvider& DataStorageUi) const override;
+};
+
+// Content Browser Label + Thumbnail widget
+USTRUCT()
+struct FContentBrowserListViewNameWidgetConstructor : public FSimpleWidgetConstructor
+{
+	GENERATED_BODY()
+
+public:
+	TEDSASSETDATA_API FContentBrowserListViewNameWidgetConstructor();
+	TEDSASSETDATA_API explicit FContentBrowserListViewNameWidgetConstructor(const UScriptStruct* TypeInfo);
+	~FContentBrowserListViewNameWidgetConstructor() override = default;
+
+	TEDSASSETDATA_API virtual TSharedPtr<SWidget> CreateWidget(
+		IEditorDataStorageProvider* DataStorage,
+		IEditorDataStorageUiProvider* DataStorageUi,
+		UE::Editor::DataStorage::RowHandle TargetRow,
+		UE::Editor::DataStorage::RowHandle WidgetRow, 
+		const UE::Editor::DataStorage::FMetaDataView& Arguments) override;
+
+private:
+	TArray<TWeakObjectPtr<const UScriptStruct>> GetThumbnailColumns();
+	TArray<TWeakObjectPtr<const UScriptStruct>> GetLabelColumns();
+};
diff --git a/Engine/Plugins/Experimental/EditorDataStorageFeatures/Source/TedsAssetData/Public/Widgets/AssetDataLabelWidget.h b/Engine/Plugins/Experimental/EditorDataStorageFeatures/Source/TedsAssetData/Public/Widgets/AssetDataLabelWidget.h
--- a/Engine/Plugins/Experimental/EditorDataStorageFeatures/Source/TedsAssetData/Public/Widgets/AssetDataLabelWidget.h
+++ b/Engine/Plugins/Experimental/EditorDataStorageFeatures/Source/TedsAssetData/Public/Widgets/AssetDataLabelWidget.h
@@ -41,4 +41,7 @@ public:
 		UE::Editor::DataStorage::RowHandle TargetRow,
 		UE::Editor::DataStorage::RowHandle WidgetRow, 
 		const UE::Editor::DataStorage::FMetaDataView& Arguments) override;
+
+private:
+	TArray<TWeakObjectPtr<const UScriptStruct>> GetLabelColumns();
 };
\ No newline at end of file
diff --git a/Engine/Plugins/Experimental/EditorDataStorageFeatures/Source/TedsContentBrowser/Private/TedsContentBrowserModule.cpp b/Engine/Plugins/Experimental/EditorDataStorageFeatures/Source/TedsContentBrowser/Private/TedsContentBrowserModule.cpp
--- a/Engine/Plugins/Experimental/EditorDataStorageFeatures/Source/TedsContentBrowser/Private/TedsContentBrowserModule.cpp
+++ b/Engine/Plugins/Experimental/EditorDataStorageFeatures/Source/TedsContentBrowser/Private/TedsContentBrowserModule.cpp
@@ -14,6 +14,7 @@
 #include "QueryStack/FQueryStackNode_RowView.h"
 #include "TedsAlertColumns.h"
 #include "TedsAssetDataColumns.h"
+#include "Elements/Columns/TypedElementFolderColumns.h"
 #include "Elements/Columns/TypedElementMiscColumns.h"
 #include "Widgets/STedsTableViewer.h"
 #include "Widgets/Views/SListView.h"
@@ -147,10 +148,10 @@ namespace UE::Editor::ContentBrowser
 		// Create the table viewer widget
 		TableViewer = SNew(STedsTableViewer)
 					.QueryStack(RowQueryStack)
-					.CellWidgetPurposes({TEXT("General.RowLabel"), TEXT("General.Cell")})
+					.CellWidgetPurposes({TEXT("ContentBrowser.RowLabel"), TEXT("General.RowLabel"), TEXT("General.Cell")})
 					// Default list of columns to display
 					.Columns({ FNameColumn::StaticStruct(), FTedsAlertColumn::StaticStruct(),
-						FAssetClassColumn::StaticStruct(), FAssetTag::StaticStruct(), FAssetPathColumn_Experimental::StaticStruct(),
+						FAssetClassColumn::StaticStruct(), FAssetTag::StaticStruct(), FFolderTag::StaticStruct(), FAssetPathColumn_Experimental::StaticStruct(),
 						FDiskSizeColumn::StaticStruct(), FVirtualPathColumn_Experimental::StaticStruct(), DynamicSkeletalMeshSkeletonColumn })
 					.ListSelectionMode(ESelectionMode::Multi)
 					.OnSelectionChanged_Lambda([this](RowHandle Row)
