[TEDS Outliner] Fix crash that might happen when the Chaos Visual Debugger is being closed.

The issue is not in CVD itself, but on the TEDS Outliner.

The teds outliner mode binds to a selection change delegate on the Teds outliner impl object. That delegate is triggered in a tick if the selection is dirty.

The problem is that the delegate is never unbound if the outliner mode is destroyed. Therefore if the tick happens after the outliner mode was destroyed but before the mode impl object was destroyed, and the selection state is dirty, a crash will happen.

CVD clears the selection as part of its shutdown process, which makes the issue more likely to happen.

The fix is to just unbind that delegate when the outliner mode is destroyed.

diff --git a/Engine/Plugins/Experimental/EditorDataStorageFeatures/Source/TedsOutliner/Private/TedsOutlinerMode.cpp b/Engine/Plugins/Experimental/EditorDataStorageFeatures/Source/TedsOutliner/Private/TedsOutlinerMode.cpp
--- a/Engine/Plugins/Experimental/EditorDataStorageFeatures/Source/TedsOutliner/Private/TedsOutlinerMode.cpp
+++ b/Engine/Plugins/Experimental/EditorDataStorageFeatures/Source/TedsOutliner/Private/TedsOutlinerMode.cpp
@@ -42,7 +42,10 @@ FTedsOutlinerMode::FTedsOutlinerMode(const FTedsOutlinerParams& InParams)
 
 FTedsOutlinerMode::~FTedsOutlinerMode()
 {
-	
+	if (TedsOutlinerImpl)
+	{
+		TedsOutlinerImpl->OnSelectionChanged().RemoveAll(this);
+	}
 }
 
 void FTedsOutlinerMode::Rebuild()
