TEDS UI:
Add a new STedsTileViewer widget to display TEDS rows as tiles
Add a cvar to use STedsTileViewer in the TEDS-CB integration (TEDS.UI.TedsContentBrowserViewMode)
Create a new ITableViewer base class for all table viewer widgets

diff --git a/Engine/Plugins/Experimental/EditorDataStorageFeatures/Source/TedsContentBrowser/Private/TedsContentBrowserModule.cpp b/Engine/Plugins/Experimental/EditorDataStorageFeatures/Source/TedsContentBrowser/Private/TedsContentBrowserModule.cpp
--- a/Engine/Plugins/Experimental/EditorDataStorageFeatures/Source/TedsContentBrowser/Private/TedsContentBrowserModule.cpp
+++ b/Engine/Plugins/Experimental/EditorDataStorageFeatures/Source/TedsContentBrowser/Private/TedsContentBrowserModule.cpp
@@ -17,6 +17,7 @@
 #include "Elements/Columns/TypedElementFolderColumns.h"
 #include "Elements/Columns/TypedElementMiscColumns.h"
 #include "Widgets/STedsTableViewer.h"
+#include "Widgets/STedsTileViewer.h"
 #include "Widgets/Views/SListView.h"
 
 #define LOCTEXT_NAMESPACE "TedsContentBrowserModule"
@@ -47,6 +48,18 @@ namespace UE::Editor::ContentBrowser
 			}
 		}));
 
+	enum class ETableViewerMode : int32
+	{
+		List = 0,
+		Tile = 1
+	};
+
+	static int32 TableViewerMode = static_cast<int>(ETableViewerMode::List);
+	static FAutoConsoleVariableRef CvarSetTableViewMode(
+		TEXT("TEDS.UI.TedsContentBrowserViewMode"),
+		TableViewerMode,
+		TEXT("Set the view mode of the TEDS-CB. 0 = List View, 1 = Tile View. (requires re-opening any currently open content browsers)")); 
+
 	TSharedPtr<IContentBrowserViewExtender> FTedsContentBrowserModule::CreateContentBrowserViewExtender()
 	{
 		return MakeShared<FTedsContentBrowserViewExtender>();
@@ -131,12 +144,8 @@ namespace UE::Editor::ContentBrowser
 		return RowHandle;
 	}
 
-	FTedsContentBrowserViewExtender::FTedsContentBrowserViewExtender()
+	void FTedsContentBrowserViewExtender::CreateListView()
 	{
-		DataStorage = GetMutableDataStorageFeature<IEditorDataStorageProvider>(StorageFeatureName);
-		
-		RowQueryStack = MakeShared<FQueryStackNode_RowView>(&Rows);
-
 		// Sample dynamic column to display the "Skeleton" attribute on skeletal meshes
 		// We probably want the dynamic columns in the table viewer to be data driven based on the rows in the future
 		const UScriptStruct* DynamicSkeletalMeshSkeletonColumn = DataStorage->GenerateDynamicColumn(FDynamicColumnDescription
@@ -162,33 +171,62 @@ namespace UE::Editor::ContentBrowser
 							OnSelectionChangedDelegate.Execute(AssetViewItem, ESelectInfo::Direct);
 						}
 					});
+		
+		BindViewColumns();
+	}
+
+	void FTedsContentBrowserViewExtender::CreateTileView()
+	{
+		// Create the table viewer widget
+		TableViewer = SNew(STedsTileViewer)
+					.QueryStack(RowQueryStack)
+					.WidgetPurpose(TEXT("ContentBrowser.RowLabel"))
+					.SelectionMode(ESelectionMode::Multi)
+					.OnSelectionChanged_Lambda([this](RowHandle Row)
+					{
+						if(TSharedPtr<FAssetViewItem> AssetViewItem = GetAssetViewItemFromRow(Row))
+						{
+							// CB 2.0 TODO: Does the CB use ESelectInfo and we need to propagate it from the table viewer?
+							OnSelectionChangedDelegate.Execute(AssetViewItem, ESelectInfo::Direct);
+						}
+					});
+		
+		BindViewColumns();
+	}
 
+	void FTedsContentBrowserViewExtender::BindViewColumns()
+	{
+		if (!TableViewer)
+		{
+			return;
+		}
+		
 		// Bind the delegates the CB view extender requires to delegates in TEDS columns on the widget row that are fired
 		// when the event occurs
 		const RowHandle WidgetRow = TableViewer->GetWidgetRowHandle();
 
-		if(FWidgetContextMenuColumn* ContextMenuColumn = DataStorage->GetColumn<FWidgetContextMenuColumn>(WidgetRow))
+		if (FWidgetContextMenuColumn* ContextMenuColumn = DataStorage->GetColumn<FWidgetContextMenuColumn>(WidgetRow))
 		{
 			ContextMenuColumn->OnContextMenuOpening.BindLambda([this]()
 			{
 				return OnContextMenuOpenedDelegate.Execute();
 			});
 		}
-		if(FWidgetRowScrolledIntoView* ScrolledIntoViewColumn = DataStorage->GetColumn<FWidgetRowScrolledIntoView>(WidgetRow))
+		if (FWidgetRowScrolledIntoView* ScrolledIntoViewColumn = DataStorage->GetColumn<FWidgetRowScrolledIntoView>(WidgetRow))
 		{
 			ScrolledIntoViewColumn->OnItemScrolledIntoView.BindLambda([this](FTedsRowHandle Row, const TSharedPtr<ITableRow>& TableRow)
 			{
-				if(TSharedPtr<FAssetViewItem> AssetViewItem = GetAssetViewItemFromRow(Row))
+				if (TSharedPtr<FAssetViewItem> AssetViewItem = GetAssetViewItemFromRow(Row))
 				{
 					return OnItemScrolledIntoViewDelegate.Execute(AssetViewItem, TableRow);
 				}
 			});
 		}
-		if(FWidgetDoubleClickedColumn* DoubleClickedColumn = DataStorage->GetColumn<FWidgetDoubleClickedColumn>(WidgetRow))
+		if (FWidgetDoubleClickedColumn* DoubleClickedColumn = DataStorage->GetColumn<FWidgetDoubleClickedColumn>(WidgetRow))
 		{
 			DoubleClickedColumn->OnMouseButtonDoubleClick.BindLambda([this](FTedsRowHandle Row)
 			{
-				if(TSharedPtr<FAssetViewItem> AssetViewItem = GetAssetViewItemFromRow(Row))
+				if (TSharedPtr<FAssetViewItem> AssetViewItem = GetAssetViewItemFromRow(Row))
 				{
 					return OnItemDoubleClickedDelegate.Execute(AssetViewItem);
 				}
@@ -196,10 +234,32 @@ namespace UE::Editor::ContentBrowser
 		}
 	}
 
+	FTedsContentBrowserViewExtender::FTedsContentBrowserViewExtender()
+	{
+		DataStorage = GetMutableDataStorageFeature<IEditorDataStorageProvider>(StorageFeatureName);
+		
+		RowQueryStack = MakeShared<FQueryStackNode_RowView>(&Rows);
+	}
+
 	TSharedRef<SWidget> FTedsContentBrowserViewExtender::CreateView(TArray<TSharedPtr<FAssetViewItem>>* InItemsSource)
 	{
+		switch (static_cast<ETableViewerMode>(TableViewerMode))
+		{
+		case ETableViewerMode::List:
+			CreateListView();
+			break;
+
+		case ETableViewerMode::Tile:
+			CreateTileView();
+			break;
+
+		default:
+			CreateListView();
+		}
+		
 		RefreshRows(InItemsSource);
-		return TableViewer.ToSharedRef();
+		
+		return TableViewer->AsWidget();
 	}
 	
 	void FTedsContentBrowserViewExtender::OnItemListChanged(TArray<TSharedPtr<FAssetViewItem>>* InItemsSource)
@@ -247,18 +307,18 @@ namespace UE::Editor::ContentBrowser
 
 	FText FTedsContentBrowserViewExtender::GetViewDisplayName()
 	{
-		return LOCTEXT("TedsCBViewName", "TEDS List View");
+		return LOCTEXT("TedsCBViewName", "TEDS Table View");
 	}
 
 	FText FTedsContentBrowserViewExtender::GetViewTooltipText()
 	{
-		return LOCTEXT("TedsCBViewTooltip", "A List view populated using TEDS UI and the asset registry data in TEDS");
+		return LOCTEXT("TedsCBViewTooltip", "A table viewer populated using TEDS UI and the asset registry data in TEDS");
 	}
 
 	void FTedsContentBrowserViewExtender::FocusList()
 	{
 		// CB 2.0 TODO: Do we need to focus the internal list? If so, implement using a Teds column
-		FSlateApplication::Get().SetKeyboardFocus(TableViewer, EFocusCause::SetDirectly);
+		FSlateApplication::Get().SetKeyboardFocus(TableViewer->AsWidget(), EFocusCause::SetDirectly);
 	}
 
 	void FTedsContentBrowserViewExtender::SetSelection(const TSharedPtr<FAssetViewItem>& Item, bool bSelected, const ESelectInfo::Type SelectInfo)
diff --git a/Engine/Plugins/Experimental/EditorDataStorageFeatures/Source/TedsContentBrowser/Private/TedsContentBrowserModule.h b/Engine/Plugins/Experimental/EditorDataStorageFeatures/Source/TedsContentBrowser/Private/TedsContentBrowserModule.h
--- a/Engine/Plugins/Experimental/EditorDataStorageFeatures/Source/TedsContentBrowser/Private/TedsContentBrowserModule.h
+++ b/Engine/Plugins/Experimental/EditorDataStorageFeatures/Source/TedsContentBrowser/Private/TedsContentBrowserModule.h
@@ -12,13 +12,14 @@ class IEditorDataStorageProvider;
 namespace UE::Editor::DataStorage
 {
 	class FQueryStackNode_RowView;
-	class STedsTableViewer;
+	class ITableViewer;
 }
 
 class FAssetViewItem;
 
 namespace UE::Editor::ContentBrowser
 {
+	
 	// A custom view for the content browser that uses the TEDS asset registry integration to display rows with widgets using TEDS UI
 	class FTedsContentBrowserViewExtender : public IContentBrowserViewExtender
 	{
@@ -43,7 +44,6 @@ namespace UE::Editor::ContentBrowser
 		virtual void OnItemListChanged(TArray<TSharedPtr<FAssetViewItem>>* InItemsSource) override;
 		// ~IContentBrowserViewExtender interface
 
-
 		// Refresh the rows in the current view by syncing the the items source
 		void RefreshRows(TArray<TSharedPtr<FAssetViewItem>>* InItemsSource);
 
@@ -55,13 +55,24 @@ namespace UE::Editor::ContentBrowser
 
 		DataStorage::RowHandle GetRowFromAssetViewItem(const TSharedPtr<FAssetViewItem>& Item);
 
+	protected:
+
+		// Update the table viewer used by the integration to use a list view
+		void CreateListView();
+
+		// Update the table viewer used by the integration to use a tile view
+		void CreateTileView();
+
+		// Bind the table viewer's columns to CB delegates
+		void BindViewColumns();
+
 	private:
 
 		// Ptr to the data storage interface
 		IEditorDataStorageProvider* DataStorage;
 		
 		// The actual table viewer widget
-		TSharedPtr<DataStorage::STedsTableViewer> TableViewer;
+		TSharedPtr<DataStorage::ITableViewer> TableViewer;
 		
 		// Query stack used by the table viewer
 		TSharedPtr<DataStorage::FQueryStackNode_RowView> RowQueryStack;
diff --git a/Engine/Plugins/Experimental/EditorDataStorageFeatures/Source/TedsTableViewer/Private/Widgets/STedsTableViewer.cpp b/Engine/Plugins/Experimental/EditorDataStorageFeatures/Source/TedsTableViewer/Private/Widgets/STedsTableViewer.cpp
--- a/Engine/Plugins/Experimental/EditorDataStorageFeatures/Source/TedsTableViewer/Private/Widgets/STedsTableViewer.cpp
+++ b/Engine/Plugins/Experimental/EditorDataStorageFeatures/Source/TedsTableViewer/Private/Widgets/STedsTableViewer.cpp
@@ -26,8 +26,12 @@ namespace UE::Editor::DataStorage
 		HeaderRowWidget = SNew( SHeaderRow )
 							.CanSelectGeneratedColumn(true);
 
-		// Create the internal Teds Widget
-		CreateTedsWidget();
+		ChildSlot
+		[
+			SAssignNew(TedsWidget, STedsWidget)
+		];
+		
+		AddWidgetColumns();
 
 		// Attribute binder to bind widget columns to attributes on the ListView
 		FAttributeBinder Binder(TedsWidget->GetRowHandle());
@@ -42,7 +46,7 @@ namespace UE::Editor::DataStorage
 			.OnItemScrolledIntoView(Binder.BindEvent(&FWidgetRowScrolledIntoView::OnItemScrolledIntoView))
 			.OnMouseButtonDoubleClick(Binder.BindEvent(&FWidgetDoubleClickedColumn::OnMouseButtonDoubleClick));
 
-		AssignChildSlot();
+		CreateInternalWidget();
 		
 		// Add each Teds column from the model to our header row widget
 		Model->ForEachColumn([this](const TSharedRef<FTedsTableViewerColumn>& Column)
@@ -54,14 +58,12 @@ namespace UE::Editor::DataStorage
 		Model->GetOnModelChanged().AddLambda([this]()
 		{
 			ListView->RequestListRefresh();
-			AssignChildSlot();
+			CreateInternalWidget();
 		});
 	}
 
-	void STedsTableViewer::CreateTedsWidget()
+	void STedsTableViewer::AddWidgetColumns()
 	{
-		TedsWidget = SNew(STedsWidget);
-
 		if(IEditorDataStorageProvider* DataStorage = Model->GetDataStorageInterface())
 		{
 			const RowHandle WidgetRowHandle = TedsWidget->GetRowHandle();
@@ -77,14 +79,9 @@ namespace UE::Editor::DataStorage
 				DataStorage->AddColumn(WidgetRowHandle, FWidgetDoubleClickedColumn::StaticStruct());
 			}
 		}
-
-		ChildSlot
-		[
-			TedsWidget.ToSharedRef()
-		];
 	}
 
-	void STedsTableViewer::AssignChildSlot()
+	void STedsTableViewer::CreateInternalWidget()
 	{
 		TSharedPtr<SWidget> ContentWidget;
 
@@ -127,7 +124,7 @@ namespace UE::Editor::DataStorage
 			HeaderRowWidget->AddColumn(Column->ConstructHeaderRowColumn());
 		});
 
-		AssignChildSlot();
+		CreateInternalWidget();
 	}
 
 	void STedsTableViewer::OnListSelectionChanged(TableViewerItemPtr Item, ESelectInfo::Type SelectInfo)
@@ -183,6 +180,11 @@ namespace UE::Editor::DataStorage
 		ListView->ClearSelection();
 	}
 
+	TSharedRef<SWidget> STedsTableViewer::AsWidget()
+	{
+		return AsShared();
+	}
+
 	bool STedsTableViewer::IsItemVisible(TableViewerItemPtr InItem) const
 	{
 		return ListView->IsItemVisible(InItem);
diff --git a/Engine/Plugins/Experimental/EditorDataStorageFeatures/Source/TedsTableViewer/Private/Widgets/STedsTileViewer.cpp b/Engine/Plugins/Experimental/EditorDataStorageFeatures/Source/TedsTableViewer/Private/Widgets/STedsTileViewer.cpp
new file mode 100644
--- /dev/null
+++ b/Engine/Plugins/Experimental/EditorDataStorageFeatures/Source/TedsTableViewer/Private/Widgets/STedsTileViewer.cpp
@@ -0,0 +1,166 @@
+// Copyright Epic Games, Inc. All Rights Reserved.
+
+#include "Widgets/STedsTileViewer.h"
+#include "TedsTableViewerColumn.h"
+#include "Columns/SlateDelegateColumns.h"
+#include "Elements/Columns/TypedElementUIColumns.h"
+#include "Elements/Common/EditorDataStorageFeatures.h"
+#include "Elements/Framework/TypedElementAttributeBinding.h"
+#include "Elements/Framework/TypedElementDataStorageWidget.h"
+#include "Widgets/STedsTableViewerRow.h"
+
+namespace UE::Editor::DataStorage
+{
+	void STedsTileViewer::Construct(const FArguments& InArgs)
+	{
+		OnSelectionChanged = InArgs._OnSelectionChanged;
+		WidgetPurpose = InArgs._WidgetPurpose;
+
+		TArray<TWeakObjectPtr<const UScriptStruct>> Columns;
+		TArray<FName> Purposes;
+		
+		Model = MakeShared<FTedsTableViewerModel>(InArgs._QueryStack, Columns, Purposes,
+			FTedsTableViewerModel::FIsItemVisible::CreateSP(this, &STedsTileViewer::IsItemVisible));
+
+		ChildSlot
+		[
+			SAssignNew(TedsWidget, STedsWidget)
+		];
+		
+		AddWidgetColumns();
+		
+		// Attribute binder to bind widget columns to attributes on the ListView
+		FAttributeBinder Binder(TedsWidget->GetRowHandle());
+
+		TileView = SNew(STileView<TableViewerItemPtr>)
+			.ListItemsSource(&Model->GetItems())
+			.OnGenerateTile(this, &STedsTileViewer::MakeTileWidget)
+			.OnSelectionChanged(this, &STedsTileViewer::OnListSelectionChanged)
+			.SelectionMode(InArgs._SelectionMode)
+			.OnContextMenuOpening(Binder.BindEvent(&FWidgetContextMenuColumn::OnContextMenuOpening))
+			.OnItemScrolledIntoView(Binder.BindEvent(&FWidgetRowScrolledIntoView::OnItemScrolledIntoView))
+			.OnMouseButtonDoubleClick(Binder.BindEvent(&FWidgetDoubleClickedColumn::OnMouseButtonDoubleClick));
+
+		TedsWidget->SetContent(TileView.ToSharedRef());
+		
+		// Whenever the model changes, refresh the list to update the UI
+		Model->GetOnModelChanged().AddLambda([this]()
+		{
+			TileView->RequestListRefresh();
+		});
+
+		CreateTileWidgetConstructor();
+	}
+
+	void STedsTileViewer::CreateTileWidgetConstructor()
+	{
+		// We're going to use a custom column to represent our tile so we can use FTedsTableViewerModel with STileView which doesn't have the concept
+		// of columns
+		TSharedPtr<FTedsTableViewerColumn> Column;
+		
+		auto AssignWidgetToColumn = [&Column](TUniquePtr<FTypedElementWidgetConstructor> Constructor, TConstArrayView<TWeakObjectPtr<const UScriptStruct>>)
+		{
+			TSharedPtr<FTypedElementWidgetConstructor> WidgetConstructor(Constructor.Release());
+			Column = MakeShared<FTedsTableViewerColumn>(TEXT("TileView"), WidgetConstructor);
+			return false;
+		};
+	
+		IEditorDataStorageUiProvider* StorageUi = GetMutableDataStorageFeature<IEditorDataStorageUiProvider>(UiFeatureName);
+
+		StorageUi->CreateWidgetConstructors(WidgetPurpose, FMetaDataView(), AssignWidgetToColumn);
+
+		if (Column)
+		{
+			Model->AddCustomColumn(Column.ToSharedRef());
+		}
+	}
+
+	void STedsTileViewer::ForEachSelectedRow(TFunctionRef<void(RowHandle)> InCallback) const
+	{
+		TArray<TableViewerItemPtr> SelectedRows;
+		TileView->GetSelectedItems(SelectedRows);
+
+		for(TableViewerItemPtr& Row : SelectedRows)
+		{
+			InCallback(Row);
+		}
+	}
+
+	RowHandle STedsTileViewer::GetWidgetRowHandle() const
+	{
+		return TedsWidget->GetRowHandle();
+	}
+
+	void STedsTileViewer::SetSelection(RowHandle Row, bool bSelected, const ESelectInfo::Type SelectInfo) const
+	{
+		FTedsRowHandle TedsRowHandle{ .RowHandle = Row };
+		TileView->SetItemSelection(TedsRowHandle, bSelected, SelectInfo);
+
+	}
+
+	void STedsTileViewer::ScrollIntoView(RowHandle Row) const
+	{
+		FTedsRowHandle TedsRowHandle{ .RowHandle = Row };
+		TileView->RequestScrollIntoView(TedsRowHandle);
+	}
+
+	void STedsTileViewer::ClearSelection() const
+	{
+		TileView->ClearSelection();
+	}
+
+	TSharedRef<SWidget> STedsTileViewer::AsWidget()
+	{
+		return AsShared();
+	}
+
+	TSharedRef<ITableRow> STedsTileViewer::MakeTileWidget(TableViewerItemPtr InItem,
+	                                                      const TSharedRef<STableViewBase>& OwnerTable) const
+	{
+		TSharedPtr<SWidget> RowWidget = SNullWidget::NullWidget;
+		
+		if(TSharedPtr<FTedsTableViewerColumn> Column = Model->GetColumn("TileView"))
+		{
+			RowWidget = Column->ConstructRowWidget(InItem);
+		}
+
+		return SNew(STableRow<TableViewerItemPtr>, OwnerTable)
+			[
+				RowWidget.ToSharedRef()
+			];
+	}
+
+	bool STedsTileViewer::IsItemVisible(TableViewerItemPtr InItem) const
+	{
+		return TileView->IsItemVisible(InItem);
+	}
+
+	void STedsTileViewer::OnListSelectionChanged(TableViewerItemPtr Item, ESelectInfo::Type SelectInfo)
+	{
+		if(OnSelectionChanged.IsBound())
+		{
+			OnSelectionChanged.Execute(Item);
+		}
+	}
+
+	void STedsTileViewer::AddWidgetColumns()
+	{
+		if(IEditorDataStorageProvider* DataStorage = Model->GetDataStorageInterface())
+		{
+			const RowHandle WidgetRowHandle = TedsWidget->GetRowHandle();
+		
+			if(DataStorage->IsRowAvailable(WidgetRowHandle))
+			{
+				// The table viewer should not show up as a row in a table viewer because that will cause all sorts of recursion issues
+				DataStorage->AddColumn(WidgetRowHandle, FHideRowFromUITag::StaticStruct());
+
+				// Columns we are going to bind to attributes on SListView
+				DataStorage->AddColumn(WidgetRowHandle, FWidgetContextMenuColumn::StaticStruct());
+				DataStorage->AddColumn(WidgetRowHandle, FWidgetRowScrolledIntoView::StaticStruct());
+				DataStorage->AddColumn(WidgetRowHandle, FWidgetDoubleClickedColumn::StaticStruct());
+			}
+		}
+	}
+}
+
+
diff --git a/Engine/Plugins/Experimental/EditorDataStorageFeatures/Source/TedsTableViewer/Public/ITedsTableViewer.h b/Engine/Plugins/Experimental/EditorDataStorageFeatures/Source/TedsTableViewer/Public/ITedsTableViewer.h
new file mode 100644
--- /dev/null
+++ b/Engine/Plugins/Experimental/EditorDataStorageFeatures/Source/TedsTableViewer/Public/ITedsTableViewer.h
@@ -0,0 +1,36 @@
+// Copyright Epic Games, Inc. All Rights Reserved.
+#pragma once
+
+#include "Elements/Common/TypedElementHandles.h"
+#include "Widgets/SWidget.h"
+
+namespace ESelectInfo { enum Type : int; }
+
+namespace UE::Editor::DataStorage
+{
+	// Interface for all types of TEDS table viewers
+    class ITableViewer
+    {
+    public:
+
+    	// Execute the given callback for each row handle selected in the table viewer
+        virtual void ForEachSelectedRow(TFunctionRef<void(RowHandle)> InCallback) const = 0;
+
+    	// Get the row handle for the widget row that contains the contents of the table viewer itself
+        virtual RowHandle GetWidgetRowHandle() const = 0;
+
+    	// Set the currently selected row in the table viewer
+        virtual void SetSelection(RowHandle Row, bool bSelected, const ESelectInfo::Type SelectInfo) const = 0;
+
+    	// Scroll the input row into view in the table viewer
+        virtual void ScrollIntoView(RowHandle Row) const = 0;
+
+    	// Deselect all items in the table viewer
+        virtual void ClearSelection() const = 0;
+
+    	// Get the table viewer as an SWidget
+    	virtual TSharedRef<SWidget> AsWidget() = 0;
+
+    	virtual ~ITableViewer() = default;
+    };
+}
\ No newline at end of file
diff --git a/Engine/Plugins/Experimental/EditorDataStorageFeatures/Source/TedsTableViewer/Public/Widgets/STedsTableViewer.h b/Engine/Plugins/Experimental/EditorDataStorageFeatures/Source/TedsTableViewer/Public/Widgets/STedsTableViewer.h
--- a/Engine/Plugins/Experimental/EditorDataStorageFeatures/Source/TedsTableViewer/Public/Widgets/STedsTableViewer.h
+++ b/Engine/Plugins/Experimental/EditorDataStorageFeatures/Source/TedsTableViewer/Public/Widgets/STedsTableViewer.h
@@ -1,6 +1,7 @@
 // Copyright Epic Games, Inc. All Rights Reserved.
 #pragma once
 
+#include "ITedsTableViewer.h"
 #include "Elements/Common/TypedElementHandles.h"
 #include "QueryStack/IQueryStackNode_Row.h"
 #include "TedsTableViewerModel.h"
@@ -26,7 +27,7 @@ namespace UE::Editor::DataStorage
      *		.QueryStack(MakeShared<UE::Editor::DataStorage::FQueryStackNode_RowView>(&Rows))
 	 *		.Columns({FTypedElementLabelColumn::StaticStruct(), FTypedElementClassTypeInfoColumn::StaticStruct());
 	 */
-	class STedsTableViewer : public SCompoundWidget
+	class STedsTableViewer : public SCompoundWidget, public ITableViewer
 	{
 	public:
 		
@@ -74,18 +75,20 @@ namespace UE::Editor::DataStorage
 		TEDSTABLEVIEWER_API void AddCustomColumn(const TSharedRef<FTedsTableViewerColumn>& InColumn);
 
 		// Execute the given callback for each row that is selected in the table viewer
-		TEDSTABLEVIEWER_API void ForEachSelectedRow(TFunctionRef<void(RowHandle)> InCallback) const;
+		TEDSTABLEVIEWER_API virtual void ForEachSelectedRow(TFunctionRef<void(RowHandle)> InCallback) const override;
 
 		// Get the row handle for the widget row the table viewer's contents are stored in
-		TEDSTABLEVIEWER_API RowHandle GetWidgetRowHandle() const;
+		TEDSTABLEVIEWER_API virtual RowHandle GetWidgetRowHandle() const override;
 		
 		// Select the given row in the table viewer
-		TEDSTABLEVIEWER_API void SetSelection(RowHandle Row, bool bSelected, const ESelectInfo::Type SelectInfo) const;
+		TEDSTABLEVIEWER_API virtual void SetSelection(RowHandle Row, bool bSelected, const ESelectInfo::Type SelectInfo) const override;
 
 		// Scroll the given row into view in the table viewer
-		TEDSTABLEVIEWER_API void ScrollIntoView(RowHandle Row) const;
+		TEDSTABLEVIEWER_API virtual void ScrollIntoView(RowHandle Row) const override;
 
-		TEDSTABLEVIEWER_API void ClearSelection() const;
+		TEDSTABLEVIEWER_API virtual void ClearSelection() const override;
+		
+		TEDSTABLEVIEWER_API virtual TSharedRef<SWidget> AsWidget() override;
 
 	protected:
 		
@@ -93,13 +96,13 @@ namespace UE::Editor::DataStorage
 
 		bool IsItemVisible(TableViewerItemPtr InItem) const;
 
-		void AssignChildSlot();
+		void CreateInternalWidget();
 
 		void RefreshColumnWidgets();
 
 		void OnListSelectionChanged(TableViewerItemPtr Item, ESelectInfo::Type SelectInfo);
 
-		void CreateTedsWidget();
+		void AddWidgetColumns();
 
 	private:
 
diff --git a/Engine/Plugins/Experimental/EditorDataStorageFeatures/Source/TedsTableViewer/Public/Widgets/STedsTileViewer.h b/Engine/Plugins/Experimental/EditorDataStorageFeatures/Source/TedsTableViewer/Public/Widgets/STedsTileViewer.h
new file mode 100644
--- /dev/null
+++ b/Engine/Plugins/Experimental/EditorDataStorageFeatures/Source/TedsTableViewer/Public/Widgets/STedsTileViewer.h
@@ -0,0 +1,102 @@
+// Copyright Epic Games, Inc. All Rights Reserved.
+#pragma once
+
+#include "ITedsTableViewer.h"
+#include "Elements/Common/TypedElementHandles.h"
+#include "QueryStack/IQueryStackNode_Row.h"
+#include "TedsTableViewerModel.h"
+#include "Widgets/SCompoundWidget.h"
+#include "Widgets/Views/STableRow.h"
+#include "Widgets/Views/STileView.h"
+
+class ITableRow;
+class STableViewBase;
+class SHeaderRow;
+class STedsWidget;
+
+namespace UE::Editor::DataStorage
+{
+	/* Table Viewer where each row is represented as a tile using STileView
+	 * Unlike STedsTableViewer which takes a list of columns to display the tile viewer simply uses a single widget purpose
+	 * to create tiles for each row.
+	 */
+	class STedsTileViewer : public SCompoundWidget, public ITableViewer
+	{
+	public:
+		
+		// Delegate fired when the selection in the table viewer changes
+		DECLARE_DELEGATE_OneParam(FOnSelectionChanged, RowHandle)
+
+		SLATE_BEGIN_ARGS(STedsTileViewer)
+			: _WidgetPurpose(TEXT("General.Cell")) // TODO: Have a better "default" widget for tiles
+			, _SelectionMode(ESelectionMode::Type::Single)
+		{
+			
+		}
+
+		// Query Stack that will supply the rows to be displayed
+		SLATE_ARGUMENT(TSharedPtr<IQueryStackNode_Row>, QueryStack)
+
+		// The widget purpose to use to create the widgets
+		SLATE_ARGUMENT(FName, WidgetPurpose)
+		
+		// Delegate called when the selection changes
+		SLATE_EVENT(FOnSelectionChanged, OnSelectionChanged)
+			
+		// The selection mode for the table viewer (single/multi etc)
+		SLATE_ARGUMENT(ESelectionMode::Type, SelectionMode)
+
+		SLATE_END_ARGS()
+
+	public:
+		
+		TEDSTABLEVIEWER_API void Construct(const FArguments& InArgs);
+
+		// Execute the given callback for each row that is selected in the table viewer
+		TEDSTABLEVIEWER_API virtual void ForEachSelectedRow(TFunctionRef<void(RowHandle)> InCallback) const override;
+
+		// Get the row handle for the widget row the table viewer's contents are stored in
+		TEDSTABLEVIEWER_API virtual RowHandle GetWidgetRowHandle() const override;
+		
+		// Select the given row in the table viewer
+		TEDSTABLEVIEWER_API virtual void SetSelection(RowHandle Row, bool bSelected, const ESelectInfo::Type SelectInfo) const override;
+
+		// Scroll the given row into view in the table viewer
+		TEDSTABLEVIEWER_API virtual void ScrollIntoView(RowHandle Row) const override;
+
+		// Deselect all items in the table viewer
+		TEDSTABLEVIEWER_API virtual void ClearSelection() const override;
+
+		// Get the table viewer as an SWidget
+		TEDSTABLEVIEWER_API virtual TSharedRef<SWidget> AsWidget() override;
+
+	protected:
+		
+		TSharedRef<ITableRow> MakeTileWidget(TableViewerItemPtr InItem, const TSharedRef<STableViewBase>& OwnerTable) const;
+
+		bool IsItemVisible(TableViewerItemPtr InItem) const;
+
+		void OnListSelectionChanged(TableViewerItemPtr Item, ESelectInfo::Type SelectInfo);
+
+		void AddWidgetColumns();
+		
+		void CreateTileWidgetConstructor();
+
+	private:
+
+		// The actual ListView widget that displays the rows
+		TSharedPtr<STileView<TableViewerItemPtr>> TileView;
+
+		// Our model class
+		TSharedPtr<FTedsTableViewerModel> Model;
+
+		// Delegate fired when the selection changes
+		FOnSelectionChanged OnSelectionChanged;
+
+		// Wrapper Teds Widget around our contents so we can use Teds columns to specify behavior
+		TSharedPtr<STedsWidget> TedsWidget;
+
+		// The widget purpose used to generate tiles
+		FName WidgetPurpose;
+	};
+}
