TEDS Alerts: Check Alert.NextAlert to make sure it is pointing at a valid row before running a subquery on it

diff --git a/Engine/Plugins/Experimental/EditorDataStorageFeatures/Source/TedsAlerts/Private/TedsAlertsFactory.cpp b/Engine/Plugins/Experimental/EditorDataStorageFeatures/Source/TedsAlerts/Private/TedsAlertsFactory.cpp
--- a/Engine/Plugins/Experimental/EditorDataStorageFeatures/Source/TedsAlerts/Private/TedsAlertsFactory.cpp
+++ b/Engine/Plugins/Experimental/EditorDataStorageFeatures/Source/TedsAlerts/Private/TedsAlertsFactory.cpp
@@ -92,7 +92,7 @@ void UTedsAlertsFactory::RegisterSortUnsortedAlertsQuery(IEditorDataStorageProvi
 			[](IQueryContext& NewAlertContext, RowHandle NewAlertRow, FTedsAlertColumn& NewAlert)
 			{
 				RowHandle NextRow = NewAlert.NextAlert;
-				do
+				while (NewAlertContext.IsRowAssigned(NextRow))
 				{
 					// Walk the chain and add link at the appropriate spot. This can include the row with the active alert.
 					NewAlertContext.RunSubquery(0, NextRow, CreateSubqueryCallbackBinding([NewAlertRow, &NewAlert, &NewAlertContext, &NextRow]
@@ -141,7 +141,7 @@ void UTedsAlertsFactory::RegisterSortUnsortedAlertsQuery(IEditorDataStorageProvi
 								}
 							}
 						}));
-				} while (NextRow != InvalidRowHandle);
+				}
 			})
 		.Where()
 			.Any<FTedsUnsortedAlertChainTag>()
