STedsTableViewer:
Integrate the new TEDS query stack
Remove the placeholder query stack used before

Teds Debugger and TEDS Content Browser:
Update the table viewer to use the new query stack

diff --git a/Engine/Plugins/Experimental/EditorDataStorageFeatures/Source/TedsContentBrowser/Private/TedsContentBrowserModule.cpp b/Engine/Plugins/Experimental/EditorDataStorageFeatures/Source/TedsContentBrowser/Private/TedsContentBrowserModule.cpp
--- a/Engine/Plugins/Experimental/EditorDataStorageFeatures/Source/TedsContentBrowser/Private/TedsContentBrowserModule.cpp
+++ b/Engine/Plugins/Experimental/EditorDataStorageFeatures/Source/TedsContentBrowser/Private/TedsContentBrowserModule.cpp
@@ -11,9 +11,9 @@
 #include "Elements/Interfaces/TypedElementDataStorageInterface.h"
 #include "Experimental/ContentBrowserViewExtender.h"
 #include "Modules/ModuleManager.h"
-#include "QueryStack/FQueryStackNode_RowView.h"
 #include "TedsAlertColumns.h"
 #include "TedsAssetDataColumns.h"
+#include "TedsRowViewNode.h"
 #include "Elements/Columns/TypedElementFolderColumns.h"
 #include "Elements/Columns/TypedElementMiscColumns.h"
 #include "Widgets/STedsTableViewer.h"
@@ -95,7 +95,8 @@ namespace UE::Editor::ContentBrowser
 			AddRow(Item);
 		}
 		
-		RowQueryStack->MarkDirty();
+		RowQueryStack->ResetView(FRowHandleArrayView(Rows.GetData(), Rows.Num(),
+			FRowHandleArrayView::EFlags::IsSorted | FRowHandleArrayView::EFlags::IsUnique));
 	}
 
 	void FTedsContentBrowserViewExtender::AddRow(const TSharedPtr<FAssetViewItem>& Item)
@@ -238,7 +239,8 @@ namespace UE::Editor::ContentBrowser
 	{
 		DataStorage = GetMutableDataStorageFeature<IEditorDataStorageProvider>(StorageFeatureName);
 		
-		RowQueryStack = MakeShared<FQueryStackNode_RowView>(&Rows);
+		RowQueryStack = MakeShared<QueryStack::FRowViewNode>(FRowHandleArrayView(Rows.GetData(), Rows.Num(),
+			FRowHandleArrayView::EFlags::IsSorted | FRowHandleArrayView::EFlags::IsUnique));
 	}
 
 	TSharedRef<SWidget> FTedsContentBrowserViewExtender::CreateView(TArray<TSharedPtr<FAssetViewItem>>* InItemsSource)
diff --git a/Engine/Plugins/Experimental/EditorDataStorageFeatures/Source/TedsContentBrowser/Private/TedsContentBrowserModule.h b/Engine/Plugins/Experimental/EditorDataStorageFeatures/Source/TedsContentBrowser/Private/TedsContentBrowserModule.h
--- a/Engine/Plugins/Experimental/EditorDataStorageFeatures/Source/TedsContentBrowser/Private/TedsContentBrowserModule.h
+++ b/Engine/Plugins/Experimental/EditorDataStorageFeatures/Source/TedsContentBrowser/Private/TedsContentBrowserModule.h
@@ -13,6 +13,10 @@ namespace UE::Editor::DataStorage
 {
 	class FQueryStackNode_RowView;
 	class ITableViewer;
+	namespace QueryStack
+	{
+		class FRowViewNode;
+	}
 }
 
 class FAssetViewItem;
@@ -75,7 +79,7 @@ namespace UE::Editor::ContentBrowser
 		TSharedPtr<DataStorage::ITableViewer> TableViewer;
 		
 		// Query stack used by the table viewer
-		TSharedPtr<DataStorage::FQueryStackNode_RowView> RowQueryStack;
+		TSharedPtr<DataStorage::QueryStack::FRowViewNode> RowQueryStack;
 
 		// The row handles of the items currently in the list
 		TArray<DataStorage::RowHandle> Rows;
diff --git a/Engine/Plugins/Experimental/EditorDataStorageFeatures/Source/TedsContentBrowser/TedsContentBrowser.Build.cs b/Engine/Plugins/Experimental/EditorDataStorageFeatures/Source/TedsContentBrowser/TedsContentBrowser.Build.cs
--- a/Engine/Plugins/Experimental/EditorDataStorageFeatures/Source/TedsContentBrowser/TedsContentBrowser.Build.cs
+++ b/Engine/Plugins/Experimental/EditorDataStorageFeatures/Source/TedsContentBrowser/TedsContentBrowser.Build.cs
@@ -32,7 +32,8 @@ public class TedsContentBrowser : ModuleRules
 		            "InputCore",
 					"TedsAlerts",
 		            "TedsAssetData",
-		            "ToolWidgets"
+		            "ToolWidgets",
+		            "TedsQueryStack"
             	});
 			
 			DynamicallyLoadedModuleNames.AddRange(new string[] {});
diff --git a/Engine/Plugins/Experimental/EditorDataStorageFeatures/Source/TedsDebugger/Private/Widgets/QueryEditor/TedsQueryEditorResultsView.cpp b/Engine/Plugins/Experimental/EditorDataStorageFeatures/Source/TedsDebugger/Private/Widgets/QueryEditor/TedsQueryEditorResultsView.cpp
--- a/Engine/Plugins/Experimental/EditorDataStorageFeatures/Source/TedsDebugger/Private/Widgets/QueryEditor/TedsQueryEditorResultsView.cpp
+++ b/Engine/Plugins/Experimental/EditorDataStorageFeatures/Source/TedsDebugger/Private/Widgets/QueryEditor/TedsQueryEditorResultsView.cpp
@@ -11,7 +11,7 @@
 #include "Elements/Interfaces/TypedElementDataStorageUiInterface.h"
 #include "Modules/ModuleManager.h"
 #include "QueryEditor/TedsQueryEditorModel.h"
-#include "QueryStack/FQueryStackNode_RowView.h"
+#include "TedsRowViewNode.h"
 #include "Widgets/STedsTableViewer.h"
 #include "Widgets/SRowDetails.h"
 #include "Widgets/Text/STextBlock.h"
@@ -34,7 +34,8 @@ namespace UE::Editor::DataStorage::Debug::QueryEditor
 		// Create a custom column for the table viewer to display row handles
 		CreateRowHandleColumn();
 	
-		RowQueryStack = MakeShared<FQueryStackNode_RowView>(&TableViewerRows);
+		RowQueryStack = MakeShared<QueryStack::FRowViewNode>(FRowHandleArrayView(TableViewerRows.GetData(), TableViewerRows.Num(),
+			FRowHandleArrayView::EFlags::IsSorted | FRowHandleArrayView::EFlags::IsUnique));
 
 		ChildSlot
 		[
@@ -167,7 +168,9 @@ namespace UE::Editor::DataStorage::Debug::QueryEditor
 		{
 			Swap(TableViewerRows_Set, NewTableViewerRows_Set);
 			TableViewerRows = TableViewerRows_Set.Array();
-			RowQueryStack->MarkDirty();
+
+			RowQueryStack->ResetView(FRowHandleArrayView(TableViewerRows.GetData(), TableViewerRows.Num(),
+				FRowHandleArrayView::EFlags::IsSorted | FRowHandleArrayView::EFlags::IsUnique));
 		}
 	}
 
diff --git a/Engine/Plugins/Experimental/EditorDataStorageFeatures/Source/TedsDebugger/Private/Widgets/QueryEditor/TedsQueryEditorResultsView.h b/Engine/Plugins/Experimental/EditorDataStorageFeatures/Source/TedsDebugger/Private/Widgets/QueryEditor/TedsQueryEditorResultsView.h
--- a/Engine/Plugins/Experimental/EditorDataStorageFeatures/Source/TedsDebugger/Private/Widgets/QueryEditor/TedsQueryEditorResultsView.h
+++ b/Engine/Plugins/Experimental/EditorDataStorageFeatures/Source/TedsDebugger/Private/Widgets/QueryEditor/TedsQueryEditorResultsView.h
@@ -12,10 +12,14 @@ struct FTypedElementWidgetConstructor;
 namespace UE::Editor::DataStorage
 {
 	class STedsTableViewer;
-	class FQueryStackNode_RowView;
 	class FTedsTableViewerColumn;
 	class SRowDetails;
 
+	namespace QueryStack
+	{
+		class FRowViewNode;
+	}
+
 	namespace Debug::QueryEditor
 	{
 		class FTedsQueryEditorModel;
@@ -44,10 +48,10 @@ namespace UE::Editor::DataStorage
 			QueryHandle TableViewerQueryHandle = InvalidQueryHandle;
 
 			TArray<RowHandle> TableViewerRows;
-			// We have to keep a TSet copy because queries return duplicate rows sometimes and to have some form of sorted order for the rows for now
+			// We have to keep a TSet copy to have a sorted order for the rows for now
 			TSet<RowHandle> TableViewerRows_Set;
 			TSharedPtr<STedsTableViewer> TableViewer;
-			TSharedPtr<FQueryStackNode_RowView> RowQueryStack;
+			TSharedPtr<QueryStack::FRowViewNode> RowQueryStack;
 
 			// Custom column for the table viewer to display row handles
 			TSharedPtr<FTedsTableViewerColumn> RowHandleColumn;
diff --git a/Engine/Plugins/Experimental/EditorDataStorageFeatures/Source/TedsDebugger/TedsDebugger.Build.cs b/Engine/Plugins/Experimental/EditorDataStorageFeatures/Source/TedsDebugger/TedsDebugger.Build.cs
--- a/Engine/Plugins/Experimental/EditorDataStorageFeatures/Source/TedsDebugger/TedsDebugger.Build.cs
+++ b/Engine/Plugins/Experimental/EditorDataStorageFeatures/Source/TedsDebugger/TedsDebugger.Build.cs
@@ -33,7 +33,8 @@ public class TedsDebugger : ModuleRules
 		            "EditorWidgets",
 		            "ToolWidgets",
 		            "EditorWidgets",
-		            "TedsTableViewer"
+		            "TedsTableViewer",
+		            "TedsQueryStack"
             	});
 			
 			DynamicallyLoadedModuleNames.AddRange(new string[] {});
diff --git a/Engine/Plugins/Experimental/EditorDataStorageFeatures/Source/TedsQueryStack/Private/TedsRowViewNode.cpp b/Engine/Plugins/Experimental/EditorDataStorageFeatures/Source/TedsQueryStack/Private/TedsRowViewNode.cpp
--- a/Engine/Plugins/Experimental/EditorDataStorageFeatures/Source/TedsQueryStack/Private/TedsRowViewNode.cpp
+++ b/Engine/Plugins/Experimental/EditorDataStorageFeatures/Source/TedsQueryStack/Private/TedsRowViewNode.cpp
@@ -13,6 +13,12 @@ namespace UE::Editor::DataStorage::QueryStack
 		Revision++;
 	}
 
+	void FRowViewNode::ResetView(FRowHandleArrayView InRows)
+	{
+		Rows = InRows;
+		MarkDirty();
+	}
+
 	INode::RevisionId FRowViewNode::GetRevision() const
 	{
 		return Revision;
diff --git a/Engine/Plugins/Experimental/EditorDataStorageFeatures/Source/TedsQueryStack/Public/TedsRowViewNode.h b/Engine/Plugins/Experimental/EditorDataStorageFeatures/Source/TedsQueryStack/Public/TedsRowViewNode.h
--- a/Engine/Plugins/Experimental/EditorDataStorageFeatures/Source/TedsQueryStack/Public/TedsRowViewNode.h
+++ b/Engine/Plugins/Experimental/EditorDataStorageFeatures/Source/TedsQueryStack/Public/TedsRowViewNode.h
@@ -19,6 +19,7 @@ namespace UE::Editor::DataStorage::QueryStack
 		virtual ~FRowViewNode() override = default;
 
 		TEDSQUERYSTACK_API void MarkDirty();
+		TEDSQUERYSTACK_API void ResetView(FRowHandleArrayView InRows);
 
 		TEDSQUERYSTACK_API virtual RevisionId GetRevision() const override;
 		TEDSQUERYSTACK_API virtual void Update() override;
diff --git a/Engine/Plugins/Experimental/EditorDataStorageFeatures/Source/TedsTableViewer/Private/QueryStack/FQueryStackNode_RowView.cpp b/Engine/Plugins/Experimental/EditorDataStorageFeatures/Source/TedsTableViewer/Private/QueryStack/FQueryStackNode_RowView.cpp
deleted file mode 100644
diff --git a/Engine/Plugins/Experimental/EditorDataStorageFeatures/Source/TedsTableViewer/Private/TedsTableViewerModel.cpp b/Engine/Plugins/Experimental/EditorDataStorageFeatures/Source/TedsTableViewer/Private/TedsTableViewerModel.cpp
--- a/Engine/Plugins/Experimental/EditorDataStorageFeatures/Source/TedsTableViewer/Private/TedsTableViewerModel.cpp
+++ b/Engine/Plugins/Experimental/EditorDataStorageFeatures/Source/TedsTableViewer/Private/TedsTableViewerModel.cpp
@@ -8,13 +8,12 @@
 #include "Elements/Interfaces/TypedElementDataStorageInterface.h"
 #include "Elements/Interfaces/TypedElementDataStorageCompatibilityInterface.h"
 #include "Elements/Interfaces/TypedElementDataStorageUiInterface.h"
-#include "QueryStack/IQueryStackNode_Row.h"
 #include "TedsTableViewerColumn.h"
 #include "TedsTableViewerUtils.h"
 
 namespace UE::Editor::DataStorage
 {
-	FTedsTableViewerModel::FTedsTableViewerModel(const TSharedPtr<IQueryStackNode_Row>& InRowQueryStack,
+	FTedsTableViewerModel::FTedsTableViewerModel(const TSharedPtr<QueryStack::IRowNode>& InRowQueryStack,
 		const TArray<TWeakObjectPtr<const UScriptStruct>>& InRequestedColumns, const TArray<FName>& InCellWidgetPurposes,
 		const FIsItemVisible& InIsItemVisibleDelegate)
 		: RowQueryStack(InRowQueryStack)
@@ -40,8 +39,10 @@ namespace UE::Editor::DataStorage
 	void FTedsTableViewerModel::Refresh()
 	{
 		Items.Empty();
+
+		FRowHandleArrayView Rows = RowQueryStack->GetRows();
 		
-		for(const RowHandle RowHandle : RowQueryStack->GetOrderedRowList())
+		for(const RowHandle RowHandle : Rows)
 		{
 			if(IsRowDisplayable(RowHandle))
 			{
@@ -50,7 +51,7 @@ namespace UE::Editor::DataStorage
 			}
 		}
 
-		CachedRowQueryStackRevision = RowQueryStack->GetRevisionId();
+		CachedRowQueryStackRevision = RowQueryStack->GetRevision();
 
 		OnModelChanged.Broadcast();
 	}
@@ -63,7 +64,7 @@ namespace UE::Editor::DataStorage
 	bool FTedsTableViewerModel::Tick(float DeltaTime)
 	{
 		// If the revision ID has changed, refresh to update our rows
-		if(RowQueryStack->GetRevisionId() != CachedRowQueryStackRevision)
+		if(RowQueryStack->GetRevision() != CachedRowQueryStackRevision)
 		{
 			Refresh();
 		}
diff --git a/Engine/Plugins/Experimental/EditorDataStorageFeatures/Source/TedsTableViewer/Public/QueryStack/FQueryStackNode_RowView.h b/Engine/Plugins/Experimental/EditorDataStorageFeatures/Source/TedsTableViewer/Public/QueryStack/FQueryStackNode_RowView.h
deleted file mode 100644
diff --git a/Engine/Plugins/Experimental/EditorDataStorageFeatures/Source/TedsTableViewer/Public/QueryStack/IQueryStackNode_Row.h b/Engine/Plugins/Experimental/EditorDataStorageFeatures/Source/TedsTableViewer/Public/QueryStack/IQueryStackNode_Row.h
deleted file mode 100644
diff --git a/Engine/Plugins/Experimental/EditorDataStorageFeatures/Source/TedsTableViewer/Public/TedsTableViewerModel.h b/Engine/Plugins/Experimental/EditorDataStorageFeatures/Source/TedsTableViewer/Public/TedsTableViewerModel.h
--- a/Engine/Plugins/Experimental/EditorDataStorageFeatures/Source/TedsTableViewer/Public/TedsTableViewerModel.h
+++ b/Engine/Plugins/Experimental/EditorDataStorageFeatures/Source/TedsTableViewer/Public/TedsTableViewerModel.h
@@ -7,6 +7,7 @@
 #include "Elements/Common/TypedElementHandles.h"
 #include "Templates/SharedPointer.h"
 #include "TypedElementUITypes.h"
+#include "TedsQueryStackInterfaces.h"
 #include "UObject/NameTypes.h"
 
 class IEditorDataStorageCompatibilityProvider;
@@ -16,8 +17,7 @@ class IEditorDataStorageProvider;
 namespace UE::Editor::DataStorage
 {
 	class FTedsTableViewerColumn;
-	class IQueryStackNode_Row;
-
+	
 	// Typedef for an item in the table viewer
 	using TableViewerItemPtr = FTedsRowHandle;
 
@@ -30,7 +30,7 @@ namespace UE::Editor::DataStorage
 		DECLARE_DELEGATE_RetVal_OneParam(bool, FIsItemVisible, TableViewerItemPtr);
 		DECLARE_MULTICAST_DELEGATE(FOnModelChanged);
 
-		TEDSTABLEVIEWER_API FTedsTableViewerModel(const TSharedPtr<IQueryStackNode_Row>& RowQueryStack, const TArray<TWeakObjectPtr<const UScriptStruct>>& RequestedColumns,
+		TEDSTABLEVIEWER_API FTedsTableViewerModel(const TSharedPtr<QueryStack::IRowNode>& RowQueryStack, const TArray<TWeakObjectPtr<const UScriptStruct>>& RequestedColumns,
 			const TArray<FName>& CellWidgetPurposes, const FIsItemVisible& InIsItemVisibleDelegate);
 
 		~FTedsTableViewerModel();
@@ -79,7 +79,7 @@ namespace UE::Editor::DataStorage
 	private:
 
 		// The row query stack used to supply the rows to display
-		TSharedPtr<IQueryStackNode_Row> RowQueryStack;
+		TSharedPtr<QueryStack::IRowNode> RowQueryStack;
 
 		// The cached list of rows we are currently displaying
 		TArray<TableViewerItemPtr> Items;
diff --git a/Engine/Plugins/Experimental/EditorDataStorageFeatures/Source/TedsTableViewer/Public/Widgets/STedsTableViewer.h b/Engine/Plugins/Experimental/EditorDataStorageFeatures/Source/TedsTableViewer/Public/Widgets/STedsTableViewer.h
--- a/Engine/Plugins/Experimental/EditorDataStorageFeatures/Source/TedsTableViewer/Public/Widgets/STedsTableViewer.h
+++ b/Engine/Plugins/Experimental/EditorDataStorageFeatures/Source/TedsTableViewer/Public/Widgets/STedsTableViewer.h
@@ -3,7 +3,7 @@
 
 #include "ITedsTableViewer.h"
 #include "Elements/Common/TypedElementHandles.h"
-#include "QueryStack/IQueryStackNode_Row.h"
+#include "TedsQueryStackInterfaces.h"
 #include "TedsTableViewerModel.h"
 #include "Widgets/SCompoundWidget.h"
 #include "Widgets/Views/STableRow.h"
@@ -43,7 +43,7 @@ namespace UE::Editor::DataStorage
 		
 
 		// Query Stack that will supply the rows to be displayed
-		SLATE_ARGUMENT(TSharedPtr<IQueryStackNode_Row>, QueryStack)
+		SLATE_ARGUMENT(TSharedPtr<QueryStack::IRowNode>, QueryStack)
 
 		// The Columns that this table viewer will display
 		// Table Viewer TODO: How do we specify column metadata (ReadOnly or ReadWrite)?
diff --git a/Engine/Plugins/Experimental/EditorDataStorageFeatures/Source/TedsTableViewer/Public/Widgets/STedsTileViewer.h b/Engine/Plugins/Experimental/EditorDataStorageFeatures/Source/TedsTableViewer/Public/Widgets/STedsTileViewer.h
--- a/Engine/Plugins/Experimental/EditorDataStorageFeatures/Source/TedsTableViewer/Public/Widgets/STedsTileViewer.h
+++ b/Engine/Plugins/Experimental/EditorDataStorageFeatures/Source/TedsTableViewer/Public/Widgets/STedsTileViewer.h
@@ -3,7 +3,7 @@
 
 #include "ITedsTableViewer.h"
 #include "Elements/Common/TypedElementHandles.h"
-#include "QueryStack/IQueryStackNode_Row.h"
+#include "TedsQueryStackInterfaces.h"
 #include "TedsTableViewerModel.h"
 #include "Widgets/SCompoundWidget.h"
 #include "Widgets/Views/STableRow.h"
@@ -35,7 +35,7 @@ namespace UE::Editor::DataStorage
 		}
 
 		// Query Stack that will supply the rows to be displayed
-		SLATE_ARGUMENT(TSharedPtr<IQueryStackNode_Row>, QueryStack)
+		SLATE_ARGUMENT(TSharedPtr<QueryStack::IRowNode>, QueryStack)
 
 		// The widget purpose to use to create the widgets
 		SLATE_ARGUMENT(FName, WidgetPurpose)
diff --git a/Engine/Plugins/Experimental/EditorDataStorageFeatures/Source/TedsTableViewer/TedsTableViewer.Build.cs b/Engine/Plugins/Experimental/EditorDataStorageFeatures/Source/TedsTableViewer/TedsTableViewer.Build.cs
--- a/Engine/Plugins/Experimental/EditorDataStorageFeatures/Source/TedsTableViewer/TedsTableViewer.Build.cs
+++ b/Engine/Plugins/Experimental/EditorDataStorageFeatures/Source/TedsTableViewer/TedsTableViewer.Build.cs
@@ -15,7 +15,8 @@ public class TedsTableViewer : ModuleRules
 			PublicDependencyModuleNames.AddRange(
 				new string[]
 				{
-					"TypedElementFramework"
+					"TypedElementFramework",
+					"TedsQueryStack"
 				});
 
 			PrivateDependencyModuleNames.AddRange(
