Protected TEDS Index Table against hash collisions and renamed to TEDS Mapping Table.

The TEDS Index Table uses a 64-bit hash to handle indexing, but doesn't store any additional information. Because originally only pointers were used this was ok as these are guaranteed not to collide. Over time though more data was getting indexed and cleanup delayed increasing the risk of hash collisions. These updates move away from a single 64-bit integer and introduced a strongly typed key that supports a variety of data types. These keys are stores along with the hash table to be able to detect and resolve hash collisions or report them if necessary.
The TEDS Index Table has been renamed to TEDS Mapping Table. This was done because the term "index" is going to be introduced in the future for another feature and this was the more reasonable rename.
(The dependend changes were reviewed seperately in 39657563.)

[RN] Protected TEDS Index Table against hash collisions and renamed to TEDS Mapping Table.

diff --git a/Engine/Plugins/ChaosVD/Source/ChaosVD/Private/TEDS/ChaosVDParentDataStorageFactory.cpp b/Engine/Plugins/ChaosVD/Source/ChaosVD/Private/TEDS/ChaosVDParentDataStorageFactory.cpp
--- a/Engine/Plugins/ChaosVD/Source/ChaosVD/Private/TEDS/ChaosVDParentDataStorageFactory.cpp
+++ b/Engine/Plugins/ChaosVD/Source/ChaosVD/Private/TEDS/ChaosVDParentDataStorageFactory.cpp
@@ -15,7 +15,6 @@
 #include "Elements/Columns/TypedElementLabelColumns.h"
 #include "Elements/Columns/TypedElementMiscColumns.h"
 #include "Elements/Columns/TypedElementTypeInfoColumns.h"
-#include "Elements/Framework/TypedElementIndexHasher.h"
 #include "Elements/Framework/TypedElementQueryBuilder.h"
 
 namespace Chaos::VD::ParentDataStorageUtils::Private
@@ -83,8 +82,8 @@ void UChaosVDParentDataStorageFactory::RegisterAddParentColumn(IEditorDataStorag
 							FChaosVDTableRowParentColumn ParentColumn;
 							ParentColumn.ParentObject = DefaultRoot->GetTedsRowHandle();
 
-							uint64 IdHash = GenerateIndexHash(CVDActor);
-							RowHandle ActorRow = Context.FindIndexedRow(IdHash);
+							FMapKeyView IdKey = FMapKeyView(CVDActor);
+							RowHandle ActorRow = Context.LookupMappedRow(IdKey);
 							ParentData->ChildrenSet.Emplace(ActorRow);
 							Context.AddColumn(Row, MoveTemp(ParentColumn));
 						}
@@ -120,8 +119,8 @@ void UChaosVDParentDataStorageFactory::RegisterUpdateOrRemoveParentColumn(IEdito
 					}
 					else if (AChaosVDSolverInfoActor* NewParentActor = Cast<AChaosVDSolverInfoActor>(Particle->GetParentActor()))
 					{
-						uint64 IdHash = GenerateIndexHash(NewParentActor);
-						NewParentRowHandle = Context.FindIndexedRow(IdHash);
+						FMapKeyView IdKey = FMapKeyView(NewParentActor);
+						NewParentRowHandle = Context.LookupMappedRow(IdKey);
 					}
 
 					UpdateParentData(DataStorage, ParentColumn, Row, NewParentRowHandle);
@@ -144,8 +143,8 @@ void UChaosVDParentDataStorageFactory::RegisterUpdateOrRemoveParentColumn(IEdito
 	
 					if (AActor* ParentActor = ActorPtr->GetAttachParentActor())
 					{
-						uint64 IdHash = GenerateIndexHash(ParentActor);
-						NewParentRowHandle = Context.FindIndexedRow(IdHash);
+						FMapKeyView IdKey = FMapKeyView(ParentActor);
+						NewParentRowHandle = Context.LookupMappedRow(IdKey);
 					}
 					
 					UpdateParentData(DataStorage, Parent, Row, NewParentRowHandle);
diff --git a/Engine/Plugins/Developer/PerforceSourceControl/Source/PerforceSourceControl/Private/PerforceSourceControlOperations.cpp b/Engine/Plugins/Developer/PerforceSourceControl/Source/PerforceSourceControl/Private/PerforceSourceControlOperations.cpp
--- a/Engine/Plugins/Developer/PerforceSourceControl/Source/PerforceSourceControl/Private/PerforceSourceControlOperations.cpp
+++ b/Engine/Plugins/Developer/PerforceSourceControl/Source/PerforceSourceControl/Private/PerforceSourceControlOperations.cpp
@@ -27,7 +27,6 @@
 #include "Elements/Columns/TypedElementPackageColumns.h"
 #include "Elements/Columns/TypedElementRevisionControlColumns.h"
 #include "Elements/Common/EditorDataStorageFeatures.h"
-#include "Elements/Framework/TypedElementIndexHasher.h"
 #include "Elements/Interfaces/TypedElementDataStorageInterface.h"
 #include "ProfilingDebugging/CpuProfilerTrace.h"
 
@@ -2070,13 +2069,13 @@ bool FPerforceUpdateStatusWorker::UpdateStates() const
 			FPaths::NormalizeFilename(Filename);
 			Filename = FPaths::ConvertRelativePathToFull(Filename);
 			
-			uint64 Index = GenerateIndexHash(Filename);
-			RowHandle Row = DataStorage->FindIndexedRow(Index);
+			FMapKey Key = FMapKey(Filename);
+			RowHandle Row = DataStorage->LookupMappedRow(Key);
 
 			if (!DataStorage->IsRowAvailable(Row))
 			{
 				Row = DataStorage->AddRow(RevisionControlTable);
-				DataStorage->IndexRow(Index, Row);
+				DataStorage->MapRow(Key, Row);
 			}
 			return Row;
 		};
diff --git a/Engine/Plugins/Experimental/EditorDataStorage/Source/TedsCore/Private/Commands/EditorDataStorageCompatibilityCommands.cpp b/Engine/Plugins/Experimental/EditorDataStorage/Source/TedsCore/Private/Commands/EditorDataStorageCompatibilityCommands.cpp
--- a/Engine/Plugins/Experimental/EditorDataStorage/Source/TedsCore/Private/Commands/EditorDataStorageCompatibilityCommands.cpp
+++ b/Engine/Plugins/Experimental/EditorDataStorage/Source/TedsCore/Private/Commands/EditorDataStorageCompatibilityCommands.cpp
@@ -6,7 +6,6 @@
 #include "Elements/Columns/TypedElementCompatibilityColumns.h"
 #include "Elements/Columns/TypedElementMiscColumns.h"
 #include "Elements/Columns/TypedElementTypeInfoColumns.h"
-#include "Elements/Framework/TypedElementIndexHasher.h"
 #include "Elements/Framework/TypedElementQueryBuilder.h"
 #include "Elements/Interfaces/TypedElementDataStorageInterface.h"
 #include "Memento/TypedElementMementoSystem.h"
@@ -267,8 +266,8 @@ namespace UE::Editor::DataStorage
 	{
 		if (Command.ObjectRow == InvalidRowHandle)
 		{
-			IndexHash Hash = GenerateIndexHash(Command.Object);
-			Command.ObjectRow = Storage.FindIndexedRow(Hash);
+			FMapKeyView Key = FMapKeyView(Command.Object);
+			Command.ObjectRow = Storage.LookupMappedRow(Key);
 			if (!Storage.IsRowAvailable(Command.ObjectRow))
 			{
 				Commands.ReplaceCommand<FNopCommand>(CurrentIndex);
@@ -278,8 +277,8 @@ namespace UE::Editor::DataStorage
 
 	void FPrepareCommands::operator()(FRemoveCompatibleExternalObject& Command)
 	{
-		IndexHash Hash = GenerateIndexHash(Command.Object);
-		Command.ObjectRow = Storage.FindIndexedRow(Hash);
+		FMapKeyView Key = FMapKeyView(Command.Object);
+		Command.ObjectRow = Storage.LookupMappedRow(Key);
 		if (!Storage.IsRowAvailable(Command.ObjectRow))
 		{
 			Commands.ReplaceCommand<FNopCommand>(CurrentIndex);
diff --git a/Engine/Plugins/Experimental/EditorDataStorage/Source/TedsCore/Private/Compatibility/TypedElementActorParentQueries.cpp b/Engine/Plugins/Experimental/EditorDataStorage/Source/TedsCore/Private/Compatibility/TypedElementActorParentQueries.cpp
--- a/Engine/Plugins/Experimental/EditorDataStorage/Source/TedsCore/Private/Compatibility/TypedElementActorParentQueries.cpp
+++ b/Engine/Plugins/Experimental/EditorDataStorage/Source/TedsCore/Private/Compatibility/TypedElementActorParentQueries.cpp
@@ -5,7 +5,7 @@
 #include "Elements/Columns/TypedElementCompatibilityColumns.h"
 #include "Elements/Columns/TypedElementHiearchyColumns.h"
 #include "Elements/Columns/TypedElementMiscColumns.h"
-#include "Elements/Framework/TypedElementIndexHasher.h"
+#include "Elements/Common/TypedElementMapKey.h"
 #include "Elements/Framework/TypedElementQueryBuilder.h"
 #include "GameFramework/Actor.h"
 
@@ -31,6 +31,7 @@ void UActorParentDataStorageFactory::RegisterQueries(IEditorDataStorageProvider&
 
 void UActorParentDataStorageFactory::RegisterAddParentColumn(IEditorDataStorageProvider& DataStorage) const
 {
+	using namespace UE::Editor::DataStorage;
 	using namespace UE::Editor::DataStorage::Queries;
 	
 	DataStorage.RegisterQuery(
@@ -44,15 +45,15 @@ void UActorParentDataStorageFactory::RegisterAddParentColumn(IEditorDataStorageP
 				{
 					if (AActor* Parent = ActorInstance->GetAttachParentActor())
 					{
-						uint64 IdHash = GenerateIndexHash(Parent);
-						RowHandle ParentRow = Context.FindIndexedRow(IdHash);
+						FMapKeyView IdKey = FMapKeyView(Parent);
+						RowHandle ParentRow = Context.LookupMappedRow(IdKey);
 						if (Context.IsRowAvailable(ParentRow))
 						{
 							Context.AddColumn(Row, FTableRowParentColumn{ .Parent = ParentRow });
 						}
 						else
 						{
-							Context.AddColumn(Row, FUnresolvedTableRowParentColumn{ .ParentIdHash = IdHash });
+							Context.AddColumn(Row, FUnresolvedTableRowParentColumn{ .ParentIdKey = FMapKey(Parent)});
 						}
 					}
 				}
@@ -78,8 +79,8 @@ void UActorParentDataStorageFactory::RegisterUpdateOrRemoveParentColumn(IEditorD
 				{
 					if (AActor* ParentActor = ActorInstance->GetAttachParentActor())
 					{
-						uint64 IdHash = GenerateIndexHash(ParentActor);
-						RowHandle ParentRow = Context.FindIndexedRow(IdHash);
+						FMapKeyView IdKey = FMapKeyView(ParentActor);
+						RowHandle ParentRow = Context.LookupMappedRow(IdKey);
 
 						if (Parent.Parent != ParentRow)
 						{
@@ -90,7 +91,7 @@ void UActorParentDataStorageFactory::RegisterUpdateOrRemoveParentColumn(IEditorD
 							else
 							{
 								Context.RemoveColumns<FTableRowParentColumn>(Row);
-								Context.AddColumn(Row, FUnresolvedTableRowParentColumn{ .ParentIdHash = IdHash });
+								Context.AddColumn(Row, FUnresolvedTableRowParentColumn{ .ParentIdKey = FMapKey(ParentActor)});
 							}
 						}
 						return;
diff --git a/Engine/Plugins/Experimental/EditorDataStorage/Source/TedsCore/Private/Processors/TypedElementProcessorAdaptors.cpp b/Engine/Plugins/Experimental/EditorDataStorage/Source/TedsCore/Private/Processors/TypedElementProcessorAdaptors.cpp
--- a/Engine/Plugins/Experimental/EditorDataStorage/Source/TedsCore/Private/Processors/TypedElementProcessorAdaptors.cpp
+++ b/Engine/Plugins/Experimental/EditorDataStorage/Source/TedsCore/Private/Processors/TypedElementProcessorAdaptors.cpp
@@ -756,12 +756,12 @@ namespace UE::Editor::DataStorage
 				}
 			}
 	
-			RowHandle FindIndexedRow(IndexHash Index) const
+			RowHandle LookupMappedRow(const FMapKeyView& Key) const
 			{
 				EGlobalLockScope Scope = FGlobalLock::GetLockStatus(EGlobalLockScope::Internal) == EGlobalLockStatus::Unlocked
 					? EGlobalLockScope::Public // There's no internal lock so use a public lock instead.
 					: EGlobalLockScope::Internal; // There's an internal lock set so use that.
-				return Environment.GetIndexTable().FindIndexedRow(Scope, Index);
+				return Environment.GetMappingTable().Lookup(Scope, Key);
 			}
 
 			UE::Editor::DataStorage::FQueryResult RunQuery(QueryHandle Query)
@@ -851,7 +851,7 @@ namespace UE::Editor::DataStorage
 			const UObject* GetDependency(const UClass* DependencyClass) override { return Implementation.GetDependency(DependencyClass); }
 			UObject* GetMutableDependency(const UClass* DependencyClass) override { return Implementation.GetMutableDependency(DependencyClass); }
 			void GetDependencies(TArrayView<UObject*> RetrievedAddresses, TConstArrayView<TWeakObjectPtr<const UClass>> DependencyTypes, TConstArrayView<EQueryAccessType> AccessTypes) override { return Implementation.GetDependencies(RetrievedAddresses, DependencyTypes, AccessTypes); }
-			RowHandle FindIndexedRow(IndexHash Index) const override { return Implementation.FindIndexedRow(Index); }
+			RowHandle LookupMappedRow(const FMapKeyView& Index) const override { return Implementation.LookupMappedRow(Index); }
 			FQueryResult RunQuery(QueryHandle Query) override { return Implementation.RunQuery(Query); }
 			FQueryResult RunSubquery(int32 SubqueryIndex) override { return Implementation.RunSubquery(SubqueryIndex); }
 			FQueryResult RunSubquery(int32 SubqueryIndex, UE::Editor::DataStorage::SubqueryCallbackRef Callback) override { return Implementation.RunSubquery(SubqueryIndex, Callback); }
diff --git a/Engine/Plugins/Experimental/EditorDataStorage/Source/TedsCore/Private/Queries/TypedElementHierarchyQueries.cpp b/Engine/Plugins/Experimental/EditorDataStorage/Source/TedsCore/Private/Queries/TypedElementHierarchyQueries.cpp
--- a/Engine/Plugins/Experimental/EditorDataStorage/Source/TedsCore/Private/Queries/TypedElementHierarchyQueries.cpp
+++ b/Engine/Plugins/Experimental/EditorDataStorage/Source/TedsCore/Private/Queries/TypedElementHierarchyQueries.cpp
@@ -15,7 +15,7 @@ void UTypedElementHiearchyQueriesFactory::RegisterQueries(IEditorDataStorageProv
 			FProcessor(EQueryTickPhase::PrePhysics, DataStorage.GetQueryTickGroupName(EQueryTickGroups::Default)),
 			[](IQueryContext& Context, RowHandle Row, const FUnresolvedTableRowParentColumn& UnresolvedParent)
 			{
-				RowHandle ParentRow = Context.FindIndexedRow(UnresolvedParent.ParentIdHash);
+				RowHandle ParentRow = Context.LookupMappedRow(UnresolvedParent.ParentIdKey);
 				if (Context.IsRowAvailable(ParentRow))
 				{
 					Context.RemoveColumns<FUnresolvedTableRowParentColumn>(Row);
diff --git a/Engine/Plugins/Experimental/EditorDataStorage/Source/TedsCore/Private/TypedElementDatabase.cpp b/Engine/Plugins/Experimental/EditorDataStorage/Source/TedsCore/Private/TypedElementDatabase.cpp
--- a/Engine/Plugins/Experimental/EditorDataStorage/Source/TedsCore/Private/TypedElementDatabase.cpp
+++ b/Engine/Plugins/Experimental/EditorDataStorage/Source/TedsCore/Private/TypedElementDatabase.cpp
@@ -469,7 +469,7 @@ void UEditorDataStorage::RemoveRow(RowHandle Row)
 			Environment->GetDirectDeferredCommands().Clear(Row);
 			ActiveEditorEntityManager->ReleaseReservedEntity(FMassEntityHandle::FromNumber(Row));
 		}
-		Environment->GetIndexTable().MarkDirty();
+		Environment->GetMappingTable().MarkDirty();
 	}
 }
 
@@ -480,7 +480,7 @@ void UEditorDataStorage::BatchRemoveRows(TConstArrayView<RowHandle> Rows)
 	if (ActiveEditorEntityManager)
 	{
 		ActiveEditorEntityManager->BatchDestroyEntities(RowsToMassEntitiesConversion(Rows));
-		Environment->GetIndexTable().MarkDirty();
+		Environment->GetMappingTable().MarkDirty();
 	}
 }
 
@@ -520,7 +520,7 @@ void UEditorDataStorage::RemoveAllRowsWithColumns(TConstArrayView<UScriptStruct*
 
 			ActiveEditorEntityManager->BatchDestroyEntityChunks(Collections);
 
-			Environment->GetIndexTable().MarkDirty();
+			Environment->GetMappingTable().MarkDirty();
 		}
 	}
 }
@@ -1140,30 +1140,29 @@ void UEditorDataStorage::ActivateQueries(FName ActivationName)
 	}
 }
 
-UE::Editor::DataStorage::RowHandle UEditorDataStorage::FindIndexedRow(UE::Editor::DataStorage::IndexHash Index) const
+UE::Editor::DataStorage::RowHandle UEditorDataStorage::LookupMappedRow(const UE::Editor::DataStorage::FMapKeyView& Key) const
 {
-	return Environment->GetIndexTable().FindIndexedRow(UE::Editor::DataStorage::EGlobalLockScope::Public, Index);
+	return Environment->GetMappingTable().Lookup(UE::Editor::DataStorage::EGlobalLockScope::Public, Key);
 }
 
-void UEditorDataStorage::IndexRow(UE::Editor::DataStorage::IndexHash Index, RowHandle Row)
+void UEditorDataStorage::MapRow(UE::Editor::DataStorage::FMapKey Key, RowHandle Row)
 {
-	Environment->GetIndexTable().IndexRow(UE::Editor::DataStorage::EGlobalLockScope::Public, Index, Row);
+	Environment->GetMappingTable().Map(UE::Editor::DataStorage::EGlobalLockScope::Public, MoveTemp(Key), Row);
 }
 
-void UEditorDataStorage::BatchIndexRows(
-	TConstArrayView<TPair<UE::Editor::DataStorage::IndexHash, RowHandle>> IndexRowPairs)
+void UEditorDataStorage::BatchMapRows(TArrayView<TPair<UE::Editor::DataStorage::FMapKey, RowHandle>> MapRowPairs)
 {
-	Environment->GetIndexTable().BatchIndexRows(UE::Editor::DataStorage::EGlobalLockScope::Public, IndexRowPairs);
+	Environment->GetMappingTable().BatchMap(UE::Editor::DataStorage::EGlobalLockScope::Public, MapRowPairs);
 }
 
-void UEditorDataStorage::Reindex(UE::Editor::DataStorage::IndexHash OriginalIndex, UE::Editor::DataStorage::IndexHash NewIndex)
+void UEditorDataStorage::RemapRow(const UE::Editor::DataStorage::FMapKeyView& OriginalKey, UE::Editor::DataStorage::FMapKey NewKey)
 {
-	Environment->GetIndexTable().Reindex(UE::Editor::DataStorage::EGlobalLockScope::Public, OriginalIndex, NewIndex);
+	Environment->GetMappingTable().Remap(UE::Editor::DataStorage::EGlobalLockScope::Public, OriginalKey, MoveTemp(NewKey));
 }
 
-void UEditorDataStorage::RemoveIndex(UE::Editor::DataStorage::IndexHash Index)
+void UEditorDataStorage::RemoveRowMapping(const UE::Editor::DataStorage::FMapKeyView& Key)
 {
-	Environment->GetIndexTable().RemoveIndex(UE::Editor::DataStorage::EGlobalLockScope::Public, Index);
+	Environment->GetMappingTable().Remove(UE::Editor::DataStorage::EGlobalLockScope::Public, Key);
 }
 
 FTypedElementOnDataStorageUpdate& UEditorDataStorage::OnUpdate()
diff --git a/Engine/Plugins/Experimental/EditorDataStorage/Source/TedsCore/Private/TypedElementDatabase.h b/Engine/Plugins/Experimental/EditorDataStorage/Source/TedsCore/Private/TypedElementDatabase.h
--- a/Engine/Plugins/Experimental/EditorDataStorage/Source/TedsCore/Private/TypedElementDatabase.h
+++ b/Engine/Plugins/Experimental/EditorDataStorage/Source/TedsCore/Private/TypedElementDatabase.h
@@ -158,14 +158,11 @@ public:
 		DirectQueryCallbackRef Callback) override;
 	virtual void ActivateQueries(FName ActivationName) override;
 
-	virtual RowHandle FindIndexedRow(UE::Editor::DataStorage::IndexHash Index) const override;
-	virtual void IndexRow(UE::Editor::DataStorage::IndexHash Index, RowHandle Row) override;
-	virtual void BatchIndexRows(
-		TConstArrayView<TPair<UE::Editor::DataStorage::IndexHash, RowHandle>> IndexRowPairs) override;
-	virtual void Reindex(
-		UE::Editor::DataStorage::IndexHash OriginalIndex, 
-		UE::Editor::DataStorage::IndexHash NewIndex) override;
-	virtual void RemoveIndex(UE::Editor::DataStorage::IndexHash Index) override;
+	virtual RowHandle LookupMappedRow(const UE::Editor::DataStorage::FMapKeyView& Key) const override;
+	virtual void MapRow(UE::Editor::DataStorage::FMapKey Key, RowHandle Row) override;
+	virtual void BatchMapRows(TArrayView<TPair<UE::Editor::DataStorage::FMapKey, RowHandle>> MapRowPairs) override;
+	virtual void RemapRow(const UE::Editor::DataStorage::FMapKeyView& OriginalKey, UE::Editor::DataStorage::FMapKey NewKey) override;
+	virtual void RemoveRowMapping(const UE::Editor::DataStorage::FMapKeyView& Key) override;
 
 	virtual FTypedElementOnDataStorageUpdate& OnUpdate() override;
 	virtual FTypedElementOnDataStorageUpdate& OnUpdateCompleted() override;
diff --git a/Engine/Plugins/Experimental/EditorDataStorage/Source/TedsCore/Private/TypedElementDatabaseCompatibility.cpp b/Engine/Plugins/Experimental/EditorDataStorage/Source/TedsCore/Private/TypedElementDatabaseCompatibility.cpp
--- a/Engine/Plugins/Experimental/EditorDataStorage/Source/TedsCore/Private/TypedElementDatabaseCompatibility.cpp
+++ b/Engine/Plugins/Experimental/EditorDataStorage/Source/TedsCore/Private/TypedElementDatabaseCompatibility.cpp
@@ -13,7 +13,6 @@
 #include "Elements/Columns/TypedElementPackageColumns.h"
 #include "Elements/Columns/TypedElementTransformColumns.h"
 #include "Elements/Columns/TypedElementTypeInfoColumns.h"
-#include "Elements/Framework/TypedElementIndexHasher.h"
 #include "Elements/Framework/TypedElementQueryBuilder.h"
 #include "HAL/IConsoleManager.h"
 #include "Logging/LogMacros.h"
@@ -228,7 +227,7 @@ UE::Editor::DataStorage::RowHandle UEditorDataStorageCompatibility::AddCompatibl
 	if (!Storage->IsRowAvailable(Result))
 	{
 		Result = Storage->ReserveRow();
-		Storage->IndexRow(GenerateIndexHash(Object), Result);
+		Storage->MapRow(FMapKey(Object), Result);
 		if (Private::bUseCommandBuffer)
 		{
 			QueuedCommands.AddCommand(FAddCompatibleExternalObject{ .Object = Object, .TypeInfo = TypeInfo, .Row = Result });
@@ -260,8 +259,7 @@ void UEditorDataStorageCompatibility::RemoveCompatibleObjectExplicit(void* Objec
 	{
 		FScopedExclusiveLock Lock(EGlobalLockScope::Public);
 
-		IndexHash Hash = GenerateIndexHash(Object);
-		RowHandle Row = Storage->FindIndexedRow(Hash);
+		RowHandle Row = Storage->LookupMappedRow(FMapKeyView(Object));
 		if (Storage->IsRowAvailable(Row))
 		{
 			const FTypedElementScriptStructTypeInfoColumn* TypeInfoColumn = Storage->GetColumn<FTypedElementScriptStructTypeInfoColumn>(Row);
@@ -282,7 +280,7 @@ UE::Editor::DataStorage::RowHandle UEditorDataStorageCompatibility::FindRowWithC
 	{
 		FScopedSharedLock Lock(EGlobalLockScope::Public);
 
-		RowHandle Row = Storage->FindIndexedRow(GenerateIndexHash(Object));
+		RowHandle Row = Storage->LookupMappedRow(FMapKeyView(Object));
 		return Storage->IsRowAvailable(Row) ? Row : DealiasObject(Object);
 	}
 	return InvalidRowHandle;
@@ -294,7 +292,7 @@ UE::Editor::DataStorage::RowHandle UEditorDataStorageCompatibility::FindRowWithC
 
 	using namespace UE::Editor::DataStorage;
 
-	return (Object && Storage && Storage->IsAvailable()) ? Storage->FindIndexedRow(GenerateIndexHash(Object)) : InvalidRowHandle;
+	return (Object && Storage && Storage->IsAvailable()) ? Storage->LookupMappedRow(FMapKeyView(Object)) : InvalidRowHandle;
 }
 
 bool UEditorDataStorageCompatibility::SupportsExtension(FName Extension) const
@@ -401,7 +399,7 @@ bool UEditorDataStorageCompatibility::ShouldAddObject(const UObject* Object) con
 	FScopedSharedLock Lock(EGlobalLockScope::Public);
 
 	bool Include = true;
-	if (!Storage->IsRowAvailable(Storage->FindIndexedRow(GenerateIndexHash(Object))))
+	if (!Storage->IsRowAvailable(Storage->LookupMappedRow(FMapKeyView(Object))))
 	{
 		const ObjectRegistrationFilter* Filter = ObjectRegistrationFilters.GetData();
 		const ObjectRegistrationFilter* FilterEnd = Filter + ObjectRegistrationFilters.Num();
@@ -442,7 +440,7 @@ UE::Editor::DataStorage::RowHandle UEditorDataStorageCompatibility::AddCompatibl
 	if (!Storage->IsRowAvailable(Result))
 	{
 		Result = Storage->ReserveRow();
-		Storage->IndexRow(GenerateIndexHash(Object), Result);
+		Storage->MapRow(FMapKey(Object), Result);
 		if (Private::bUseCommandBuffer)
 		{
 			QueuedCommands.AddCommand(FAddCompatibleUObject{ .Object = Object, .Row = Result });
@@ -485,9 +483,7 @@ void UEditorDataStorageCompatibility::RemoveCompatibleObjectExplicitTransactiona
 	// non-buffered to reduce the additional locks/unlocks while the buffered version doesn't need any locking beyond the
 	// shared lock FindIndexedRow does. Not adding an exclusive here means some additional lock/unlocking but doesn't make
 	// the code thread unsafe.
-	IndexHash Hash = GenerateIndexHash(Object);
-	RowHandle Row = Storage->FindIndexedRow(Hash);
-
+	RowHandle Row = Storage->LookupMappedRow(FMapKeyView(Object));
 	if (Storage->IsRowAvailable(Row))
 	{
 		RemoveCompatibleObjectExplicitTransactionable<bEnableTransactions>(Object, Row);
diff --git a/Engine/Plugins/Experimental/EditorDataStorage/Source/TedsCore/Private/TypedElementDatabaseEnvironment.cpp b/Engine/Plugins/Experimental/EditorDataStorage/Source/TedsCore/Private/TypedElementDatabaseEnvironment.cpp
--- a/Engine/Plugins/Experimental/EditorDataStorage/Source/TedsCore/Private/TypedElementDatabaseEnvironment.cpp
+++ b/Engine/Plugins/Experimental/EditorDataStorage/Source/TedsCore/Private/TypedElementDatabaseEnvironment.cpp
@@ -9,7 +9,7 @@ namespace UE::Editor::DataStorage
 		FMassEntityManager& InMassEntityManager, FMassProcessingPhaseManager& InMassPhaseManager)
 		: DataStorage(InDataStorage)
 		, DirectDeferredCommands(*this)
-		, IndexTable(InDataStorage)
+		, MappingTable(InDataStorage)
 		, Queries(DynamicColumnGenerator)
 		, MementoSystem(InDataStorage)
 		, ValueTagManager(DynamicColumnGenerator)
@@ -29,14 +29,14 @@ namespace UE::Editor::DataStorage
 		return DirectDeferredCommands;
 	}
 
-	FIndexTable& FEnvironment::GetIndexTable()
+	FMappingTable& FEnvironment::GetMappingTable()
 	{
-		return IndexTable;
+		return MappingTable;
 	}
 
-	const FIndexTable& FEnvironment::GetIndexTable() const
+	const FMappingTable& FEnvironment::GetMappingTable() const
 	{
-		return IndexTable;
+		return MappingTable;
 	}
 
 	FScratchBuffer& FEnvironment::GetScratchBuffer()
@@ -119,7 +119,7 @@ namespace UE::Editor::DataStorage
 		Queries.UpdateActivatableQueries();
 		FlushCommands();
 		ScratchBuffer.BatchDelete();
-		IndexTable.RemoveInvalidRows();
+		MappingTable.RemoveInvalidRows();
 
 		UpdateCycleId++;
 	}
diff --git a/Engine/Plugins/Experimental/EditorDataStorage/Source/TedsCore/Private/TypedElementDatabaseEnvironment.h b/Engine/Plugins/Experimental/EditorDataStorage/Source/TedsCore/Private/TypedElementDatabaseEnvironment.h
--- a/Engine/Plugins/Experimental/EditorDataStorage/Source/TedsCore/Private/TypedElementDatabaseEnvironment.h
+++ b/Engine/Plugins/Experimental/EditorDataStorage/Source/TedsCore/Private/TypedElementDatabaseEnvironment.h
@@ -25,8 +25,8 @@ namespace UE::Editor::DataStorage
 		Legacy::FCommandBuffer& GetDirectDeferredCommands();
 		const Legacy::FCommandBuffer& GetDirectDeferredCommands() const;
 
-		FIndexTable& GetIndexTable();
-		const FIndexTable& GetIndexTable() const;
+		FMappingTable& GetMappingTable();
+		const FMappingTable& GetMappingTable() const;
 
 		FScratchBuffer& GetScratchBuffer();
 		const FScratchBuffer& GetScratchBuffer() const;
@@ -76,7 +76,7 @@ namespace UE::Editor::DataStorage
 		
 		UEditorDataStorage& DataStorage;
 		Legacy::FCommandBuffer DirectDeferredCommands;
-		FIndexTable IndexTable;
+		FMappingTable MappingTable;
 		FScratchBuffer ScratchBuffer;
 		FDynamicColumnGenerator DynamicColumnGenerator;
 		FExtendedQueryStore Queries;
diff --git a/Engine/Plugins/Experimental/EditorDataStorage/Source/TedsCore/Private/TypedElementDatabaseIndexTable.cpp b/Engine/Plugins/Experimental/EditorDataStorage/Source/TedsCore/Private/TypedElementDatabaseIndexTable.cpp
--- a/Engine/Plugins/Experimental/EditorDataStorage/Source/TedsCore/Private/TypedElementDatabaseIndexTable.cpp
+++ b/Engine/Plugins/Experimental/EditorDataStorage/Source/TedsCore/Private/TypedElementDatabaseIndexTable.cpp
@@ -5,96 +5,108 @@
 #include "Async/ParallelFor.h"
 #include "Elements/Common/TypedElementDataStorageLog.h"
 #include "GenericPlatform/GenericPlatformTime.h"
+#include "Hash/CityHash.h"
 #include "TypedElementDatabase.h"
 #include "TypedElementDataStorageProfilingMacros.h"
 
 namespace UE::Editor::DataStorage
 {
-	const FTimespan FIndexTable::FCleanUpInfo::TargetDuration = FTimespan::FromMicroseconds(2000.0);
-	const FTimespan FIndexTable::FCleanUpInfo::JobShrinkThreshold = FTimespan::FromMicroseconds(100);
-	const FTimespan FIndexTable::FCleanUpInfo::JobGrowthThreshold = FTimespan::FromMicroseconds(500);
+	const FTimespan FMappingTable::FCleanUpInfo::TargetDuration = FTimespan::FromMicroseconds(2000.0);
+	const FTimespan FMappingTable::FCleanUpInfo::JobShrinkThreshold = FTimespan::FromMicroseconds(100);
+	const FTimespan FMappingTable::FCleanUpInfo::JobGrowthThreshold = FTimespan::FromMicroseconds(500);
 
-	const FTimespan FIndexTable::FCleanUpInfo::MaxBatchDuration = FTimespan::FromMicroseconds(500);
-	const FTimespan FIndexTable::FCleanUpInfo::MinBatchDuration = FTimespan::FromMicroseconds(350);
+	const FTimespan FMappingTable::FCleanUpInfo::MaxBatchDuration = FTimespan::FromMicroseconds(500);
+	const FTimespan FMappingTable::FCleanUpInfo::MinBatchDuration = FTimespan::FromMicroseconds(350);
 
-	FIndexTable::FIndexTable(UEditorDataStorage& DataStorage)
+	FMappingTable::FMappingTable(UEditorDataStorage& DataStorage)
 		: DataStorage(DataStorage)
 	{
 	}
 
-	RowHandle FIndexTable::FindIndexedRow(
-		EGlobalLockScope LockScope, IndexHash Index) const
-	{	
+	RowHandle FMappingTable::Lookup(EGlobalLockScope LockScope, const FMapKeyView& Key) const
+	{
 		FScopedSharedLock Lock(LockScope);
-	
-		RowHandle Result = InvalidRowHandle;
-		if (const int32* RowIndex = IndexLookupMap.Find(Index))
+
+		if (int32 Index = FindIndexUnguarded(Key); Index >= 0)
 		{
-			Result = Rows[*RowIndex];
+			RowHandle Result = Rows[Index];
 			if (IsDirty() && !DataStorage.IsRowAvailableUnsafe(Result))
 			{
 				Result = InvalidRowHandle;
 			}
+			return Result;
+		}
+		else
+		{
+			return InvalidRowHandle;
 		}
-		return Result;
 	}
 
-	void FIndexTable::BatchIndexRows(EGlobalLockScope LockScope,
-		TConstArrayView<TPair<IndexHash, RowHandle>> IndexRowPairs)
+	void FMappingTable::Map(EGlobalLockScope LockScope, FMapKey Key, RowHandle Row)
 	{
 		FScopedExclusiveLock Lock(LockScope);
-	
-		IndexLookupMap.Reserve(IndexLookupMap.Num() + IndexRowPairs.Num());
-		
-		checkf(Rows.Num() >= FreeList.Num(), 
+		IndexRowUnguarded(MoveTemp(Key), Row);
+	}
+
+	void FMappingTable::BatchMap(EGlobalLockScope LockScope, TArrayView<TPair<FMapKey, RowHandle>> MapRowPairs)
+	{
+		FScopedExclusiveLock Lock(LockScope);
+
+		IndexLookupMap.Reserve(IndexLookupMap.Num() + MapRowPairs.Num());
+
+		checkf(Rows.Num() >= FreeList.Num(),
 			TEXT("There can't be less rows than there are rows stored in the free list as free list is a subset of rows."));
-		int32 ArraySize = Rows.Num() - FreeList.Num() + IndexRowPairs.Num();
+		int32 ArraySize = Rows.Num() - FreeList.Num() + MapRowPairs.Num();
 		Rows.Reserve(ArraySize);
-		Hashes.Reserve(ArraySize);
+		Keys.Reserve(ArraySize);
 
-		for (const TPair<IndexHash, RowHandle>& IndexAndRow : IndexRowPairs)
+		for (TPair<FMapKey, RowHandle>& IndexAndRow : MapRowPairs)
 		{
-			IndexRowUnguarded(IndexAndRow.Key, IndexAndRow.Value);
+			IndexRowUnguarded(MoveTemp(IndexAndRow.Key), IndexAndRow.Value);
 		}
 	}
 
-	void FIndexTable::IndexRow(EGlobalLockScope LockScope,
-		IndexHash Index, RowHandle Row)
+	void FMappingTable::Remap(EGlobalLockScope LockScope, const FMapKeyView& OriginalKey, FMapKey NewKey)
 	{
-
 		FScopedExclusiveLock Lock(LockScope);
-		IndexRowUnguarded(Index, Row);
-	}
 
-	void FIndexTable::Reindex(EGlobalLockScope LockScope,
-		IndexHash OriginalIndex, IndexHash NewIndex)
-	{
-		FScopedExclusiveLock Lock(LockScope);
-	
-		int32 RowIndex;
-		if (IndexLookupMap.RemoveAndCopyValue(OriginalIndex, RowIndex))
+		uint64 Hash = OriginalKey.CalculateHash();
+		for (IndexLookupMapType::TKeyIterator It = IndexLookupMap.CreateKeyIterator(Hash); It; ++It)
 		{
-			IndexLookupMap.Add(NewIndex, RowIndex);
+			if (Keys[It->Value] == OriginalKey)
+			{
+				int32 Index = It->Value;
+				It.RemoveCurrent();
+				IndexLookupMap.Add(NewKey.CalculateHash(), Index);
+				Keys[Index] = MoveTemp(NewKey);
+			}
 		}
 	}
 
-	void FIndexTable::RemoveIndex(EGlobalLockScope LockScope, IndexHash Index)
+	void FMappingTable::Remove(EGlobalLockScope LockScope, const FMapKeyView& Key)
 	{
 		FScopedExclusiveLock Lock(LockScope);
 		
-		int32 RowIndex;
-		if (IndexLookupMap.RemoveAndCopyValue(Index, RowIndex))
+		uint64 Hash = Key.CalculateHash();
+		for (IndexLookupMapType::TKeyIterator It = IndexLookupMap.CreateKeyIterator(Hash); It; ++It)
 		{
-			FreeList.PushFirst(RowIndex);
+			if (Keys[It->Value] == Key)
+			{
+				int32 Index = It->Value;
+				Rows[Index] = InvalidRowHandle;
+				Keys[Index].Clear();
+				FreeList.PushFirst(Index);
+				It.RemoveCurrent();
+			}
 		}
 	}
 
-	void FIndexTable::MarkDirty()
+	void FMappingTable::MarkDirty()
 	{
 		DirtyDueToRemoval = true;
 	}
 
-	void FIndexTable::RemoveInvalidRows()
+	void FMappingTable::RemoveInvalidRows()
 	{
 		TEDS_EVENT_SCOPE("Index Table clean up");
 
@@ -158,31 +170,76 @@ namespace UE::Editor::DataStorage
 			CleanUpInfo.MaxNumJobs);
 	}
 
-	bool FIndexTable::IsDirty() const
+	int32 FMappingTable::FindIndexUnguarded(const FMapKey& Key) const
+	{
+		return FindIndexUnguarded(Key.CalculateHash(), Key);
+	}
+
+	int32 FMappingTable::FindIndexUnguarded(uint64 Hash, const FMapKey& Key) const
+	{
+		for (IndexLookupMapType::TConstKeyIterator It = IndexLookupMap.CreateConstKeyIterator(Hash); It; ++It)
+		{
+			if (Keys[It->Value] == Key)
+			{
+				return It->Value;
+			}
+		}
+		return -1;
+	}
+
+	int32 FMappingTable::FindIndexUnguarded(const FMapKeyView& Key) const
+	{
+		return FindIndexUnguarded(Key.CalculateHash(), Key);
+	}
+
+	int32 FMappingTable::FindIndexUnguarded(uint64 Hash, const FMapKeyView& Key) const
+	{
+		for (IndexLookupMapType::TConstKeyIterator It = IndexLookupMap.CreateConstKeyIterator(Hash); It; ++It)
+		{
+			if (Keys[It->Value] == Key)
+			{
+				return It->Value;
+			}
+		}
+		return -1;
+	}
+
+	bool FMappingTable::IsDirty() const
 	{
 		return CleanUpInfo.RemainingJobs > 0 || DirtyDueToRemoval || !CleanUpInfo.DeletionQueue.IsEmpty();
 	}
 
-	void FIndexTable::IndexRowUnguarded(IndexHash Index, RowHandle Row)
+	void FMappingTable::IndexRowUnguarded(FMapKey&& Key, RowHandle Row)
 	{
+		uint64 Hash = Key.CalculateHash();
+		if (int32 Index = FindIndexUnguarded(Hash, Key); Index >= 0)
+		{
+			bool bUpdatingAllowed = IsDirty() && !DataStorage.IsRowAvailableUnsafe(Rows[Index]);
+			if (ensureMsgf(bUpdatingAllowed || Row == Rows[Index], TEXT("Another row has already been registered under key '%s'."), *Key.ToString()))
+			{
+				Rows[Index] = Row; // Update the stored row to the new row.
+				return;
+			}
+		}
+		
+		// There's no existing row stored under the given key, so create a new one.
 		int32 RowIndex;
 		if (FreeList.IsEmpty())
 		{
 			RowIndex = Rows.Add(Row);
-			Hashes.Add(Index);
+			Keys.Add(MoveTemp(Key));
 		}
 		else
 		{
 			RowIndex = FreeList.Last();
 			FreeList.PopLast();
 			Rows[RowIndex] = Row;
-			Hashes[RowIndex] = Index;
+			Keys[RowIndex] = MoveTemp(Key);
 		}
-		
-		IndexLookupMap.Add(Index, RowIndex);
+		IndexLookupMap.Add(Hash, RowIndex);
 	}
 
-	void FIndexTable::InspectRowBlockForCleanUp(int32 Block)
+	void FMappingTable::InspectRowBlockForCleanUp(int32 Block)
 	{
 		TEDS_EVENT_SCOPE("Index Table inspect rows");
 
@@ -225,7 +282,7 @@ namespace UE::Editor::DataStorage
 		}
 	}
 
-	void FIndexTable::DrainDeletionQueue(FTimespan RemainingFrameTime)
+	void FMappingTable::DrainDeletionQueue(FTimespan RemainingFrameTime)
 	{
 		TEDS_EVENT_SCOPE("Index Table drain deletion queue");
 
@@ -245,30 +302,18 @@ namespace UE::Editor::DataStorage
 		}
 	}
 
-	void FIndexTable::ClearRow(int32 Index)
+	void FMappingTable::ClearRow(int32 Index)
 	{
-		IndexHash Hash = Hashes[Index];
-
-		int32 StoredIndex;
-		if (IndexLookupMap.RemoveAndCopyValue(Hash, StoredIndex) && StoredIndex == Index)
-		{
-			Rows[Index] = InvalidRowHandle;
-			Hashes[Index] = 0;
-
-			FreeList.PushFirst(Index);
-		}
-		else
-		{
-			// TODO: This case was hit in the "Test 'SceneGraph.EditorWorkflows.Non-OFPE.Level Editor.Viewport.Transform.Mirror",
-			//		indicating that hash collisions are happening. This fix ignores those for now at the cost of leaving some
-			//		old values in, but a more sustainable solution needs to be created.
-			// Original code:
-			// bool Result = IndexLookupMap.RemoveAndCopyValue(Hash, StoredIndex);
-			// checkf(Result && StoredIndex == Index, TEXT("Attempting to remove an invalid row but it doesn't exist in the lookup map."));
-		}
+		FMapKey& Key = Keys[Index];
+		uint64 Hash = Key.CalculateHash();
+		IndexLookupMap.RemoveSingle(Hash, Index);
+		
+		Rows[Index] = InvalidRowHandle;
+		Keys[Index].Clear();
+		FreeList.PushFirst(Index);
 	}
 
-	void FIndexTable::AdjustJobCount(FTimespan RemainingFrameTime)
+	void FMappingTable::AdjustJobCount(FTimespan RemainingFrameTime)
 	{
 		if (RemainingFrameTime >= FCleanUpInfo::JobGrowthThreshold)
 		{
@@ -282,7 +327,7 @@ namespace UE::Editor::DataStorage
 		CleanUpInfo.MaxNumJobs = FMath::Clamp(CleanUpInfo.MaxNumJobs, FCleanUpInfo::MinJobCount, FCleanUpInfo::MaxJobCount);
 	}
 
-	void FIndexTable::AdjustBatchSize(int32 JobCount)
+	void FMappingTable::AdjustBatchSize(int32 JobCount)
 	{
 		if (CleanUpInfo.BatchWentOverTime.load() >= FCleanUpInfo::BatchShrinkThreshold)
 		{
diff --git a/Engine/Plugins/Experimental/EditorDataStorage/Source/TedsCore/Private/TypedElementDatabaseIndexTable.h b/Engine/Plugins/Experimental/EditorDataStorage/Source/TedsCore/Private/TypedElementDatabaseIndexTable.h
--- a/Engine/Plugins/Experimental/EditorDataStorage/Source/TedsCore/Private/TypedElementDatabaseIndexTable.h
+++ b/Engine/Plugins/Experimental/EditorDataStorage/Source/TedsCore/Private/TypedElementDatabaseIndexTable.h
@@ -8,6 +8,7 @@
 #include "Containers/Deque.h"
 #include "Containers/MpscQueue.h"
 #include "Elements/Common/TypedElementHandles.h"
+#include "Elements/Common/TypedElementMapKey.h"
 #include "Elements/Interfaces/TypedElementQueryStorageInterfaces.h"
 #include "GlobalLock.h"
 #include "Misc/Timespan.h"
@@ -17,21 +18,21 @@ class UEditorDataStorage;
 namespace UE::Editor::DataStorage
 {
 	/**
-	 * Storage for an index to row mapping.
-	 * Access to the index table is thread safe and guarded by the global lock.
+	 * Storage for a key to row mapping.
+	 * Access to the mapping table is thread safe and guarded by the global lock.
 	 */
 
-	class FIndexTable final
+	class FMappingTable final
 	{
 	public:
-		explicit FIndexTable(UEditorDataStorage& DataStorage);
+		explicit FMappingTable(UEditorDataStorage& DataStorage);
+
+		RowHandle Lookup(EGlobalLockScope LockScope, const FMapKeyView& Key) const;
+		void Map(EGlobalLockScope LockScope, FMapKey Key, RowHandle Row);
+		void BatchMap(EGlobalLockScope LockScope, TArrayView<TPair<FMapKey, RowHandle>> MapRowPairs);
+		void Remap(EGlobalLockScope LockScope, const FMapKeyView& OriginalKey, FMapKey NewKey);
+		void Remove(EGlobalLockScope LockScope, const FMapKeyView& Key);
 
-		RowHandle FindIndexedRow(EGlobalLockScope LockScope, IndexHash Index) const;
-		void IndexRow(EGlobalLockScope LockScope, IndexHash Index, RowHandle Row);
-		void BatchIndexRows(EGlobalLockScope LockScope, TConstArrayView<TPair<IndexHash, RowHandle>> IndexRowPairs);
-		void Reindex(EGlobalLockScope LockScope, IndexHash OriginalIndex, IndexHash NewIndex);
-		void RemoveIndex(EGlobalLockScope LockScope, IndexHash Index);
-		
 		void MarkDirty();
 		void RemoveInvalidRows();
 
@@ -97,18 +98,24 @@ namespace UE::Editor::DataStorage
 			int32 RemainingJobs = 0;
 		};
 		
-		TMap<IndexHash, int32> IndexLookupMap;
+		using IndexLookupMapType = TMultiMap<uint64, int32>;
+		IndexLookupMapType IndexLookupMap;
 		TArray<RowHandle> Rows;
-		TArray<IndexHash> Hashes;
+		TArray<FMapKey> Keys;
 		TDeque<int32> FreeList;
-		
+
 		FCleanUpInfo CleanUpInfo;
 		
 		UEditorDataStorage& DataStorage;
 		std::atomic<bool> DirtyDueToRemoval = false;
 
+		int32 FindIndexUnguarded(const FMapKey& Key) const;
+		int32 FindIndexUnguarded(uint64 Hash, const FMapKey& Key) const;
+		int32 FindIndexUnguarded(const FMapKeyView& Key) const;
+		int32 FindIndexUnguarded(uint64 Hash, const FMapKeyView& Key) const;
+
 		bool IsDirty() const;
-		void IndexRowUnguarded(IndexHash Index, RowHandle Row);
+		void IndexRowUnguarded(FMapKey&& Key, RowHandle Row);
 		
 		void InspectRowBlockForCleanUp(int32 Block);
 		void DrainDeletionQueue(FTimespan RemainingFrameTime);
diff --git a/Engine/Plugins/Experimental/EditorDataStorageFeatures/Source/TedsAssetData/Private/Debug/TedsAssetDataDebugCommands.cpp b/Engine/Plugins/Experimental/EditorDataStorageFeatures/Source/TedsAssetData/Private/Debug/TedsAssetDataDebugCommands.cpp
--- a/Engine/Plugins/Experimental/EditorDataStorageFeatures/Source/TedsAssetData/Private/Debug/TedsAssetDataDebugCommands.cpp
+++ b/Engine/Plugins/Experimental/EditorDataStorageFeatures/Source/TedsAssetData/Private/Debug/TedsAssetDataDebugCommands.cpp
@@ -5,7 +5,6 @@
 #include "Containers/UnrealString.h"
 #include "Elements/Common/TypedElementHandles.h"
 #include "Elements/Common/EditorDataStorageFeatures.h"
-#include "Elements/Framework/TypedElementIndexHasher.h"
 #include "HAL/IConsoleManager.h"
 #include "TedsAssetDataColumns.h"
 #include "UObject/NameTypes.h"
@@ -36,7 +35,7 @@ namespace UE::Editor::DataStorage::Debug::Private
 				for (const FString& Path : InArgs)
 				{
 					FName PathAsName(*Path);
-					RowHandle Row = Database->FindIndexedRow(GenerateIndexHash(PathAsName));
+					RowHandle Row = Database->LookupMappedRow(FMapKeyView(PathAsName));
 
 					if (Database->IsRowAssigned(Row))
 					{
@@ -98,7 +97,7 @@ namespace UE::Editor::DataStorage::Debug::Private
 					AssetRegistry.GetAssetsByPath(PathAsName, Assets);
 					for (const FAssetData& Asset : Assets)
 					{
-						RowHandle AssetRow = Database->FindIndexedRow(GenerateIndexHash(Asset.GetSoftObjectPath()));
+						RowHandle AssetRow = Database->LookupMappedRow(FMapKey(Asset.GetSoftObjectPath()));
 
 						if (const FUnresolvedAssetsInPathColumn_Experimental* UnresolvedAssetsInPathColumn = Database->GetColumn<FUnresolvedAssetsInPathColumn_Experimental>(AssetRow))
 						{
diff --git a/Engine/Plugins/Experimental/EditorDataStorageFeatures/Source/TedsAssetData/Private/Processors/AssetProcessors.cpp b/Engine/Plugins/Experimental/EditorDataStorageFeatures/Source/TedsAssetData/Private/Processors/AssetProcessors.cpp
--- a/Engine/Plugins/Experimental/EditorDataStorageFeatures/Source/TedsAssetData/Private/Processors/AssetProcessors.cpp
+++ b/Engine/Plugins/Experimental/EditorDataStorageFeatures/Source/TedsAssetData/Private/Processors/AssetProcessors.cpp
@@ -15,7 +15,6 @@
 #include "Experimental/ContentBrowserExtensionUtils.h"
 #include "TedsAssetDataColumns.h"
 #include "Elements/Columns/TypedElementMiscColumns.h"
-#include "Elements/Framework/TypedElementIndexHasher.h"
 #include "String/LexFromString.h"
 #include "TedsAlerts.h"
 
@@ -208,8 +207,8 @@ void UTedsAssetDataFactory::PreShutdown(IEditorDataStorageProvider& DataStorage)
 
 void UTedsAssetDataFactory::OnSetFolderColor(const FString& Path, IEditorDataStorageProvider* DataStorage)
 {
-	const UE::Editor::DataStorage::IndexHash PathHash = UE::Editor::DataStorage::GenerateIndexHash(FName(Path));
-	const UE::Editor::DataStorage::RowHandle Row = DataStorage->FindIndexedRow(PathHash);
+	const UE::Editor::DataStorage::FMapKey PathKey = UE::Editor::DataStorage::FMapKey(FName(Path));
+	const UE::Editor::DataStorage::RowHandle Row = DataStorage->LookupMappedRow(PathKey);
 
 	if (DataStorage->IsRowAvailable(Row))
 	{
diff --git a/Engine/Plugins/Experimental/EditorDataStorageFeatures/Source/TedsAssetData/Private/TedsAssetData.cpp b/Engine/Plugins/Experimental/EditorDataStorageFeatures/Source/TedsAssetData/Private/TedsAssetData.cpp
--- a/Engine/Plugins/Experimental/EditorDataStorageFeatures/Source/TedsAssetData/Private/TedsAssetData.cpp
+++ b/Engine/Plugins/Experimental/EditorDataStorageFeatures/Source/TedsAssetData/Private/TedsAssetData.cpp
@@ -14,7 +14,6 @@
 #include "Elements/Columns/TypedElementMiscColumns.h"
 #include "Elements/Columns/TypedElementSlateWidgetColumns.h"
 #include "Elements/Common/TypedElementQueryTypes.h"
-#include "Elements/Framework/TypedElementIndexHasher.h"
 #include "Elements/Framework/TypedElementQueryBuilder.h"
 #include "Elements/Interfaces/TypedElementDataStorageInterface.h"
 #include "Engine/Engine.h"
@@ -31,8 +30,8 @@ struct FPopulateAssetDataRowArgs
 {
 	FAssetData AssetData;
 
-	IndexHash ObjectPathHash;
-	IndexHash PathHash;
+	FMapKey ObjectPathKey;
+	FMapKey PathKey;
 	RowHandle PathRow = InvalidRowHandle;
 };
 
@@ -41,19 +40,19 @@ template<typename TAssetData>
 FPopulateAssetDataRowArgs ThreadSafe_PopulateAssetDataTableRow(TAssetData&& InAssetData, const IEditorDataStorageProvider& Database)
 {
 	FPopulateAssetDataRowArgs Output;
-	Output.ObjectPathHash = GenerateIndexHash(InAssetData.GetSoftObjectPath());
+	Output.ObjectPathKey = FMapKey(InAssetData.GetSoftObjectPath());
 
 	// Looks safe but might not be depending on the implementation of the database
-	if (Database.IsRowAssigned(Database.FindIndexedRow(Output.ObjectPathHash)))
+	if (Database.IsRowAssigned(Database.LookupMappedRow(Output.ObjectPathKey)))
 	{
 		// No need to initialize the rest of the row here. The invalid's asset data will be used as flag to skip the data generated here.
 		return Output;
 	}
 
-	Output.PathHash = GenerateIndexHash(InAssetData.PackagePath);
+	Output.PathKey = FMapKey(InAssetData.PackagePath);
 
 	// Looks safe but might not be depending on the implementation of the database
-	Output.PathRow = Database.FindIndexedRow(Output.PathHash);
+	Output.PathRow = Database.LookupMappedRow(Output.PathKey);
 	Output.AssetData = MoveTempIfPossible(InAssetData);
 
 	return Output;
@@ -62,26 +61,26 @@ FPopulateAssetDataRowArgs ThreadSafe_PopulateAssetDataTableRow(TAssetData&& InAs
 
 void PopulateAssetDataTableRow(FPopulateAssetDataRowArgs&& InAssetDataRowArgs, IEditorDataStorageProvider& Database, RowHandle RowHandle)
 {
-	if (InAssetDataRowArgs.PathRow !=  InvalidRowHandle)
+	if (InAssetDataRowArgs.PathRow != InvalidRowHandle)
 	{
 		Database.GetColumn<FAssetsInPathColumn_Experimental>(InAssetDataRowArgs.PathRow)->AssetsRow.Add(RowHandle);
 	}
 	else
 	{
 		FUnresolvedAssetsInPathColumn_Experimental UnresolvedAssetsInPathColumn;
-		UnresolvedAssetsInPathColumn.Hash = InAssetDataRowArgs.PathHash;
+		UnresolvedAssetsInPathColumn.Key = MoveTemp(InAssetDataRowArgs.PathKey);
 		Database.AddColumn<FUnresolvedAssetsInPathColumn_Experimental>(RowHandle, MoveTemp(UnresolvedAssetsInPathColumn));
 	}
 
-	Database.GetColumn<FAssetDataColumn_Experimental>(RowHandle)->AssetData =  MoveTemp(InAssetDataRowArgs.AssetData);
+	Database.GetColumn<FAssetDataColumn_Experimental>(RowHandle)->AssetData = MoveTemp(InAssetDataRowArgs.AssetData);
 }
 
 
 struct FPopulatePathRowArgs
 {
 	FName AssetRegistryPath;
-	IndexHash AssetRegistryPathHash;
-	IndexHash ParentAssetRegistryPathHash;
+	FMapKey AssetRegistryPathKey;
+	FMapKey ParentAssetRegistryPathKey;
 	FName AssetName;
 	uint32 PathDepth;
 
@@ -119,9 +118,9 @@ void GetPathDepthAndParentFolderIndex(FStringView Path, uint32& OutDepth, int32&
 }
 
 // Only thread safe if the game thread is blocked 
-FPopulatePathRowArgs ThreadSafe_PopulatePathRowArgs(IndexHash AssetRegistryPathHash, FName InAssetRegistryPath, FStringView PathAsString)
+FPopulatePathRowArgs ThreadSafe_PopulatePathRowArgs(FMapKey&& AssetRegistryPathKey, FName InAssetRegistryPath, FStringView PathAsString)
 {
-	IndexHash ParentAssetRegistryPathHash = InvalidRowHandle;
+	FMapKey ParentAssetRegistryPathKey;
 	int32 CharacterIndex;
 	uint32 Depth;
 	FString AssetName;
@@ -129,15 +128,15 @@ FPopulatePathRowArgs ThreadSafe_PopulatePathRowArgs(IndexHash AssetRegistryPathH
 	if (CharacterIndex != INDEX_NONE)
 	{
 		FStringView ParentPath = PathAsString.Left(CharacterIndex);
-		ParentAssetRegistryPathHash = GenerateIndexHash(FName(ParentPath));
+		ParentAssetRegistryPathKey = FMapKey(FName(ParentPath));
 		AssetName = PathAsString.RightChop(CharacterIndex);
 	}
 
 	FPopulatePathRowArgs Args;
 	Args.AssetRegistryPath = InAssetRegistryPath;
 	Args.AssetName = FName(AssetName);
-	Args.AssetRegistryPathHash = AssetRegistryPathHash;
-	Args.ParentAssetRegistryPathHash = ParentAssetRegistryPathHash;
+	Args.AssetRegistryPathKey = MoveTemp(AssetRegistryPathKey);
+	Args.ParentAssetRegistryPathKey = MoveTemp(ParentAssetRegistryPathKey);
 	Args.PathDepth = Depth;
 
 	return Args;
@@ -145,9 +144,9 @@ FPopulatePathRowArgs ThreadSafe_PopulatePathRowArgs(IndexHash AssetRegistryPathH
 
 void PopulatePathDataTableRow(FPopulatePathRowArgs&& InPopulatePathRowArgs, IEditorDataStorageProvider& Database, RowHandle InRowHandle)
 {
-	if (InPopulatePathRowArgs.ParentAssetRegistryPathHash != InvalidRowHandle)
+	if (InPopulatePathRowArgs.ParentAssetRegistryPathKey.IsSet())
 	{
-		const RowHandle ParentRow = Database.FindIndexedRow(InPopulatePathRowArgs.ParentAssetRegistryPathHash);
+		const RowHandle ParentRow = Database.LookupMappedRow(InPopulatePathRowArgs.ParentAssetRegistryPathKey);
 		if (Database.IsRowAssigned(ParentRow))
 		{
 			Database.GetColumn<FChildrenAssetPathColumn_Experimental>(ParentRow)->ChildrenRows.Add(InRowHandle);
@@ -157,7 +156,7 @@ void PopulatePathDataTableRow(FPopulatePathRowArgs&& InPopulatePathRowArgs, IEdi
 		{
 			// If we were unlucky we may be missing the parent path data.This should resolve it self during a latter event on paths added event.
 			FUnresolvedParentAssetPathColumn_Experimental UnresolvedParentRow;
-			UnresolvedParentRow.Hash = InPopulatePathRowArgs.ParentAssetRegistryPathHash;
+			UnresolvedParentRow.Key = MoveTemp(InPopulatePathRowArgs.ParentAssetRegistryPathKey);
 			Database.AddColumn(InRowHandle, MoveTemp(UnresolvedParentRow));
 		}
 	}
@@ -256,7 +255,7 @@ FTedsAssetData::FTedsAssetData(IEditorDataStorageProvider& InDatabase)
 				LLM_SCOPE_BYNAME(TEXT("FTedsAssetData"))
 #endif
 
-				const RowHandle PathRow = Context.FindIndexedRow(UnresolvedAssetPath.Hash);
+				const RowHandle PathRow = Context.LookupMappedRow(UnresolvedAssetPath.Key);
 				if (Context.IsRowAssigned(PathRow))
 				{
 					Context.RemoveColumns<FUnresolvedAssetsInPathColumn_Experimental>(Row);
@@ -290,7 +289,7 @@ FTedsAssetData::FTedsAssetData(IEditorDataStorageProvider& InDatabase)
 				LLM_SCOPE_BYNAME(TEXT("FTedsAssetData"))
 #endif
 
-				const RowHandle ParentPathRow = Context.FindIndexedRow(UnresolvedParentAssetPath.Hash);
+				const RowHandle ParentPathRow = Context.LookupMappedRow(UnresolvedParentAssetPath.Key);
 				if (Context.IsRowAssigned(ParentPathRow))
 				{
 					Context.RemoveColumns<FUnresolvedParentAssetPathColumn_Experimental>(Row);
@@ -324,31 +323,31 @@ FTedsAssetData::FTedsAssetData(IEditorDataStorageProvider& InDatabase)
 
 	TArray<FNameBuilder> NameBuilders;
 	TArray<FPopulatePathRowArgs> PopulatePathRowArgs;
-	PopulatePathRowArgs.AddUninitialized(CachedPaths.Num());
+	PopulatePathRowArgs.AddDefaulted(CachedPaths.Num());
 
 	// Prepare Path Rows
 	ParallelForWithTaskContext(TEXT("Populating TEDS Asset Registry Path"), NameBuilders, CachedPaths.Num(),ParallelForMinBatchSize, [&PopulatePathRowArgs, &CachedPaths, this](FNameBuilder& NameBuilder, int32 Index)
 		{
 			FName Path = CachedPaths[Index];
 			Path.ToString(NameBuilder);
-			PopulatePathRowArgs[Index] = ThreadSafe_PopulatePathRowArgs(GenerateIndexHash(Path), Path, NameBuilder);
+			PopulatePathRowArgs[Index] = ThreadSafe_PopulatePathRowArgs(FMapKey(Path), Path, NameBuilder);
 		});
 
 	TArray<RowHandle> ReservedRows;
 	ReservedRows.AddUninitialized(PopulatePathRowArgs.Num() + AssetsData.Num());
 	Database.BatchReserveRows(ReservedRows);
 
-	TArray<TPair<IndexHash, RowHandle>> IndexesToReservedRows;
-	IndexesToReservedRows.AddUninitialized(ReservedRows.Num());
+	TArray<TPair<FMapKey, RowHandle>> KeysToReservedRows;
+	KeysToReservedRows.AddDefaulted(ReservedRows.Num());
 
 	// Index Reserved Path Rows
 	TConstArrayView<RowHandle> ReservedPopulatePathRows(ReservedRows.GetData(), PopulatePathRowArgs.Num());
-	TArrayView<TPair<IndexHash, RowHandle>> IndexesToReservedPathRows(IndexesToReservedRows.GetData(), PopulatePathRowArgs.Num());
-	ParallelFor(TEXT("Populating TEDS Asset Registry Path Data Indexes"), ReservedPopulatePathRows.Num(), ParallelForMinBatchSize, [&IndexesToReservedPathRows, &ReservedPopulatePathRows, &PopulatePathRowArgs](int32 Index)
+	TArrayView<TPair<FMapKey, RowHandle>> KeysToReservedPathRows(KeysToReservedRows.GetData(), PopulatePathRowArgs.Num());
+	ParallelFor(TEXT("Populating TEDS Asset Registry Path Data Keys"), ReservedPopulatePathRows.Num(), ParallelForMinBatchSize, [&KeysToReservedPathRows, &ReservedPopulatePathRows, &PopulatePathRowArgs](int32 Index)
 		{
-			IndexesToReservedPathRows[Index] = TPair<IndexHash, RowHandle>(PopulatePathRowArgs[Index].AssetRegistryPathHash, ReservedPopulatePathRows[Index]);
+			KeysToReservedPathRows[Index] = TPair<FMapKey, RowHandle>(PopulatePathRowArgs[Index].AssetRegistryPathKey, ReservedPopulatePathRows[Index]);
 		});
-	Database.BatchIndexRows(IndexesToReservedPathRows);
+	Database.BatchMapRows(KeysToReservedPathRows);
 
 	// Populate Path Rows
 	Database.BatchAddRow(PathsTable, ReservedPopulatePathRows, [PathRowArgs = MoveTemp(PopulatePathRowArgs), Index = 0, this](RowHandle InRowHandle) mutable
@@ -367,12 +366,12 @@ FTedsAssetData::FTedsAssetData(IEditorDataStorageProvider& InDatabase)
 
 	// Index Reserved Asset Data Rows
 	TConstArrayView<RowHandle> ReservedAssetDataRows(ReservedRows.GetData() + ReservedPopulatePathRows.Num(), AssetsData.Num());
-	TArrayView<TPair<IndexHash, RowHandle>>  IndexesToResevedAssetRows(IndexesToReservedRows.GetData() + ReservedPopulatePathRows.Num(), ReservedAssetDataRows.Num());
-	ParallelFor(TEXT("Populating TEDS Asset Registry Asset Data Indexes"), PopulateAssetDataRowArgs.Num(), ParallelForMinBatchSize, [&PopulateAssetDataRowArgs, &ReservedAssetDataRows, &IndexesToResevedAssetRows](int32 Index)
+	TArrayView<TPair<FMapKey, RowHandle>> KeysToResevedAssetRows(KeysToReservedRows.GetData() + ReservedPopulatePathRows.Num(), ReservedAssetDataRows.Num());
+	ParallelFor(TEXT("Populating TEDS Asset Registry Asset Data Keys"), PopulateAssetDataRowArgs.Num(), ParallelForMinBatchSize, [&PopulateAssetDataRowArgs, &ReservedAssetDataRows, &KeysToResevedAssetRows](int32 Index)
 		{
-			IndexesToResevedAssetRows[Index] = TPair<IndexHash, RowHandle>(PopulateAssetDataRowArgs[Index].ObjectPathHash, ReservedAssetDataRows[Index]);
+			KeysToResevedAssetRows[Index] = TPair<FMapKey, RowHandle>(PopulateAssetDataRowArgs[Index].ObjectPathKey, ReservedAssetDataRows[Index]);
 		});
-	Database.BatchIndexRows(IndexesToResevedAssetRows);
+	Database.BatchMapRows(KeysToResevedAssetRows);
 
 	// Populate Asset Rows
 	Database.BatchAddRow(AssetsDataTable, ReservedAssetDataRows,  [AssetDataRowArgs = MoveTemp(PopulateAssetDataRowArgs), Index = 0, this](RowHandle InRowHandle) mutable
@@ -413,19 +412,17 @@ FTedsAssetData::~FTedsAssetData()
 	
 			AssetRegistry->EnumerateAllCachedPaths([this](FName InPath)
 				{
-					const IndexHash PathHash = GenerateIndexHash(InPath);
-					const RowHandle Row = Database.FindIndexedRow(PathHash);
+					const FMapKeyView PathKey = FMapKeyView(InPath);
+					const RowHandle Row = Database.LookupMappedRow(PathKey);
 					Database.RemoveRow(Row);
-					Database.RemoveIndex(PathHash);
 					return true;
 				});
 
 			AssetRegistry->EnumerateAllAssets([this](const FAssetData& InAssetData)
 				{
-					const IndexHash AssetPathHash = GenerateIndexHash(InAssetData.GetSoftObjectPath());
-					const RowHandle Row = Database.FindIndexedRow(AssetPathHash);
+					const FMapKey AssetPathKey = FMapKey(InAssetData.GetSoftObjectPath());
+					const RowHandle Row = Database.LookupMappedRow(AssetPathKey);
 					Database.RemoveRow(Row);
-					Database.RemoveIndex(AssetPathHash);
 					return true;
 				});
 		}
@@ -480,9 +477,9 @@ void FTedsAssetData::OnAssetsAdded(TConstArrayView<FAssetData> InAssetsAdded)
 	int32 Index = 0;
 	if (NewRowsCount > 0)
 	{
-		TArray<TPair<IndexHash, RowHandle>> IndexToRow;
-		IndexToRow.Reserve(NewRowsCount);
-		Database.BatchAddRow(AssetsDataTable, NewRowsCount, [RowArgs = MoveTemp(PopulateRowArgs), Index = 0, &IndexToRow, this](RowHandle InRowHandle) mutable
+		TArray<TPair<FMapKey, RowHandle>> KeyToRow;
+		KeyToRow.Reserve(NewRowsCount);
+		Database.BatchAddRow(AssetsDataTable, NewRowsCount, [RowArgs = MoveTemp(PopulateRowArgs), Index = 0, &KeyToRow, this](RowHandle InRowHandle) mutable
 			{
 				FPopulateAssetDataRowArgs ARowArgs = MoveTemp(RowArgs[Index]);
 				while (!ARowArgs.AssetData.IsValid())
@@ -491,12 +488,12 @@ void FTedsAssetData::OnAssetsAdded(TConstArrayView<FAssetData> InAssetsAdded)
 					ARowArgs = MoveTemp(RowArgs[Index]);
 				}
 
-				IndexToRow.Emplace(ARowArgs.ObjectPathHash, InRowHandle);
+				KeyToRow.Emplace(ARowArgs.ObjectPathKey, InRowHandle);
 				PopulateAssetDataTableRow(MoveTemp(ARowArgs), Database, InRowHandle);
 				++Index;
 			});
 
-		Database.BatchIndexRows(IndexToRow);
+		Database.BatchMapRows(KeyToRow);
 	}
 }
 
@@ -510,21 +507,20 @@ void FTedsAssetData::OnAssetsRemoved(TConstArrayView<FAssetData> InAssetsRemoved
 
 	for (const FAssetData& Asset : InAssetsRemoved)
 	{
-		const IndexHash AssetHash = GenerateIndexHash(Asset.GetSoftObjectPath());
-		const RowHandle AssetRow = Database.FindIndexedRow(AssetHash);
+		const FMapKey AssetKey = FMapKey(Asset.GetSoftObjectPath());
+		const RowHandle AssetRow = Database.LookupMappedRow(AssetKey);
 		if (Database.IsRowAssigned(AssetRow))
 		{
 			if (const FAssetDataColumn_Experimental* AssetDataColunm = Database.GetColumn<FAssetDataColumn_Experimental>(AssetRow))
 			{
-				const IndexHash FolderPathHash = GenerateIndexHash(AssetDataColunm->AssetData.PackagePath);
-				const RowHandle FolderRow = Database.FindIndexedRow(FolderPathHash);
+				const FMapKey FolderPathKey = FMapKey(AssetDataColunm->AssetData.PackagePath);
+				const RowHandle FolderRow = Database.LookupMappedRow(FolderPathKey);
 				if (FAssetsInPathColumn_Experimental* AssetInFolder = Database.GetColumn<FAssetsInPathColumn_Experimental>(FolderRow))
 				{
 					AssetInFolder->AssetsRow.Remove(AssetRow);
 				}
 			}
-			Database.RemoveRow(Database.FindIndexedRow(AssetHash));
-			Database.RemoveIndex(AssetHash);
+			Database.RemoveRow(AssetRow);
 		}
 	}
 }
@@ -539,8 +535,8 @@ void FTedsAssetData::OnAssetsUpdated(TConstArrayView<FAssetData> InAssetsUpdated
 
 	for (const FAssetData& Asset : InAssetsUpdated)
 	{
-		const IndexHash AssetHash = GenerateIndexHash(Asset.GetSoftObjectPath());
-		const RowHandle Row = Database.FindIndexedRow(AssetHash);
+		const FMapKey AssetKey = FMapKey(Asset.GetSoftObjectPath());
+		const RowHandle Row = Database.LookupMappedRow(AssetKey);
 		if (Database.IsRowAssigned(Row))
 		{
 			Database.GetColumn<FAssetDataColumn_Experimental>(Row)->AssetData = Asset;
@@ -559,8 +555,8 @@ void FTedsAssetData::OnAssetsUpdatedOnDisk(TConstArrayView<FAssetData> InAssetsU
 
 	for (const FAssetData& Asset : InAssetsUpdated)
 	{
-		const IndexHash AssetHash = GenerateIndexHash(Asset.GetSoftObjectPath());
-		const RowHandle Row = Database.FindIndexedRow(AssetHash);
+		const FMapKey AssetKey = FMapKey(Asset.GetSoftObjectPath());
+		const RowHandle Row = Database.LookupMappedRow(AssetKey);
 		if (Database.IsRowAssigned(Row))
 		{
 			Database.GetColumn<FAssetDataColumn_Experimental>(Row)->AssetData = Asset;
@@ -577,15 +573,17 @@ void FTedsAssetData::OnAssetRenamed(const FAssetData& InAsset, const FString& In
 	LLM_SCOPE_BYNAME(TEXT("FTedsAssetData"))
 #endif
 
-	const IndexHash NewAssetHash = GenerateIndexHash(InAsset.GetSoftObjectPath());
-	const IndexHash OldAssetHash = GenerateIndexHash(FSoftObjectPath(InOldObjectPath));
-	const RowHandle Row = Database.FindIndexedRow(OldAssetHash);
+	using namespace UE::Editor::DataStorage;
+
+	FMapKey NewAssetHash = FMapKey(InAsset.GetSoftObjectPath());
+	const FMapKey OldAssetHash = FMapKey(FSoftObjectPath(InOldObjectPath));
+	const RowHandle Row = Database.LookupMappedRow(OldAssetHash);
 	if (Database.IsRowAssigned(Row))
 	{
 		Database.GetColumn<FAssetDataColumn_Experimental>(Row)->AssetData = InAsset;
 
 		// Update the asset in folder columns
-		const IndexHash NewFolderHash = GenerateIndexHash(InAsset.PackagePath);
+		const FMapKeyView NewFolderHash(InAsset.PackagePath);
 		FStringView OldPackagePath(InOldObjectPath);
 
 		{
@@ -594,11 +592,11 @@ void FTedsAssetData::OnAssetRenamed(const FAssetData& InAsset, const FString& In
 			OldPackagePath.LeftInline(CharacterIndex);
 		}
 
-		const IndexHash OldFolderHash = GenerateIndexHash(FName(OldPackagePath));
+		const FMapKey OldFolderHash = FMapKey(FName(OldPackagePath));
 
 		if (NewFolderHash != OldFolderHash)
 		{
-			const RowHandle NewPathRow = Database.FindIndexedRow(NewFolderHash);
+			const RowHandle NewPathRow = Database.LookupMappedRow(NewFolderHash);
 			if (NewPathRow !=  InvalidRowHandle)
 			{
 				Database.GetColumn<FAssetsInPathColumn_Experimental>(NewPathRow)->AssetsRow.Add(Row);
@@ -606,11 +604,11 @@ void FTedsAssetData::OnAssetRenamed(const FAssetData& InAsset, const FString& In
 			else
 			{
 				FUnresolvedAssetsInPathColumn_Experimental UnresolvedAssetsInPathColumn;
-				UnresolvedAssetsInPathColumn.Hash = NewFolderHash;
+				UnresolvedAssetsInPathColumn.Key = NewFolderHash.CreateKey();
 				Database.AddColumn<FUnresolvedAssetsInPathColumn_Experimental>(Row, MoveTemp(UnresolvedAssetsInPathColumn));
 			}
 
-			const RowHandle OldPathRow = Database.FindIndexedRow(OldFolderHash);
+			const RowHandle OldPathRow = Database.LookupMappedRow(OldFolderHash);
 			if (NewPathRow !=  InvalidRowHandle)
 			{
 				Database.GetColumn<FAssetsInPathColumn_Experimental>(NewPathRow)->AssetsRow.Remove(Row);
@@ -622,7 +620,7 @@ void FTedsAssetData::OnAssetRenamed(const FAssetData& InAsset, const FString& In
 		}
 
 		Database.AddColumn<FUpdatedPathTag>(Row);
-		Database.Reindex(OldAssetHash, NewAssetHash);
+		Database.RemapRow(OldAssetHash, MoveTemp(NewAssetHash));
 	}
 }
 
@@ -642,16 +640,16 @@ void FTedsAssetData::OnPathsAdded(TConstArrayView<FStringView> InPathsAdded)
 		{
 			FStringView Path = InPathsAdded[Index];
 			FName PathName(Path);
-			const IndexHash AssetRegistryPathHash = GenerateIndexHash(PathName);
+			FMapKey AssetRegistryPathKey = FMapKey(PathName);
 			FPopulatePathRowArgs RowArgs;
 
-			if (Database.FindIndexedRow(AssetRegistryPathHash) != InvalidRowHandle)
+			if (Database.LookupMappedRow(AssetRegistryPathKey) != InvalidRowHandle)
 			{
 				RowArgs.MarkAsInvalid();
 			}
 			else
 			{
-				RowArgs = ThreadSafe_PopulatePathRowArgs(AssetRegistryPathHash, PathName, Path);
+				RowArgs = ThreadSafe_PopulatePathRowArgs(MoveTemp(AssetRegistryPathKey), PathName, Path);
 				++WorkerValidCount;
 			}
 
@@ -671,8 +669,8 @@ void FTedsAssetData::OnPathsAdded(TConstArrayView<FStringView> InPathsAdded)
 		ReservedRow.AddUninitialized(NewRowsCount);
 		Database.BatchReserveRows(ReservedRow);
 
-		TArray<TPair<IndexHash, RowHandle>> IndexesAndRows;
-		IndexesAndRows.Reserve(NewRowsCount);
+		TArray<TPair<FMapKey, RowHandle>> KeysAndRows;
+		KeysAndRows.Reserve(NewRowsCount);
 
 		int32 RowCount = 0;
 		for (int32 Index = 0; Index < NewRowsCount; ++Index)
@@ -684,11 +682,11 @@ void FTedsAssetData::OnPathsAdded(TConstArrayView<FStringView> InPathsAdded)
 				RowArgs = &PopulateRowArgs[RowCount];
 			}
 
-			IndexesAndRows.Emplace(RowArgs->AssetRegistryPathHash, ReservedRow[Index]);
+			KeysAndRows.Emplace(RowArgs->AssetRegistryPathKey, ReservedRow[Index]);
 			++RowCount;
 		}
 
-		Database.BatchIndexRows(IndexesAndRows);
+		Database.BatchMapRows(KeysAndRows);
 
 		int32 Index = 0;
 		Database.BatchAddRow(PathsTable, ReservedRow, [&PopulateRowArgs, &Index, this](RowHandle InRowHandle)
@@ -716,11 +714,10 @@ void FTedsAssetData::OnPathsRemoved(TConstArrayView<FStringView> InPathsRemoved)
 
 	for (const FStringView Path : InPathsRemoved)
 	{
-		const IndexHash PathHash = GenerateIndexHash(FName(Path));
-		Database.RemoveRow(Database.FindIndexedRow(PathHash));
-		Database.RemoveIndex(PathHash);
+		FMapKey PathKey = FMapKey(FName(Path));
+		Database.RemoveRow(Database.LookupMappedRow(PathKey));
 	}
 }
 } // namespace UE::Editor::AssetData::Private
 
-#undef TRACK_TEDSASSETDATA_MEMORY
\ No newline at end of file
+#undef TRACK_TEDSASSETDATA_MEMORY
diff --git a/Engine/Plugins/Experimental/EditorDataStorageFeatures/Source/TedsAssetData/Private/Tests/TedsAssetDataTest.cpp b/Engine/Plugins/Experimental/EditorDataStorageFeatures/Source/TedsAssetData/Private/Tests/TedsAssetDataTest.cpp
--- a/Engine/Plugins/Experimental/EditorDataStorageFeatures/Source/TedsAssetData/Private/Tests/TedsAssetDataTest.cpp
+++ b/Engine/Plugins/Experimental/EditorDataStorageFeatures/Source/TedsAssetData/Private/Tests/TedsAssetDataTest.cpp
@@ -9,7 +9,6 @@
 #include "Containers/UnrealString.h"
 #include "Elements/Common/EditorDataStorageFeatures.h"
 #include "Elements/Common/TypedElementHandles.h"
-#include "Elements/Framework/TypedElementIndexHasher.h"
 #include "Elements/Framework/TypedElementRegistry.h"
 #include "Elements/Interfaces/TypedElementDataStorageInterface.h"
 #include "Misc/ScopeExit.h"
@@ -71,7 +70,7 @@ bool FTedsAssetDataTest::RunTest(const FString& Parameters)
 	
 	for (const TPair<FName, int32>& PathAndAssetCount : AssetRegistryPathsAndAssetCount)
 	{
-		RowHandle Row = Database->FindIndexedRow(GenerateIndexHash(PathAndAssetCount.Key));
+		RowHandle Row = Database->LookupMappedRow(FMapKeyView(PathAndAssetCount.Key));
 		bool bHadError = false;
 		bHadError |= !TestTrue(TEXT("Asset registry folder/path is indexed in TEDS"), Row != InvalidRowHandle);
 
@@ -109,7 +108,7 @@ bool FTedsAssetDataTest::RunTest(const FString& Parameters)
 		// Verse package don't emit asset added event so we will ignore those for the test to avoid false errors
 		if (!FPackageName::IsVersePackage(Builder))
 		{ 
-			RowHandle Row = Database->FindIndexedRow(GenerateIndexHash(AssetData.GetSoftObjectPath()));
+			RowHandle Row = Database->LookupMappedRow(FMapKey(AssetData.GetSoftObjectPath()));
 
 			bool bHadError = false;
 
@@ -119,7 +118,7 @@ bool FTedsAssetDataTest::RunTest(const FString& Parameters)
 
 			if (!Database->GetColumn<FUnresolvedAssetsInPathColumn_Experimental>(Row))
 			{
-				RowHandle PathRow = Database->FindIndexedRow(GenerateIndexHash(AssetData.PackagePath));
+				RowHandle PathRow = Database->LookupMappedRow(FMapKeyView(AssetData.PackagePath));
 			
 				const FAssetsInPathColumn_Experimental* AssetsInPathColumn = Database->GetColumn<FAssetsInPathColumn_Experimental>(PathRow);
 
diff --git a/Engine/Plugins/Experimental/EditorDataStorageFeatures/Source/TedsAssetData/Public/TedsAssetDataColumns.h b/Engine/Plugins/Experimental/EditorDataStorageFeatures/Source/TedsAssetData/Public/TedsAssetDataColumns.h
--- a/Engine/Plugins/Experimental/EditorDataStorageFeatures/Source/TedsAssetData/Public/TedsAssetDataColumns.h
+++ b/Engine/Plugins/Experimental/EditorDataStorageFeatures/Source/TedsAssetData/Public/TedsAssetDataColumns.h
@@ -9,6 +9,7 @@
 #include "Containers/UnrealString.h"
 #include "Containers/VersePath.h"
 #include "Elements/Common/TypedElementCommonTypes.h"
+#include "Elements/Common/TypedElementMapKey.h"
 #include "Elements/Interfaces/TypedElementQueryStorageInterfaces.h"
 #include "Internationalization/Text.h"
 #include "Templates/SharedPointer.h"
@@ -71,7 +72,7 @@ struct FUnresolvedParentAssetPathColumn_Experimental : public FEditorDataStorage
 {
 	GENERATED_BODY()
 	
-	UE::Editor::DataStorage::IndexHash Hash;
+	UE::Editor::DataStorage::FMapKey Key;
 };
 
 USTRUCT()
@@ -79,7 +80,7 @@ struct FUnresolvedAssetsInPathColumn_Experimental : public FEditorDataStorageCol
 {
 	GENERATED_BODY()
 	
-	UE::Editor::DataStorage::IndexHash Hash;
+	UE::Editor::DataStorage::FMapKey Key;
 };
 
 // Tag to identify assets
diff --git a/Engine/Plugins/Experimental/EditorDataStorageFeatures/Source/TedsContentBrowser/Private/TedsContentBrowserModule.cpp b/Engine/Plugins/Experimental/EditorDataStorageFeatures/Source/TedsContentBrowser/Private/TedsContentBrowserModule.cpp
--- a/Engine/Plugins/Experimental/EditorDataStorageFeatures/Source/TedsContentBrowser/Private/TedsContentBrowserModule.cpp
+++ b/Engine/Plugins/Experimental/EditorDataStorageFeatures/Source/TedsContentBrowser/Private/TedsContentBrowserModule.cpp
@@ -10,7 +10,6 @@
 #include "Elements/Columns/TypedElementMiscColumns.h"
 #include "Elements/Common/TypedElementQueryTypes.h"
 #include "Elements/Common/EditorDataStorageFeatures.h"
-#include "Elements/Framework/TypedElementIndexHasher.h"
 #include "Elements/Interfaces/TypedElementDataStorageInterface.h"
 #include "Experimental/ContentBrowserViewExtender.h"
 #include "Modules/ModuleManager.h"
@@ -172,13 +171,13 @@ namespace UE::Editor::ContentBrowser
 		
 		if (Item->GetItem().Legacy_TryGetAssetData(ItemAssetData))
 		{
-			IndexHash IndexHash = GenerateIndexHash(ItemAssetData.GetSoftObjectPath());
-			RowHandle = DataStorage->FindIndexedRow(IndexHash);
+			FMapKey Key = FMapKey(ItemAssetData.GetSoftObjectPath());
+			RowHandle = DataStorage->LookupMappedRow(Key);
 		}
 		else if(Item->GetItem().Legacy_TryGetPackagePath(PackagePath))
 		{
-			IndexHash IndexHash = GenerateIndexHash(PackagePath);
-			RowHandle = DataStorage->FindIndexedRow(IndexHash);
+			FMapKeyView Key = FMapKeyView(PackagePath);
+			RowHandle = DataStorage->LookupMappedRow(Key);
 		}
 
 		return RowHandle;
diff --git a/Engine/Plugins/Experimental/EditorDataStorageFeatures/Source/TedsRevisionControl/Private/Queries/ObjectPackagePathToColumnQueries.cpp b/Engine/Plugins/Experimental/EditorDataStorageFeatures/Source/TedsRevisionControl/Private/Queries/ObjectPackagePathToColumnQueries.cpp
--- a/Engine/Plugins/Experimental/EditorDataStorageFeatures/Source/TedsRevisionControl/Private/Queries/ObjectPackagePathToColumnQueries.cpp
+++ b/Engine/Plugins/Experimental/EditorDataStorageFeatures/Source/TedsRevisionControl/Private/Queries/ObjectPackagePathToColumnQueries.cpp
@@ -4,7 +4,6 @@
 
 #include "Elements/Columns/TypedElementCompatibilityColumns.h"
 #include "Elements/Columns/TypedElementPackageColumns.h"
-#include "Elements/Framework/TypedElementIndexHasher.h"
 #include "Elements/Framework/TypedElementQueryBuilder.h"
 #include "HAL/IConsoleManager.h"
 #include "Misc/PackageName.h"
@@ -66,7 +65,7 @@ void UTypedElementUObjectPackagePathFactory::RegisterQueries(IEditorDataStorageP
 			FProcessor(EQueryTickPhase::FrameEnd, DataStorage.GetQueryTickGroupName(EQueryTickGroups::SyncExternalToDataStorage)),
 			[](IQueryContext& Context, RowHandle Row, const FTypedElementUObjectColumn& Object, const FTypedElementPackageUnresolvedReference& UnresolvedPackageReference)
 			{
-				RowHandle PackageRow = Context.FindIndexedRow(UnresolvedPackageReference.Index);
+				RowHandle PackageRow = Context.LookupMappedRow(FMapKeyView(UnresolvedPackageReference.PathOnDisk));
 				if (!Context.IsRowAvailable(PackageRow))
 				{
 					return;
@@ -126,8 +125,8 @@ void UTypedElementUObjectPackagePathFactory::RegisterTryAddPackageRef(IEditorDat
 					{
 						FPaths::NormalizeFilename(PackageFilename);
 						FString FullPackageFilename = FPaths::ConvertRelativePathToFull(PackageFilename);
-						IndexHash Index = GenerateIndexHash(FullPackageFilename);
-						RowHandle PackageRow = Context.FindIndexedRow(Index);
+						FMapKeyView Key = FMapKeyView(FullPackageFilename);
+						RowHandle PackageRow = Context.LookupMappedRow(Key);
 						if (Context.IsRowAvailable(PackageRow))
 						{
 							const UPackage* Package = ObjectInstance->GetPackage();
@@ -136,7 +135,6 @@ void UTypedElementUObjectPackagePathFactory::RegisterTryAddPackageRef(IEditorDat
 						else
 						{
 							FTypedElementPackageUnresolvedReference UnresolvedPackageReference;
-							UnresolvedPackageReference.Index = Index;
 							UnresolvedPackageReference.PathOnDisk = MoveTemp(FullPackageFilename);
 							Context.AddColumn(Row, MoveTemp(UnresolvedPackageReference));
 						}
diff --git a/Engine/Plugins/Experimental/EditorDataStorageFeatures/Source/TedsSettings/Private/TedsSettingsManager.cpp b/Engine/Plugins/Experimental/EditorDataStorageFeatures/Source/TedsSettings/Private/TedsSettingsManager.cpp
--- a/Engine/Plugins/Experimental/EditorDataStorageFeatures/Source/TedsSettings/Private/TedsSettingsManager.cpp
+++ b/Engine/Plugins/Experimental/EditorDataStorageFeatures/Source/TedsSettings/Private/TedsSettingsManager.cpp
@@ -7,7 +7,6 @@
 #include "Elements/Columns/TypedElementSlateWidgetColumns.h"
 #include "Elements/Columns/TypedElementTypeInfoColumns.h"
 #include "Elements/Common/EditorDataStorageFeatures.h"
-#include "Elements/Framework/TypedElementIndexHasher.h"
 #include "Elements/Framework/TypedElementQueryBuilder.h"
 #include "Elements/Interfaces/TypedElementDataStorageInterface.h"
 #include "ISettingsCategory.h"
@@ -23,14 +22,16 @@
 
 namespace UE::Editor::Settings::Private
 {
-	static UE::Editor::DataStorage::IndexHash GenerateIndexHash(const FName& ContainerName, const FName& CategoryName, const FName& SectionName)
+	static UE::Editor::DataStorage::FMapKey GenerateIndexKey(const FName& ContainerName, const FName& CategoryName, const FName& SectionName)
 	{
-		constexpr static const char SeedName[] = "ISettingsSection";
-		static uint64 Seed = CityHash64(SeedName, sizeof(SeedName) - 1);
-
-		uint64 Hash = CityHash128to64({ Seed, UE::Editor::DataStorage::GenerateIndexHash(ContainerName) });
-		Hash = CityHash128to64({ Hash, UE::Editor::DataStorage::GenerateIndexHash(CategoryName) });
-		return CityHash128to64({ Hash, UE::Editor::DataStorage::GenerateIndexHash(SectionName) });
+		FString Key = "ISettingsSection: ";
+		ContainerName.AppendString(Key);
+		Key.AppendChar(TEXT(','));
+		CategoryName.AppendString(Key);
+		Key.AppendChar(TEXT(','));
+		SectionName.AppendString(Key);
+	
+		return UE::Editor::DataStorage::FMapKey(MoveTemp(Key));
 	}
 
 	static bool IsUnknownColumn(const UScriptStruct* ColumnType)
@@ -123,9 +124,9 @@ UE::Editor::DataStorage::RowHandle FTedsSettingsManager::FindOrAddSettingsSectio
 			// for a delegate bound in another module before it has been called for our bound delegate.
 			RegisterTables(*DataStorage);
 
-			uint64 SectionIndexHash = GenerateIndexHash(ContainerName, CategoryName, SectionName);
+			FMapKey SectionIndexKey = GenerateIndexKey(ContainerName, CategoryName, SectionName);
 
-			RowHandle SectionRow = DataStorage->FindIndexedRow(SectionIndexHash);
+			RowHandle SectionRow = DataStorage->LookupMappedRow(SectionIndexKey);
 
 			if (SectionRow == InvalidRowHandle)
 			{
@@ -136,7 +137,7 @@ UE::Editor::DataStorage::RowHandle FTedsSettingsManager::FindOrAddSettingsSectio
 				DataStorage->AddColumn<FSettingsContainerReferenceColumn>(SectionRow, { .ContainerName = ContainerName, .ContainerRow = InvalidRowHandle });
 				DataStorage->AddColumn<FSettingsCategoryReferenceColumn>(SectionRow, { .CategoryName = CategoryName, .CategoryRow = InvalidRowHandle });
 
-				DataStorage->IndexRow(SectionIndexHash, SectionRow);
+				DataStorage->MapRow(MoveTemp(SectionIndexKey), SectionRow);
 			}
 
 			return SectionRow;
@@ -238,8 +239,8 @@ void FTedsSettingsManager::RegisterSettingsContainer(const FName& ContainerName)
 
 	ISettingsContainerPtr ContainerPtr = SettingsModule->GetContainer(ContainerName);
 
-	uint64 ContainerIndexHash = GenerateIndexHash(ContainerPtr.Get());
-	RowHandle ContainerRow = DataStorage->FindIndexedRow(ContainerIndexHash);
+	FMapKey ContainerIndexKey = FMapKey(ContainerPtr.Get());
+	RowHandle ContainerRow = DataStorage->LookupMappedRow(ContainerIndexKey);
 	if (ContainerRow == InvalidRowHandle)
 	{
 		ContainerRow = DataStorage->AddRow(SettingsContainerTable);
@@ -248,7 +249,7 @@ void FTedsSettingsManager::RegisterSettingsContainer(const FName& ContainerName)
 		DataStorage->AddColumn<FDescriptionColumn>(ContainerRow, { .Description = ContainerPtr->GetDescription() });
 		DataStorage->AddColumn<FSettingsContainerTag>(ContainerRow);
 
-		DataStorage->IndexRow(ContainerIndexHash, ContainerRow);
+		DataStorage->MapRow(MoveTemp(ContainerIndexKey), ContainerRow);
 	}
 
 	TArray<ISettingsCategoryPtr> Categories;
@@ -317,27 +318,23 @@ void FTedsSettingsManager::UnregisterActiveSettings()
 				{
 					DataStorageCompatibility->RemoveCompatibleObject(SettingsObjectPtr);
 
-					DataStorage->RemoveIndex(GenerateIndexHash(ContainerName, CategoryPtr->GetName(), SectionPtr->GetName()));
-
 					UE_LOG(LogTedsSettings, Verbose, TEXT("Removed Settings Section : '%s'"), *SectionPtr->GetName().ToString());
 				}
 			}
 
-			uint64 CategoryIndexHash = GenerateIndexHash(CategoryPtr.Get());
-			RowHandle CategoryRow = DataStorage->FindIndexedRow(CategoryIndexHash);
+			FMapKeyView CategoryIndexKey = FMapKeyView(CategoryPtr.Get());
+			RowHandle CategoryRow = DataStorage->LookupMappedRow(CategoryIndexKey);
 			if (CategoryRow != InvalidRowHandle)
 			{
 				DataStorage->RemoveRow(CategoryRow);
-				DataStorage->RemoveIndex(CategoryIndexHash);
 			}
 		}
 
-		uint64 ContainerIndexHash = GenerateIndexHash(ContainerPtr.Get());
-		RowHandle ContainerRow = DataStorage->FindIndexedRow(ContainerIndexHash);
+		FMapKeyView ContainerIndexKey = FMapKeyView(ContainerPtr.Get());
+		RowHandle ContainerRow = DataStorage->LookupMappedRow(ContainerIndexKey);
 		if (ContainerRow != InvalidRowHandle)
 		{
 			DataStorage->RemoveRow(ContainerRow);
-			DataStorage->RemoveIndex(ContainerIndexHash);
 		}
 	}
 }
@@ -384,9 +381,9 @@ void FTedsSettingsManager::UpdateSettingsCategory(TSharedPtr<ISettingsCategory>
 
 	UE_LOG(LogTedsSettings, Verbose, TEXT("Update Settings Category: '%s->%s'"), *ContainerName.ToString(), *CategoryName.ToString());
 
-	uint64 CategoryIndexHash = GenerateIndexHash(SettingsCategory.Get());
+	FMapKey CategoryIndexKey = FMapKey(SettingsCategory.Get());
 
-	RowHandle CategoryRow = DataStorage->FindIndexedRow(CategoryIndexHash);
+	RowHandle CategoryRow = DataStorage->LookupMappedRow(CategoryIndexKey);
 	if (CategoryRow == InvalidRowHandle)
 	{
 		CategoryRow = DataStorage->AddRow(SettingsCategoryTable);
@@ -397,7 +394,7 @@ void FTedsSettingsManager::UpdateSettingsCategory(TSharedPtr<ISettingsCategory>
 		DataStorage->AddColumn<FDescriptionColumn>(CategoryRow, { .Description = SettingsCategory->GetDescription() });
 		DataStorage->AddColumn<FSettingsCategoryTag>(CategoryRow);
 
-		DataStorage->IndexRow(CategoryIndexHash, CategoryRow);
+		DataStorage->MapRow(MoveTemp(CategoryIndexKey), CategoryRow);
 	}
 
 	TArray<RowHandle> OldRowHandles;
@@ -446,12 +443,12 @@ void FTedsSettingsManager::UpdateSettingsCategory(TSharedPtr<ISettingsCategory>
 		{
 			NewSectionNames.Emplace(SectionName);
 
-			uint64 SectionIndexHash = GenerateIndexHash(ContainerName, CategoryName, SectionName);
+			FMapKey SectionIndexKey = GenerateIndexKey(ContainerName, CategoryName, SectionName);
 
 			TArray<void*> ColumnsToCopy;
 			TArray<const UScriptStruct*> ColumnTypesToCopy;
 
-			RowHandle OldSectionRow = DataStorage->FindIndexedRow(SectionIndexHash);
+			RowHandle OldSectionRow = DataStorage->LookupMappedRow(SectionIndexKey);
 			if (OldSectionRow != InvalidRowHandle)
 			{
 				UE_LOG(LogTedsSettings, Verbose, TEXT("Settings Section : '%s' is already in data storage"), *SectionName.ToString());
@@ -489,7 +486,7 @@ void FTedsSettingsManager::UpdateSettingsCategory(TSharedPtr<ISettingsCategory>
 			UE_CLOG(OldSectionRow != NewSectionRow, LogTedsSettings, Verbose, TEXT("Added Settings Section : '%s'"), *SectionName.ToString());
 			UE_CLOG(OldSectionRow == NewSectionRow, LogTedsSettings, Verbose, TEXT("Updated Settings Section : '%s'"), *SectionName.ToString());
 
-			DataStorage->IndexRow(SectionIndexHash, NewSectionRow);			
+			DataStorage->MapRow(MoveTemp(SectionIndexKey), NewSectionRow);			
 		}
 	}
 
@@ -531,7 +528,7 @@ void FTedsSettingsManager::UpdateSettingsCategory(TSharedPtr<ISettingsCategory>
 
 		DataStorage->RemoveRow(OldSectionRow);
 
-		DataStorage->IndexRow(GenerateIndexHash(ContainerName, CategoryName, OldSectionName), NewSectionRow);
+		DataStorage->MapRow(GenerateIndexKey(ContainerName, CategoryName, OldSectionName), NewSectionRow);
 
 		UE_LOG(LogTedsSettings, Verbose, TEXT("Inactivated Settings Section : '%s'"), *OldSectionName.ToString());
 	}
diff --git a/Engine/Plugins/Experimental/EditorDataStorageFeatures/Source/TedsStyling/Private/TedsStylingFactory.cpp b/Engine/Plugins/Experimental/EditorDataStorageFeatures/Source/TedsStyling/Private/TedsStylingFactory.cpp
--- a/Engine/Plugins/Experimental/EditorDataStorageFeatures/Source/TedsStyling/Private/TedsStylingFactory.cpp
+++ b/Engine/Plugins/Experimental/EditorDataStorageFeatures/Source/TedsStyling/Private/TedsStylingFactory.cpp
@@ -6,7 +6,6 @@
 #include "Elements/Columns/TypedElementMiscColumns.h"
 #include "Elements/Columns/TypedElementSlateWidgetColumns.h"
 #include "Elements/Common/EditorDataStorageFeatures.h"
-#include "Elements/Framework/TypedElementIndexHasher.h"
 #include "Elements/Interfaces/TypedElementDataStorageInterface.h"
 #include "HAL/IConsoleManager.h"
 #include "Styling/SlateBrush.h"
@@ -121,14 +120,14 @@ UE::Editor::DataStorage::RowHandle UTedsStylingFactory::AddOrGetStyleRow(IEditor
 	const FName& StyleName, const ISlateStyle& OwnerStyle)
 {
 	const FName StyleSetName = OwnerStyle.GetStyleSetName();
-	const UE::Editor::DataStorage::IndexHash Index = UE::Editor::DataStorage::GenerateIndexHash(ISlateStyle::Join(StyleSetName, TCHAR_TO_ANSI(*StyleName.ToString())));
-	UE::Editor::DataStorage::RowHandle Row = DataStorage->FindIndexedRow(Index);
+	const UE::Editor::DataStorage::FMapKey Key = UE::Editor::DataStorage::FMapKey(ISlateStyle::Join(StyleSetName, TCHAR_TO_ANSI(*StyleName.ToString())));
+	UE::Editor::DataStorage::RowHandle Row = DataStorage->LookupMappedRow(Key);
 
 	if (!DataStorage->IsRowAssigned(Row))
 	{
 		static UE::Editor::DataStorage::RowHandle Table = DataStorage->FindTable(UE::TedsStylingFactory::Local::TableName);
 		Row = DataStorage->AddRow(Table);
-		DataStorage->IndexRow(Index, Row);
+		DataStorage->MapRow(Key, Row);
 
 		DataStorage->GetColumn<FNameColumn>(Row)->Name = StyleName;
 		DataStorage->GetColumn<FSlateStyleSetColumn>(Row)->StyleSetName = StyleSetName;
diff --git a/Engine/Source/Runtime/TypedElementFramework/Private/Elements/Common/TypedElementMapKey.cpp b/Engine/Source/Runtime/TypedElementFramework/Private/Elements/Common/TypedElementMapKey.cpp
new file mode 100644
--- /dev/null
+++ b/Engine/Source/Runtime/TypedElementFramework/Private/Elements/Common/TypedElementMapKey.cpp
@@ -0,0 +1,443 @@
+// Copyright Epic Games, Inc. All Rights Reserved.
+
+#include "Elements/Common/TypedElementMapKey.h"
+
+namespace UE::Editor::DataStorage
+{
+	struct FHasher
+	{
+		uint64 operator()(FEmptyVariantState)	{ return 0; }
+		uint64 operator()(const void* Value)	{ return reinterpret_cast<uint64>(Value); }
+		uint64 operator()(const UObject* Value) { return reinterpret_cast<uint64>(Value); }
+		uint64 operator()(int64 Value)			{ return static_cast<uint64>(Value); }
+		uint64 operator()(uint64 Value)			{ return static_cast<uint64>(Value); }
+		uint64 operator()(float Value)			{ return static_cast<uint64>(Value); }
+		uint64 operator()(double Value)			{ return static_cast<uint64>(Value); }
+		uint64 operator()(const FString& Value)
+		{
+			return CityHash64(reinterpret_cast<const char*>(*Value), Value.Len() * sizeof(**Value));
+		}
+		uint64 operator()(const FStringView& Value)
+		{
+			return CityHash64(reinterpret_cast<const char*>(Value.GetData()), Value.NumBytes());
+		}
+		uint64 operator()(const FName& Value)	{ return Value.ToUnstableInt(); }
+		uint64 operator()(const FName* Value)	{ return Value->ToUnstableInt(); }
+		uint64 operator()(const FSoftObjectPath* Value)
+		{
+			FTopLevelAssetPath TopLevelAssetPath = Value->GetAssetPath();
+			uint64 Hash = TopLevelAssetPath.GetPackageName().ToUnstableInt();
+			Hash = CityHash128to64({ Hash, TopLevelAssetPath.GetAssetName().ToUnstableInt() });
+			const FString& SubPath = Value->GetSubPathString();
+			return CityHash64WithSeed(reinterpret_cast<const char*>(*SubPath), SubPath.Len() * sizeof(**SubPath), Hash);
+		}
+		uint64 operator()(const TUniquePtr<FSoftObjectPath>& Value) { return operator()(Value.Get()); }
+	};
+
+	struct FKeyToKeyViewConverter
+	{
+		template<typename T>
+		FMapKeyView::KeyViewType operator()(const T& Value)
+		{
+			return FMapKeyView::KeyViewType(TInPlaceType<T>(), Value);
+		}
+
+		FMapKeyView::KeyViewType operator()(const FString& Value)
+		{
+			return FMapKeyView::KeyViewType(TInPlaceType<FStringView>(), Value);
+		}
+
+		FMapKeyView::KeyViewType operator()(const FName& Value)
+		{
+			return FMapKeyView::KeyViewType(TInPlaceType<const FName*>(), &Value);
+		}
+
+		FMapKeyView::KeyViewType operator()(const TUniquePtr<FSoftObjectPath>& Value)
+		{
+			return FMapKeyView::KeyViewType(TInPlaceType<const FSoftObjectPath*>(), Value.Get());
+		}
+	};
+
+	struct FKeyViewToKeyConverter
+	{
+		FMapKey operator()(const FEmptyVariantState& Value)
+		{
+			return FMapKey();
+		}
+
+		template<typename T>
+		FMapKey operator()(const T& Value)
+		{
+			return FMapKey(Value);
+		}
+
+		FMapKey operator()(const FStringView& Value)
+		{
+			return FMapKey(FString(Value));
+		}
+
+		FMapKey operator()(const FName* Value)
+		{
+			return FMapKey(*Value);
+		}
+
+		FMapKey operator()(const FSoftObjectPath* Value)
+		{
+			return FMapKey(*Value);
+		}
+	};
+
+	struct FKeyMove
+	{
+		FMapKey::KeyType operator()(FEmptyVariantState&)
+		{
+			return FMapKey::KeyType(TInPlaceType<FEmptyVariantState>());
+		}
+
+		template<typename T>
+		FMapKey::KeyType operator()(T&& Other)
+		{
+			return FMapKey::KeyType(TInPlaceType<T>(), Forward<T>(Other));
+		}
+	};
+
+	struct FKeyCopy
+	{
+		FMapKey::KeyType operator()(const FEmptyVariantState&)
+		{
+			return FMapKey::KeyType(TInPlaceType<FEmptyVariantState>());
+		}
+
+		template<typename T>
+		FMapKey::KeyType operator()(const T& Other)
+		{
+			return FMapKey::KeyType(TInPlaceType<T>(), Other);
+		}
+
+		FMapKey::KeyType operator()(const TUniquePtr<FSoftObjectPath>& Other)
+		{
+			return FMapKey::KeyType(TInPlaceType<TUniquePtr<FSoftObjectPath>>(), MakeUnique<FSoftObjectPath>(*Other));
+		}
+	};
+
+	struct FKeyComparer
+	{
+		const FMapKey::KeyType& Key;
+		explicit FKeyComparer(const FMapKey::KeyType& Key) : Key(Key) {}
+
+		bool operator()(FEmptyVariantState)
+		{
+			return Key.IsType<FEmptyVariantState>();
+		}
+
+		template<typename T>
+		bool operator()(const T& Other)
+		{
+			const auto* Local = Key.TryGet<T>();
+			return Local ? *Local == Other : false;
+		}
+
+		bool operator()(const TUniquePtr<FSoftObjectPath>& Other)
+		{
+			if (const TUniquePtr<FSoftObjectPath>* Local = Key.TryGet<TUniquePtr<FSoftObjectPath>>())
+			{
+				if (Local->IsValid() && Other.IsValid())
+				{
+					return *(Local->Get()) == *Other;
+				}
+			}
+			return false;
+		}
+	};
+
+	struct FKeyViewComparer
+	{
+		const FMapKeyView::KeyViewType& Key;
+		explicit FKeyViewComparer(const FMapKeyView::KeyViewType& Key) : Key(Key) {}
+
+		bool operator()(FEmptyVariantState)
+		{
+			return Key.IsType<FEmptyVariantState>();
+		}
+
+		template<typename T>
+		bool operator()(const T& Other)
+		{
+			const T* Local = Key.TryGet<T>();
+			return Local ? *Local == Other : false;
+		}
+
+		bool operator()(const FString& Other)
+		{
+			const FStringView* Local = Key.TryGet<FStringView>();
+			return Local ? *Local == Other : false;
+		}
+
+		bool operator()(const FStringView& Other)
+		{
+			const FStringView* Local = Key.TryGet<FStringView>();
+			return Local ? *Local == Other : false;
+		}
+
+		bool operator()(const FName& Other)
+		{
+			const FName* const* Local = Key.TryGet<const FName*>();
+			return Local ? **Local == Other : false;
+		}
+
+		bool operator()(const FName* Other)
+		{
+			const FName* const* Local = Key.TryGet<const FName*>();
+			return Local ? **Local == *Other : false;
+		}
+
+		bool operator()(const TUniquePtr<FSoftObjectPath>& Other)
+		{
+			if (Other.IsValid())
+			{
+				const FSoftObjectPath* const* Local = Key.TryGet<const FSoftObjectPath*>();
+				return Local ? **Local == *Other : false;
+			}
+			return false;
+		}
+
+		bool operator()(const FSoftObjectPath* Other)
+		{
+			const FSoftObjectPath* const* Local = Key.TryGet<const FSoftObjectPath*>();
+			return Local ? **Local == *Other : false;
+		}
+	};
+
+	struct FToString
+	{
+		FString operator()(FEmptyVariantState) { return TEXT("Empty"); }
+		FString operator()(const void* Value) { return FString::Printf(TEXT("Pointer(0x%p)"), Value); }
+		FString operator()(const UObject* Value) { return FString::Printf(TEXT("UObject(0x%p)"), Value); }
+		FString operator()(int64 Value) { return FString::Printf(TEXT("%lld"), Value); }
+		FString operator()(uint64 Value) { return FString::Printf(TEXT("%llu"), Value); }
+		FString operator()(float Value) { return FString::SanitizeFloat(Value);; }
+		FString operator()(double Value) { return FString::SanitizeFloat(Value);; }
+		FString operator()(const FString& Value) { return Value; }
+		FString operator()(const FStringView& Value) { return FString(Value); }
+		FString operator()(const FName& Value) { return Value.ToString(); }
+		FString operator()(const FName* Value) { return Value->ToString(); }
+		FString operator()(const FSoftObjectPath* Value) { return Value->ToString(); }
+		FString operator()(const TUniquePtr<FSoftObjectPath>& Value) { return Value ? Value->ToString() : TEXT("Empty path"); }
+	};
+
+	//
+	// FMapKey
+	//
+
+	FMapKey::FMapKey(FMapKey&& Rhs)
+		: Key(Visit(FKeyMove(), MoveTemp(Rhs.Key)))
+	{
+		Rhs.Clear();
+	}
+
+	FMapKey& FMapKey::operator=(FMapKey&& Rhs)
+	{
+		Key = Visit(FKeyMove(), MoveTemp(Rhs.Key));
+		Rhs.Clear();
+		return *this;
+	}
+
+	FMapKey::FMapKey(const FMapKey& Key)
+		: Key(Visit(FKeyCopy(), Key.Key))
+	{
+	}
+
+	FMapKey::FMapKey(const void* Key)
+		: Key(TInPlaceType<const void*>(), Key)
+	{
+	}
+
+	FMapKey::FMapKey(const UObject* Key)
+		: Key(TInPlaceType<const UObject*>(), Key)
+	{
+	}
+
+	FMapKey::FMapKey(int64 Key)
+		: Key(TInPlaceType<int64>(), Key)
+	{
+	}
+
+	FMapKey::FMapKey(uint64 Key)
+		: Key(TInPlaceType<uint64>(), Key)
+	{
+	}
+
+	FMapKey::FMapKey(float Key)
+		: Key(TInPlaceType<float>(), Key)
+	{
+	}
+
+	FMapKey::FMapKey(double Key)
+		: Key(TInPlaceType<double>(), Key)
+	{
+	}
+
+	FMapKey::FMapKey(FString Key)
+		: Key(TInPlaceType<FString>(), MoveTemp(Key))
+	{
+	}
+	
+	FMapKey::FMapKey(FName Key)
+		: Key(TInPlaceType<FName>(), MoveTemp(Key))
+	{
+	}
+
+	FMapKey::FMapKey(FSoftObjectPath Key)
+		: Key(TInPlaceType<TUniquePtr<FSoftObjectPath>>(), MakeUnique<FSoftObjectPath>(MoveTemp(Key)))
+	{
+	}
+
+	uint64 FMapKey::CalculateHash() const
+	{
+		return Visit(FHasher(), Key);
+	}
+	
+	bool FMapKey::IsSet() const
+	{
+		return !Key.IsType<FEmptyVariantState>();
+	}
+
+	void FMapKey::Clear()
+	{
+		Key.Emplace<FEmptyVariantState>();
+	}
+
+	FString FMapKey::ToString() const
+	{
+		return Visit(FToString(), Key);
+	}
+
+	FMapKey& FMapKey::operator=(const FMapKey& Rhs)
+	{
+		Key = Visit(FKeyCopy(), Rhs.Key);
+		return *this;
+	}
+
+	bool FMapKey::operator==(const FMapKey& Rhs) const
+	{
+		return Visit(FKeyComparer(Key), Rhs.Key);
+	}
+
+	bool FMapKey::operator==(const FMapKeyView& Rhs) const
+	{
+		return Rhs == *this;
+	}
+
+	bool FMapKey::operator!=(const FMapKey& Rhs) const
+	{
+		return !(*this == Rhs);
+	}
+
+	bool FMapKey::operator!=(const FMapKeyView& Rhs) const
+	{
+		return !(*this == Rhs);
+	}
+
+
+
+	//
+	// FMapKeyView
+	//
+
+	FMapKeyView::FMapKeyView(const FMapKey& InKey)
+		: Key(Visit(FKeyToKeyViewConverter(), InKey.Key))
+	{
+	}
+
+	FMapKeyView::FMapKeyView(const void* Key)
+		: Key(TInPlaceType<const void*>(), Key)
+	{
+	}
+
+	FMapKeyView::FMapKeyView(const UObject* Key)
+		: Key(TInPlaceType<const UObject*>(), Key)
+	{
+	}
+
+	FMapKeyView::FMapKeyView(int64 Key)
+		: Key(TInPlaceType<int64>(), Key)
+	{
+	}
+
+	FMapKeyView::FMapKeyView(uint64 Key)
+		: Key(TInPlaceType<uint64>(), Key)
+	{
+	}
+
+	FMapKeyView::FMapKeyView(float Key)
+		: Key(TInPlaceType<float>(), Key)
+	{
+	}
+
+	FMapKeyView::FMapKeyView(double Key)
+		: Key(TInPlaceType<double>(), Key)
+	{
+	}
+
+	FMapKeyView::FMapKeyView(const FString& Key)
+		: Key(TInPlaceType<FStringView>(), Key)
+	{
+	}
+	
+	FMapKeyView::FMapKeyView(FStringView Key)
+		: Key(TInPlaceType<FStringView>(), Key)
+	{
+	}
+
+	FMapKeyView::FMapKeyView(const FName& Key)
+		: Key(TInPlaceType<const FName*>(), &Key)
+	{
+	}
+
+	FMapKeyView::FMapKeyView(const FSoftObjectPath& Key)
+		: Key(TInPlaceType<const FSoftObjectPath*>(), &Key)
+	{
+	}
+
+	uint64 FMapKeyView::CalculateHash() const
+	{
+		return Visit(FHasher(), Key);
+	}
+
+	FString FMapKeyView::ToString() const
+	{
+		return Visit(FToString(), Key);
+	}
+
+	FMapKeyView& FMapKeyView::operator=(const FMapKey& InKey)
+	{
+		Key = Visit(FKeyToKeyViewConverter(), InKey.Key);
+		return *this;
+	}
+
+	FMapKey FMapKeyView::CreateKey() const
+	{
+		return Visit(FKeyViewToKeyConverter(), Key);
+	}
+
+	bool FMapKeyView::operator==(const FMapKey& Rhs) const
+	{
+		return Visit(FKeyViewComparer(Key), Rhs.Key);
+	}
+
+	bool FMapKeyView::operator==(const FMapKeyView& Rhs) const
+	{
+		return Visit(FKeyViewComparer(Key), Rhs.Key);
+	}
+
+	bool FMapKeyView::operator!=(const FMapKey& Rhs) const
+	{
+		return !(*this == Rhs);
+	}
+
+	bool FMapKeyView::operator!=(const FMapKeyView& Rhs) const
+	{
+		return !(*this == Rhs);
+	}
+}
+// namespace UE::Editor::DataStorage
diff --git a/Engine/Source/Runtime/TypedElementFramework/Public/Elements/Columns/TypedElementHiearchyColumns.h b/Engine/Source/Runtime/TypedElementFramework/Public/Elements/Columns/TypedElementHiearchyColumns.h
--- a/Engine/Source/Runtime/TypedElementFramework/Public/Elements/Columns/TypedElementHiearchyColumns.h
+++ b/Engine/Source/Runtime/TypedElementFramework/Public/Elements/Columns/TypedElementHiearchyColumns.h
@@ -3,6 +3,7 @@
 #pragma once
 
 #include "Elements/Common/TypedElementHandles.h"
+#include "Elements/Common/TypedElementMapKey.h"
 #include "Elements/Interfaces/TypedElementDataStorageInterface.h"
 #include "UObject/ObjectMacros.h"
 
@@ -28,6 +29,5 @@ struct FUnresolvedTableRowParentColumn final : public FEditorDataStorageColumn
 {
 	GENERATED_BODY()
 
-	UPROPERTY(meta = (IgnoreForMemberInitializationTest))
-	uint64 ParentIdHash;
+	UE::Editor::DataStorage::FMapKey ParentIdKey;
 };
diff --git a/Engine/Source/Runtime/TypedElementFramework/Public/Elements/Columns/TypedElementPackageColumns.h b/Engine/Source/Runtime/TypedElementFramework/Public/Elements/Columns/TypedElementPackageColumns.h
--- a/Engine/Source/Runtime/TypedElementFramework/Public/Elements/Columns/TypedElementPackageColumns.h
+++ b/Engine/Source/Runtime/TypedElementFramework/Public/Elements/Columns/TypedElementPackageColumns.h
@@ -4,6 +4,7 @@
 
 #include "Containers/UnrealString.h"
 #include "Elements/Common/TypedElementCommonTypes.h"
+#include "Elements/Common/TypedElementMapKey.h"
 #include "Elements/Common/TypedElementHandles.h"
 #include "Elements/Common/TypedElementQueryTypes.h"
 #include "HAL/Platform.h"
@@ -22,8 +23,6 @@ struct FTypedElementPackageUnresolvedReference final : public FEditorDataStorage
 	GENERATED_BODY()
 
 	FString PathOnDisk;
-	
-	UE::Editor::DataStorage::IndexHash Index;
 };
 
 /**
diff --git a/Engine/Source/Runtime/TypedElementFramework/Public/Elements/Common/TypedElementMapKey.h b/Engine/Source/Runtime/TypedElementFramework/Public/Elements/Common/TypedElementMapKey.h
new file mode 100644
--- /dev/null
+++ b/Engine/Source/Runtime/TypedElementFramework/Public/Elements/Common/TypedElementMapKey.h
@@ -0,0 +1,97 @@
+// Copyright Epic Games, Inc. All Rights Reserved.
+
+#pragma once
+
+#include "Containers/StringFwd.h"
+#include "Containers/UnrealString.h"
+#include "Misc/TVariant.h"
+#include "Templates/UniquePtr.h"
+#include "UObject/NameTypes.h"
+#include "UObject/SoftObjectPath.h"
+
+class UObject;
+
+namespace UE::Editor::DataStorage
+{
+	class FMapKeyView;
+
+	class FMapKey
+	{
+		friend class FMapKeyView;
+		friend struct FKeyComparer;
+		friend struct FKeyCopy;
+		friend struct FKeyMove;
+
+	public:
+		FMapKey() = default;
+		
+		TYPEDELEMENTFRAMEWORK_API FMapKey(FMapKey&& Rhs);
+		TYPEDELEMENTFRAMEWORK_API FMapKey& operator=(FMapKey&& Rhs);
+		
+		TYPEDELEMENTFRAMEWORK_API FMapKey(const FMapKey& Key);
+		TYPEDELEMENTFRAMEWORK_API explicit FMapKey(const void* Key);
+		TYPEDELEMENTFRAMEWORK_API explicit FMapKey(const UObject* Key);
+		TYPEDELEMENTFRAMEWORK_API explicit FMapKey(int64 Key);
+		TYPEDELEMENTFRAMEWORK_API explicit FMapKey(uint64 Key);
+		TYPEDELEMENTFRAMEWORK_API explicit FMapKey(float Key);
+		TYPEDELEMENTFRAMEWORK_API explicit FMapKey(double Key);
+		TYPEDELEMENTFRAMEWORK_API explicit FMapKey(FString Key);
+		TYPEDELEMENTFRAMEWORK_API explicit FMapKey(FName Key);
+		TYPEDELEMENTFRAMEWORK_API explicit FMapKey(FSoftObjectPath Key);
+
+		TYPEDELEMENTFRAMEWORK_API FMapKey& operator=(const FMapKey& Rhs);
+
+		TYPEDELEMENTFRAMEWORK_API bool operator==(const FMapKey& Rhs) const;
+		TYPEDELEMENTFRAMEWORK_API bool operator==(const FMapKeyView& Rhs) const;
+		TYPEDELEMENTFRAMEWORK_API bool operator!=(const FMapKey& Rhs) const;
+		TYPEDELEMENTFRAMEWORK_API bool operator!=(const FMapKeyView& Rhs) const;
+
+		TYPEDELEMENTFRAMEWORK_API uint64 CalculateHash() const;
+		TYPEDELEMENTFRAMEWORK_API bool IsSet() const;
+		TYPEDELEMENTFRAMEWORK_API void Clear();
+		TYPEDELEMENTFRAMEWORK_API FString ToString() const;
+
+	private:
+		using KeyType = TVariant<FEmptyVariantState, const void*, const UObject*, int64, uint64, float, double, FString, FName, TUniquePtr<FSoftObjectPath>>;
+
+		KeyType Key;
+	};
+
+	class FMapKeyView
+	{
+		friend struct FKeyToKeyViewConverter;
+		friend struct FKeyViewComparer;
+
+	public:
+		FMapKeyView() = default;
+		TYPEDELEMENTFRAMEWORK_API FMapKeyView(const FMapKey& InKey);
+		TYPEDELEMENTFRAMEWORK_API explicit FMapKeyView(const void* Key);
+		TYPEDELEMENTFRAMEWORK_API explicit FMapKeyView(const UObject* Key);
+		TYPEDELEMENTFRAMEWORK_API explicit FMapKeyView(int64 Key);
+		TYPEDELEMENTFRAMEWORK_API explicit FMapKeyView(uint64 Key);
+		TYPEDELEMENTFRAMEWORK_API explicit FMapKeyView(float Key);
+		TYPEDELEMENTFRAMEWORK_API explicit FMapKeyView(double Key);
+		TYPEDELEMENTFRAMEWORK_API explicit FMapKeyView(const FString& Key);
+		TYPEDELEMENTFRAMEWORK_API explicit FMapKeyView(FStringView Key);
+		TYPEDELEMENTFRAMEWORK_API explicit FMapKeyView(const FName& Key);
+		TYPEDELEMENTFRAMEWORK_API explicit FMapKeyView(const FSoftObjectPath& Key);
+
+		TYPEDELEMENTFRAMEWORK_API FMapKeyView& operator=(const FMapKey& InKey);
+
+		TYPEDELEMENTFRAMEWORK_API FMapKey CreateKey() const;
+
+		TYPEDELEMENTFRAMEWORK_API bool operator==(const FMapKey& Rhs) const;
+		TYPEDELEMENTFRAMEWORK_API bool operator==(const FMapKeyView& Rhs) const;
+		TYPEDELEMENTFRAMEWORK_API bool operator!=(const FMapKey& Rhs) const;
+		TYPEDELEMENTFRAMEWORK_API bool operator!=(const FMapKeyView& Rhs) const;
+
+		TYPEDELEMENTFRAMEWORK_API uint64 CalculateHash() const;
+		TYPEDELEMENTFRAMEWORK_API FString ToString() const;
+
+	private:
+		using KeyViewType = TVariant<FEmptyVariantState, const void*, const UObject*, int64, uint64, float, double, FStringView, const FName*, const FSoftObjectPath*>;
+
+		KeyViewType Key;
+	};
+}
+// namespace UE::Editor::DataStorage
diff --git a/Engine/Source/Runtime/TypedElementFramework/Public/Elements/Common/TypedElementQueryTypes.h b/Engine/Source/Runtime/TypedElementFramework/Public/Elements/Common/TypedElementQueryTypes.h
--- a/Engine/Source/Runtime/TypedElementFramework/Public/Elements/Common/TypedElementQueryTypes.h
+++ b/Engine/Source/Runtime/TypedElementFramework/Public/Elements/Common/TypedElementQueryTypes.h
@@ -7,6 +7,7 @@
 
 namespace UE::Editor::DataStorage
 {
+	// UE_DEPRECATED(5.6, "Use FMapKey(View) with the new mapping functions.") // Disabled as it's used by the deprecated functions.
 	using IndexHash = uint64;
 
 	enum class EQueryTickPhase : uint8
diff --git a/Engine/Source/Runtime/TypedElementFramework/Public/Elements/Framework/TypedElementIndexHasher.h b/Engine/Source/Runtime/TypedElementFramework/Public/Elements/Framework/TypedElementIndexHasher.h
--- a/Engine/Source/Runtime/TypedElementFramework/Public/Elements/Framework/TypedElementIndexHasher.h
+++ b/Engine/Source/Runtime/TypedElementFramework/Public/Elements/Framework/TypedElementIndexHasher.h
@@ -14,18 +14,26 @@
 namespace UE::Editor::DataStorage
 {
 	template<typename T>
+	UE_DEPRECATED(5.6, "Use FMapKey(View) with the new mapping functions.")
 	IndexHash GenerateIndexHash(const T* Object);
 	
 	template<typename T>
+	UE_DEPRECATED(5.6, "Use FMapKey(View) with the new mapping functions.")
 	IndexHash GenerateIndexHash(const TWeakObjectPtr<T>& Object);
 	template<typename T>
+	UE_DEPRECATED(5.6, "Use FMapKey(View) with the new mapping functions.")
 	IndexHash GenerateIndexHash(const TObjectPtr<T>& Object);
 	template<typename T>
+	UE_DEPRECATED(5.6, "Use FMapKey(View) with the new mapping functions.")
 	IndexHash GenerateIndexHash(const TStrongObjectPtr<T>& Object);
 
+	// UE_DEPRECATED(5.6, "Use FMapKey(View) with the new mapping functions.") // Disabled deprecation as it's in use by other deprecated functions.
 	inline IndexHash GenerateIndexHash(const FString& Object);
+	UE_DEPRECATED(5.6, "Use FMapKey(View) with the new mapping functions.")
 	inline IndexHash GenerateIndexHash(FStringView Object);
+	// UE_DEPRECATED(5.6, "Use FMapKey(View) with the new mapping functions.") // Disabled deprecation as it's in use by other deprecated functions.
 	inline IndexHash GenerateIndexHash(FName Object);
+	UE_DEPRECATED(5.6, "Use FMapKey(View) with the new mapping functions.")
 	inline IndexHash GenerateIndexHash(const FSoftObjectPath& ObjectPath);
 } // namespace UE::Editor::DataStorage
 
diff --git a/Engine/Source/Runtime/TypedElementFramework/Public/Elements/Interfaces/TypedElementDataStorageInterface.h b/Engine/Source/Runtime/TypedElementFramework/Public/Elements/Interfaces/TypedElementDataStorageInterface.h
--- a/Engine/Source/Runtime/TypedElementFramework/Public/Elements/Interfaces/TypedElementDataStorageInterface.h
+++ b/Engine/Source/Runtime/TypedElementFramework/Public/Elements/Interfaces/TypedElementDataStorageInterface.h
@@ -8,6 +8,7 @@
 #include "Delegates/Delegate.h"
 #include "Elements/Common/TypedElementCommonTypes.h"
 #include "Elements/Common/TypedElementHandles.h"
+#include "Elements/Common/TypedElementMapKey.h"
 #include "Elements/Common/TypedElementQueryConditions.h"
 #include "Elements/Common/TypedElementQueryDescription.h"
 #include "Elements/Common/TypedElementQueryTypes.h"
@@ -407,31 +408,29 @@ public:
 	virtual void ActivateQueries(FName ActivationName) = 0;
 	
 	/**
-	 * @section Indexing
+	 * @section Mapping
 	 * @description
 	 * In order for rows to reference each other it's often needed to find a row based on the content of one of its columns. This can be
 	 * done by linearly searching through columns, though this comes at a performance cost. As an alternative the data storage allows
-	 * one or more indexes to be created for a row. An index is a 64-bit value and typically uses a hash value of an identifying value.
+	 * one or more key to be created for a row for fast retrieval.
 	 */
 
-	/** Retrieves the row for an indexed object. Returns an invalid row handle if the hash wasn't found. */
-	virtual RowHandle FindIndexedRow(UE::Editor::DataStorage::IndexHash Index) const = 0;
+	/** Retrieves the row for a mapped object. Returns an invalid row handle if the no row with the provided key was found. */
+	virtual RowHandle LookupMappedRow(const UE::Editor::DataStorage::FMapKeyView& Key) const = 0;
 	/** 
-	 * Registers a row under the index hash. The same row can be registered multiple, but an index hash can only be associated 
-	 * with a single row.
+	 * Registers a row under a key. The same row can be registered multiple, but an key can only be associated with a single row.
 	 */
-	virtual void IndexRow(UE::Editor::DataStorage::IndexHash Index, RowHandle Row) = 0;
+	virtual void MapRow(UE::Editor::DataStorage::FMapKey Key, RowHandle Row) = 0;
 	/**
-	 * Register multiple rows under their index hash. The same row can be registered multiple times,
-	 * but an index hash can only be associated with a single row.
-	 */
-	virtual void BatchIndexRows(
-		TConstArrayView<TPair<UE::Editor::DataStorage::IndexHash, RowHandle>> IndexRowPairs) = 0;
-	/** Updates the index of a row to a new value. Effectively this is the same as removing an index and adding a new one. */
-	virtual void Reindex(UE::Editor::DataStorage::IndexHash OriginalIndex, UE::Editor::DataStorage::IndexHash NewIndex) = 0;
-	/** Removes a previously registered index hash from the index lookup table or does nothing if the hash no longer exists. */
-	virtual void RemoveIndex(UE::Editor::DataStorage::IndexHash Index) = 0;
-
+	 * Register multiple rows under their key. The same row can be registered multiple times, but key can only be associated with a single
+	 * row. During processing the keys will be moved out into the final location.
+	 */
+	virtual void BatchMapRows(TArrayView<TPair<UE::Editor::DataStorage::FMapKey, RowHandle>> MapRowPairs) = 0;
+	/** Updates the key of a row to a new value. Effectively this is the same as removing a key and adding a new one. */
+	virtual void RemapRow(const UE::Editor::DataStorage::FMapKeyView& OriginalKey, UE::Editor::DataStorage::FMapKey NewKey) = 0;
+	/** Removes a previously registered key from the mapping table or does nothing if the key no longer exists. */
+	virtual void RemoveRowMapping(const UE::Editor::DataStorage::FMapKeyView& Key) = 0;
+	
 	/**
 	 * @section Miscellaneous
 	 */
@@ -468,7 +467,18 @@ public:
 	/**
 	 * @section Deprecated
 	 */
-	UE_DEPRECATED(5.6, "Use 'Reindex' instead. The last argument 'Row' can be left out.")
+	UE_DEPRECATED(5.6, "Use 'LookUpMappedRow' instead and use the new FMapKey(View) instead of the explicit 'GenerateIndexHash'-functions.")
+	inline RowHandle FindIndexedRow(UE::Editor::DataStorage::IndexHash Index) const;
+	UE_DEPRECATED(5.6, "Use 'MapRow' instead and use the new FMapKey(View) instead of the explicit 'GenerateIndexHash'-functions.")
+	inline void IndexRow(UE::Editor::DataStorage::IndexHash Index, RowHandle Row);
+	UE_DEPRECATED(5.6, "Use 'BatchMapRows' instead and use the new FMapKey(View) instead of the explicit 'GenerateIndexHash'-functions.")
+	inline void BatchIndexRows(TConstArrayView<TPair<UE::Editor::DataStorage::IndexHash, RowHandle>> IndexRowPairs);
+	UE_DEPRECATED(5.6, "Use 'RemapRow' instead and use the new FMapKey(View) instead of the explicit 'GenerateIndexHash'-functions.")
+	inline void Reindex(UE::Editor::DataStorage::IndexHash OriginalIndex, UE::Editor::DataStorage::IndexHash NewIndex);
+	UE_DEPRECATED(5.6, "Use 'RemovedRowMapping' instead and use the new FMapKey(View) instead of the explicit 'GenerateIndexHash'-functions.")
+	inline void RemoveIndex(UE::Editor::DataStorage::IndexHash Index);
+
+	UE_DEPRECATED(5.6, "Use 'RemapRow' instead and use the new FMapKey(View) instead of the explicit 'GenerateIndexHash'-functions. The Row argument is no longer used.")
 	inline void ReindexRow(UE::Editor::DataStorage::IndexHash OriginalIndex, UE::Editor::DataStorage::IndexHash NewIndex, RowHandle Row);
 };
 
@@ -706,5 +716,33 @@ SystemType* IEditorDataStorageProvider::GetExternalSystem()
 void IEditorDataStorageProvider::ReindexRow(
 	UE::Editor::DataStorage::IndexHash OriginalIndex, UE::Editor::DataStorage::IndexHash NewIndex, RowHandle Row)
 {
-	Reindex(OriginalIndex, NewIndex);
+	RemapRow(UE::Editor::DataStorage::FMapKeyView(OriginalIndex), UE::Editor::DataStorage::FMapKey(NewIndex));
+}
+
+UE::Editor::DataStorage::RowHandle IEditorDataStorageProvider::FindIndexedRow(UE::Editor::DataStorage::IndexHash Index) const
+{
+	return LookupMappedRow(UE::Editor::DataStorage::FMapKeyView(Index));
+}
+
+void IEditorDataStorageProvider::IndexRow(UE::Editor::DataStorage::IndexHash Index, RowHandle Row)
+{
+	MapRow(UE::Editor::DataStorage::FMapKey(Index), Row);
+}
+
+void IEditorDataStorageProvider::BatchIndexRows(TConstArrayView<TPair<UE::Editor::DataStorage::IndexHash, RowHandle>> IndexRowPairs)
+{
+	for (const TPair<UE::Editor::DataStorage::IndexHash, RowHandle>& Pair : IndexRowPairs)
+	{
+		MapRow(UE::Editor::DataStorage::FMapKey(Pair.Key), Pair.Value);
+	}
+}
+
+void IEditorDataStorageProvider::Reindex(UE::Editor::DataStorage::IndexHash OriginalIndex, UE::Editor::DataStorage::IndexHash NewIndex)
+{
+	RemapRow(UE::Editor::DataStorage::FMapKeyView(OriginalIndex), UE::Editor::DataStorage::FMapKey(NewIndex));
+}
+
+void IEditorDataStorageProvider::RemoveIndex(UE::Editor::DataStorage::IndexHash Index)
+{
+	RemoveRowMapping(UE::Editor::DataStorage::FMapKeyView(Index));
 }
\ No newline at end of file
diff --git a/Engine/Source/Runtime/TypedElementFramework/Public/Elements/Interfaces/TypedElementQueryStorageInterfaces.h b/Engine/Source/Runtime/TypedElementFramework/Public/Elements/Interfaces/TypedElementQueryStorageInterfaces.h
--- a/Engine/Source/Runtime/TypedElementFramework/Public/Elements/Interfaces/TypedElementQueryStorageInterfaces.h
+++ b/Engine/Source/Runtime/TypedElementFramework/Public/Elements/Interfaces/TypedElementQueryStorageInterfaces.h
@@ -5,6 +5,7 @@
 #include "Containers/ContainersFwd.h"
 #include "Elements/Common/TypedElementCommonTypes.h"
 #include "Elements/Common/TypedElementHandles.h"
+#include "Elements/Common/TypedElementMapKey.h"
 #include "Elements/Common/TypedElementQueryTypes.h"
 #include "UObject/Class.h"
 #include "UObject/WeakObjectPtrTemplates.h"
@@ -358,8 +359,8 @@ namespace UE::Editor::DataStorage
 		virtual void GetDependencies(TArrayView<UObject*> RetrievedAddresses, TConstArrayView<TWeakObjectPtr<const UClass>> DependencyTypes,
 			TConstArrayView<EQueryAccessType> AccessTypes) = 0;
 
-		/** Retrieves the row for an indexed object. Returns an invalid row handle if the hash wasn't found. */
-		virtual RowHandle FindIndexedRow(IndexHash Index) const = 0;
+		/** Retrieves the row for a mapped object. Returns an invalid row handle if the key wasn't found. */
+		virtual RowHandle LookupMappedRow(const FMapKeyView& Index) const = 0;
 
 		/**
 		 * Runs a previously created query. This version takes an arbitrary query, but is limited to running queries that do not directly
@@ -384,6 +385,9 @@ namespace UE::Editor::DataStorage
 		 * are executed as part of their parent query and are not scheduled separately.
 		 */
 		virtual FQueryResult RunSubquery(int32 SubqueryIndex, RowHandle Row, SubqueryCallbackRef Callback) = 0;
+
+		UE_DEPRECATED(5.6, "Use 'LookUpMappedRow' instead and use the new FMapKey(View) instead of the explicit 'GenerateIndexHash'-functions.")
+		inline RowHandle FindIndexedRow(IndexHash Index) const;
 	};
 } // namespace UE::Editor::DataStorage
 
@@ -676,4 +680,15 @@ namespace UE::Editor::DataStorage
 		}
 	}
 
+
+
+	//
+	// IQueryContext
+	//
+
+	RowHandle IQueryContext::FindIndexedRow(IndexHash Index) const
+	{
+		return LookupMappedRow(FMapKeyView(Index));
+	}
+
 } // namespace UE::Editor::DataStorage
diff --git a/Engine/Source/Runtime/TypedElementFramework/Tests/Elements/Common/TypedElementMapKeyTests.cpp b/Engine/Source/Runtime/TypedElementFramework/Tests/Elements/Common/TypedElementMapKeyTests.cpp
new file mode 100644
--- /dev/null
+++ b/Engine/Source/Runtime/TypedElementFramework/Tests/Elements/Common/TypedElementMapKeyTests.cpp
@@ -0,0 +1,481 @@
+// Copyright Epic Games, Inc. All Rights Reserved.
+
+#if WITH_TESTS
+
+#include "Elements/Common/TypedElementMapKey.h"
+#include "Tests/TestHarnessAdapter.h"
+
+namespace UE::Editor::DataStorage::Tests::Private
+{
+	FName TestFName = FName("FMapKey_TestName");
+	FSoftObjectPath TestPath = FSoftObjectPath(FName("FMapKey_TestPackage"), FName("FMapKey_TestAsset"), TEXT("FMapKey_TestSubPath"));
+
+	void CopyKeyTest(const FMapKey& Original)
+	{
+		FMapKey Copy = Original;
+		CHECK_MESSAGE(TEXT("Copied key was not empty as expected."), Copy.IsSet());
+		CHECK_MESSAGE(TEXT("Original and copied key don't have the same value."), Original == Copy);
+	}
+
+	void MoveKeyTest(FMapKey&& Original)
+	{
+		FMapKey Moved = MoveTemp(Original);
+		CHECK_MESSAGE(TEXT("Original is still set after moving."), !Original.IsSet());
+		CHECK_MESSAGE(TEXT("Moved key is empty."), Moved.IsSet());
+		CHECK_MESSAGE(TEXT("Original and moved shouldn't have the same value."), Original != Moved);
+	}
+
+	void ConvertKeyToKeyViewTest(const FMapKey& Key)
+	{
+		FMapKeyView View(Key);
+		CHECK_MESSAGE(TEXT("Key to View comparison failed"), Key == View);
+		CHECK_MESSAGE(TEXT("View to Key comparison failed"), View == Key);
+	}
+
+	TEST_CASE_NAMED(FMapKey_Tests, "Editor::DataStorage::Map Key (FMapKey)", "[ApplicationContextMask][EngineFilter]")
+	{
+		using namespace UE::Editor::DataStorage;
+		using namespace UE::Editor::DataStorage::Tests::Private;
+
+		//
+		// Construction
+		//
+
+		SECTION("Empty key")
+		{
+			FMapKey Key;
+			CHECK_MESSAGE(TEXT("Key was not empty as expected."), !Key.IsSet());
+			CHECK_MESSAGE(TEXT("Hash is not zero"), Key.CalculateHash() == 0);
+			CHECK_MESSAGE(TEXT("Expected string wasn't returned"), Key.ToString() == TEXT("Empty"));
+			CHECK_MESSAGE(TEXT("Self comparison failed"), Key == Key);
+		}
+		
+		SECTION("Construct key (void*)")
+		{
+			int* Object = reinterpret_cast<int*>(0x1234);
+			FMapKey Key(Object);
+			CHECK_MESSAGE(TEXT("Key was not set."), Key.IsSet());
+			CHECK_MESSAGE(TEXT("Hash is was zero"), Key.CalculateHash() != 0);
+			FString KeyString = Key.ToString();
+			CHECK_MESSAGE(TEXT("Expected string wasn't returned"), 
+				KeyString.StartsWith(TEXT("Pointer(0x"))
+				&& KeyString.EndsWith(TEXT("1234)")));
+			CHECK_MESSAGE(TEXT("Self comparison failed"), Key == Key);
+		}
+
+		SECTION("Construct key (UObject*)")
+		{
+			UObject* Object = reinterpret_cast<UObject*>(0x1234);
+			FMapKey Key(Object);
+			CHECK_MESSAGE(TEXT("Key was not set."), Key.IsSet());
+			CHECK_MESSAGE(TEXT("Hash is was zero"), Key.CalculateHash() != 0);
+			FString KeyString = Key.ToString();
+			CHECK_MESSAGE(TEXT("Expected string wasn't returned"), 
+				KeyString.StartsWith(TEXT("UObject(0x"))
+				&& KeyString.EndsWith(TEXT("1234)")));
+			CHECK_MESSAGE(TEXT("Self comparison failed"), Key == Key);
+		}
+
+		SECTION("Construct key (int64)")
+		{
+			int64 Object = 42;
+			FMapKey Key(Object);
+			CHECK_MESSAGE(TEXT("Key was not set."), Key.IsSet());
+			CHECK_MESSAGE(TEXT("Hash is was zero"), Key.CalculateHash() != 0);
+			CHECK_MESSAGE(TEXT("Expected string wasn't returned"), Key.ToString() == TEXT("42"));
+			CHECK_MESSAGE(TEXT("Self comparison failed"), Key == Key);
+		}
+
+		SECTION("Construct key (uint64)")
+		{
+			uint64 Object = 42;
+			FMapKey Key(Object);
+			CHECK_MESSAGE(TEXT("Key was not set."), Key.IsSet());
+			CHECK_MESSAGE(TEXT("Hash is was zero"), Key.CalculateHash() != 0);
+			CHECK_MESSAGE(TEXT("Expected string wasn't returned"), Key.ToString() == TEXT("42"));
+			CHECK_MESSAGE(TEXT("Self comparison failed"), Key == Key);
+		}
+
+		SECTION("Construct key (float)")
+		{
+			float Object = 42.0f;
+			FMapKey Key(Object);
+			CHECK_MESSAGE(TEXT("Key was not set."), Key.IsSet());
+			CHECK_MESSAGE(TEXT("Hash is was zero"), Key.CalculateHash() != 0);
+			CHECK_MESSAGE(TEXT("Expected string wasn't returned"), Key.ToString() == TEXT("42.0"));
+			CHECK_MESSAGE(TEXT("Self comparison failed"), Key == Key);
+		}
+
+		SECTION("Construct key (double)")
+		{
+			double Object = 42.0;
+			FMapKey Key(Object);
+			CHECK_MESSAGE(TEXT("Key was not set."), Key.IsSet());
+			CHECK_MESSAGE(TEXT("Hash is was zero"), Key.CalculateHash() != 0);
+			CHECK_MESSAGE(TEXT("Expected string wasn't returned"), Key.ToString() == TEXT("42.0"));
+			CHECK_MESSAGE(TEXT("Self comparison failed"), Key == Key);
+		}
+
+		SECTION("Construct key (FString)")
+		{
+			FString Object = TEXT("TestString");
+			FMapKey Key(Object);
+			CHECK_MESSAGE(TEXT("Key was not set."), Key.IsSet());
+			CHECK_MESSAGE(TEXT("Hash is was zero"), Key.CalculateHash() != 0);
+			CHECK_MESSAGE(TEXT("Expected string wasn't returned"), Key.ToString() == TEXT("TestString"));
+			CHECK_MESSAGE(TEXT("Self comparison failed"), Key == Key);
+		}
+
+		SECTION("Construct key (FName)")
+		{
+			FMapKey Key(TestFName);
+			CHECK_MESSAGE(TEXT("Key was not set."), Key.IsSet());
+			CHECK_MESSAGE(TEXT("Hash is was zero"), Key.CalculateHash() != 0);
+			CHECK_MESSAGE(TEXT("Expected string wasn't returned"), Key.ToString() == TEXT("FMapKey_TestName"));
+			CHECK_MESSAGE(TEXT("Self comparison failed"), Key == Key);
+		}
+
+		SECTION("Construct key (FSoftObjectPath)")
+		{
+			FMapKey Key(TestPath);
+			CHECK_MESSAGE(TEXT("Key was not set."), Key.IsSet());
+			CHECK_MESSAGE(TEXT("Hash is was zero"), Key.CalculateHash() != 0);
+			CHECK_MESSAGE(TEXT("Expected string wasn't returned"), Key.ToString() == TestPath.ToString());
+			CHECK_MESSAGE(TEXT("Self comparison failed"), Key == Key);
+		}
+
+		//
+		// Copy
+		//
+
+		SECTION("Copy empty key")
+		{
+			FMapKey Key0;
+			FMapKey Key1 = Key0;
+			CHECK_MESSAGE(TEXT("Original key was not empty as expected."), !Key1.IsSet());
+			CHECK_MESSAGE(TEXT("Copied key was not empty as expected."), !Key1.IsSet());
+			CHECK_MESSAGE(TEXT("Original and copied key don't have the same value."), Key0 == Key1);
+		}
+
+		SECTION("Copy key (void*)")
+		{
+			int* Object = nullptr;
+			FMapKey Key(Object);
+			CopyKeyTest(Key);
+		}
+
+		SECTION("Copy key (UObject*)")
+		{
+			UObject* Object = nullptr;
+			FMapKey Key(Object);
+			CopyKeyTest(Key);
+		}
+
+		SECTION("Copy key (int64)")
+		{
+			int64 Object = 42;
+			FMapKey Key(Object);
+			CopyKeyTest(Key);
+		}
+
+		SECTION("Copy key (uint64)")
+		{
+			uint64 Object = 42;
+			FMapKey Key(Object);
+			CopyKeyTest(Key);
+		}
+
+		SECTION("Copy key (float)")
+		{
+			float Object = 42.0f;
+			FMapKey Key(Object);
+			CopyKeyTest(Key);
+		}
+
+		SECTION("Copy key (double)")
+		{
+			double Object = 42.0;
+			FMapKey Key(Object);
+			CopyKeyTest(Key);
+		}
+
+		SECTION("Copy key (FString)")
+		{
+			FString Object = TEXT("TestString");
+			FMapKey Key(Object);
+			CopyKeyTest(Key);
+		}
+
+		SECTION("Copy key (FName)")
+		{
+			FMapKey Key(TestFName);
+			CopyKeyTest(Key);
+		}
+
+		SECTION("Copy key (FSoftObjectPath)")
+		{
+			FMapKey Key(TestPath);
+			CopyKeyTest(Key);
+		}
+
+		//
+		// Move
+		//
+
+		SECTION("Move empty key")
+		{
+			FMapKey Key0;
+			FMapKey Key1 = MoveTemp(Key0);
+			CHECK_MESSAGE(TEXT("Original key was not empty as expected."), !Key1.IsSet());
+			CHECK_MESSAGE(TEXT("Copied key was not empty as expected."), !Key1.IsSet());
+			CHECK_MESSAGE(TEXT("Original and copied key don't have the same value."), Key0 == Key1);
+		}
+
+		SECTION("Move key (void*)")
+		{
+			int* Object = nullptr;
+			FMapKey Key(Object);
+			MoveKeyTest(MoveTemp(Key));
+		}
+
+		SECTION("Move key (UObject*)")
+		{
+			UObject* Object = nullptr;
+			FMapKey Key(Object);
+			MoveKeyTest(MoveTemp(Key));
+		}
+
+		SECTION("Move key (int64)")
+		{
+			int64 Object = 42;
+			FMapKey Key(Object);
+			MoveKeyTest(MoveTemp(Key));
+		}
+
+		SECTION("Move key (uint64)")
+		{
+			uint64 Object = 42;
+			FMapKey Key(Object);
+			MoveKeyTest(MoveTemp(Key));
+		}
+
+		SECTION("Move key (float)")
+		{
+			float Object = 42.0f;
+			FMapKey Key(Object);
+			MoveKeyTest(MoveTemp(Key));
+		}
+
+		SECTION("Move key (double)")
+		{
+			double Object = 42.0;
+			FMapKey Key(Object);
+			MoveKeyTest(MoveTemp(Key));
+		}
+
+		SECTION("Move key (FString)")
+		{
+			FString Object = TEXT("TestString");
+			FMapKey Key(Object);
+			MoveKeyTest(MoveTemp(Key));
+		}
+
+		SECTION("Move key (FName)")
+		{
+			FMapKey Key(TestFName);
+			MoveKeyTest(MoveTemp(Key));
+		}
+
+		SECTION("Move key (FSoftObjectPath)")
+		{
+			FMapKey Key(TestPath);
+			MoveKeyTest(MoveTemp(Key));
+		}
+	}
+
+	TEST_CASE_NAMED(FMapKeyView_Tests, "Editor::DataStorage::Map Key View (FMapKeyView)", "[ApplicationContextMask][EngineFilter]")
+	{
+		using namespace UE::Editor::DataStorage;
+		using namespace UE::Editor::DataStorage::Tests::Private;
+
+		//
+		// Construction
+		//
+
+		SECTION("Empty key")
+		{
+			FMapKeyView View;
+			CHECK_MESSAGE(TEXT("Hash is not zero"), View.CalculateHash() == 0);
+			CHECK_MESSAGE(TEXT("Expected string wasn't returned"), View.ToString() == TEXT("Empty"));
+			CHECK_MESSAGE(TEXT("Self comparison failed"), View == View);
+		}
+
+		SECTION("Construct key (void*)")
+		{
+			int* Object = reinterpret_cast<int*>(0x1234);
+			FMapKeyView View(Object);
+			CHECK_MESSAGE(TEXT("Hash is was zero"), View.CalculateHash() != 0);
+			FString ViewString = View.ToString();
+			CHECK_MESSAGE(TEXT("Expected string wasn't returned"),
+				ViewString.StartsWith(TEXT("Pointer(0x"))
+				&& ViewString.EndsWith(TEXT("1234)")));
+			CHECK_MESSAGE(TEXT("Self comparison failed"), View == View);
+		}
+
+		SECTION("Construct key (UObject*)")
+		{
+			UObject* Object = reinterpret_cast<UObject*>(0x1234);
+			FMapKeyView View(Object);
+			CHECK_MESSAGE(TEXT("Hash is was zero"), View.CalculateHash() != 0);
+			FString ViewString = View.ToString();
+			CHECK_MESSAGE(TEXT("Expected string wasn't returned"),
+				ViewString.StartsWith(TEXT("UObject(0x"))
+				&& ViewString.EndsWith(TEXT("1234)")));
+			CHECK_MESSAGE(TEXT("Self comparison failed"), View == View);
+		}
+
+		SECTION("Construct key (int64)")
+		{
+			int64 Object = 42;
+			FMapKeyView View(Object);
+			CHECK_MESSAGE(TEXT("Hash is was zero"), View.CalculateHash() != 0);
+			CHECK_MESSAGE(TEXT("Expected string wasn't returned"), View.ToString() == TEXT("42"));
+			CHECK_MESSAGE(TEXT("Self comparison failed"), View == View);
+		}
+
+		SECTION("Construct key (uint64)")
+		{
+			uint64 Object = 42;
+			FMapKeyView View(Object);
+			CHECK_MESSAGE(TEXT("Hash is was zero"), View.CalculateHash() != 0);
+			CHECK_MESSAGE(TEXT("Expected string wasn't returned"), View.ToString() == TEXT("42"));
+			CHECK_MESSAGE(TEXT("Self comparison failed"), View == View);
+		}
+
+		SECTION("Construct key (float)")
+		{
+			float Object = 42.0f;
+			FMapKeyView View(Object);
+			CHECK_MESSAGE(TEXT("Hash is was zero"), View.CalculateHash() != 0);
+			CHECK_MESSAGE(TEXT("Expected string wasn't returned"), View.ToString() == TEXT("42.0"));
+			CHECK_MESSAGE(TEXT("Self comparison failed"), View == View);
+		}
+
+		SECTION("Construct key (double)")
+		{
+			double Object = 42.0;
+			FMapKeyView View(Object);
+			CHECK_MESSAGE(TEXT("Hash is was zero"), View.CalculateHash() != 0);
+			CHECK_MESSAGE(TEXT("Expected string wasn't returned"), View.ToString() == TEXT("42.0"));
+			CHECK_MESSAGE(TEXT("Self comparison failed"), View == View);
+		}
+
+		SECTION("Construct key (FString)")
+		{
+			FString Object = TEXT("TestString");
+			FMapKeyView View(Object);
+			CHECK_MESSAGE(TEXT("Hash is was zero"), View.CalculateHash() != 0);
+			CHECK_MESSAGE(TEXT("Expected string wasn't returned"), View.ToString() == TEXT("TestString"));
+			CHECK_MESSAGE(TEXT("Self comparison failed"), View == View);
+		}
+
+		SECTION("Construct key (FStringView)")
+		{
+			FString Object = TEXT("TestString");
+			FMapKeyView View = FMapKeyView(FStringView(Object));
+			CHECK_MESSAGE(TEXT("Hash is was zero"), View.CalculateHash() != 0);
+			CHECK_MESSAGE(TEXT("Expected string wasn't returned"), View.ToString() == TEXT("TestString"));
+			CHECK_MESSAGE(TEXT("Self comparison failed"), View == View);
+		}
+
+		SECTION("Construct key (FName)")
+		{
+			FMapKeyView View(TestFName);
+			CHECK_MESSAGE(TEXT("Hash is was zero"), View.CalculateHash() != 0);
+			CHECK_MESSAGE(TEXT("Expected string wasn't returned"), View.ToString() == TEXT("FMapKey_TestName"));
+			CHECK_MESSAGE(TEXT("Self comparison failed"), View == View);
+		}
+
+		SECTION("Construct key (FSoftObjectPath)")
+		{
+			FMapKeyView View(TestPath);
+			CHECK_MESSAGE(TEXT("Hash is was zero"), View.CalculateHash() != 0);
+			CHECK_MESSAGE(TEXT("Expected string wasn't returned"), View.ToString() == TestPath.ToString());
+			CHECK_MESSAGE(TEXT("Self comparison failed"), View == View);
+		}
+
+		//
+		// Conversion
+		//
+
+		SECTION("Convert key (Empty)")
+		{
+			FMapKey Key;
+			ConvertKeyToKeyViewTest(Key);
+		}
+
+		SECTION("Convert key (void*)")
+		{
+			int* Object = reinterpret_cast<int*>(0x1234);
+			FMapKey Key(Object);
+			ConvertKeyToKeyViewTest(Key);
+		}
+
+		SECTION("Convert key (UObject*)")
+		{
+			UObject* Object = reinterpret_cast<UObject*>(0x1234);
+			FMapKey Key(Object);
+			ConvertKeyToKeyViewTest(Key);
+		}
+
+		SECTION("Convert key (int64)")
+		{
+			int64 Object = 42;
+			FMapKey Key(Object);
+			ConvertKeyToKeyViewTest(Key);
+		}
+
+		SECTION("Convert key (uint64)")
+		{
+			uint64 Object = 42;
+			FMapKey Key(Object);
+			ConvertKeyToKeyViewTest(Key);
+		}
+
+		SECTION("Convert key (float)")
+		{
+			float Object = 42.0f;
+			FMapKey Key(Object);
+			ConvertKeyToKeyViewTest(Key);
+		}
+
+		SECTION("Convert key (double)")
+		{
+			double Object = 42.0;
+			FMapKey Key(Object);
+			ConvertKeyToKeyViewTest(Key);
+		}
+
+		SECTION("Convert key (FString)")
+		{
+			FString Object = TEXT("TestString");
+			FMapKey Key(Object);
+			ConvertKeyToKeyViewTest(Key);
+		}
+
+		SECTION("Convert key (FName)")
+		{
+			FMapKey Key(TestFName);
+			ConvertKeyToKeyViewTest(Key);
+		}
+
+		SECTION("Convert key (FSoftObjectPath)")
+		{
+			FMapKey Key(TestPath);
+			ConvertKeyToKeyViewTest(Key);
+		}
+	}
+}
+// namespace UE::Editor::DataStorage::Tests::Private
+
+#endif // WITH_TESTS
