[Truncation Warnings] Fix and enable truncation warnings in TEDS plugin modules, except TedsCore

#rnx

diff --git a/Engine/Plugins/Experimental/EditorDataStorage/Source/TedsUI/TedsUI.Build.cs b/Engine/Plugins/Experimental/EditorDataStorage/Source/TedsUI/TedsUI.Build.cs
--- a/Engine/Plugins/Experimental/EditorDataStorage/Source/TedsUI/TedsUI.Build.cs
+++ b/Engine/Plugins/Experimental/EditorDataStorage/Source/TedsUI/TedsUI.Build.cs
@@ -6,6 +6,9 @@ public class TedsUI : ModuleRules
 {
 	public TedsUI(ReadOnlyTargetRules Target) : base(Target)
 	{
+		// Enable truncation warnings in this module
+		CppCompileWarningSettings.UnsafeTypeCastWarningLevel = WarningLevel.Warning;
+
 		PCHUsage = ModuleRules.PCHUsageMode.UseExplicitOrSharedPCHs;
 
 		if (Target.bBuildEditor)
diff --git a/Engine/Plugins/Experimental/EditorDataStorageFeatures/Source/TedsAlerts/TedsAlerts.build.cs b/Engine/Plugins/Experimental/EditorDataStorageFeatures/Source/TedsAlerts/TedsAlerts.build.cs
--- a/Engine/Plugins/Experimental/EditorDataStorageFeatures/Source/TedsAlerts/TedsAlerts.build.cs
+++ b/Engine/Plugins/Experimental/EditorDataStorageFeatures/Source/TedsAlerts/TedsAlerts.build.cs
@@ -7,6 +7,9 @@ public class TedsAlerts : ModuleRules
 {
 	public TedsAlerts(ReadOnlyTargetRules Target) : base(Target)
 	{
+		// Enable truncation warnings in this module
+		CppCompileWarningSettings.UnsafeTypeCastWarningLevel = WarningLevel.Warning;
+
 		ShortName = "TedsAlerts";
 
 		PCHUsage = ModuleRules.PCHUsageMode.UseExplicitOrSharedPCHs;
diff --git a/Engine/Plugins/Experimental/EditorDataStorageFeatures/Source/TedsAssetData/Private/Widgets/AssetDataThumbnailWidget.cpp b/Engine/Plugins/Experimental/EditorDataStorageFeatures/Source/TedsAssetData/Private/Widgets/AssetDataThumbnailWidget.cpp
--- a/Engine/Plugins/Experimental/EditorDataStorageFeatures/Source/TedsAssetData/Private/Widgets/AssetDataThumbnailWidget.cpp
+++ b/Engine/Plugins/Experimental/EditorDataStorageFeatures/Source/TedsAssetData/Private/Widgets/AssetDataThumbnailWidget.cpp
@@ -41,13 +41,14 @@ TSharedPtr<SWidget> FAssetThumbnailWidgetConstructor::CreateWidget(IEditorDataSt
 
 	bool bIsAsset = DataStorage->HasColumns<FAssetTag>(TargetRow);
 
-	constexpr float ThumbnailSize = 64.f;
+	constexpr uint32 ThumbnailSize = 64;
+	constexpr float ThumbnailSizeAsFloat = static_cast<float>(ThumbnailSize);
 	constexpr float ThumbnailVerticalPadding = 4.f;
 
 	// Thumbnail Box container
 	const TSharedRef<SBox> ThumbnailBox = SNew(SBox)
-		.WidthOverride(ThumbnailSize)
-		.HeightOverride(ThumbnailSize + (ThumbnailVerticalPadding * 2))
+		.WidthOverride(ThumbnailSizeAsFloat)
+		.HeightOverride(ThumbnailSizeAsFloat + (ThumbnailVerticalPadding * 2.f))
 		.Padding(0.f, ThumbnailVerticalPadding);
 
 	// For assets, grab the color from the asset definition
@@ -71,7 +72,7 @@ TSharedPtr<SWidget> FAssetThumbnailWidgetConstructor::CreateWidget(IEditorDataSt
 		ThumbnailConfig.bAllowAssetSpecificThumbnailOverlay = !ThumbnailConfig.bForceGenericThumbnail;
 		ThumbnailConfig.ThumbnailLabel = EThumbnailLabel::ClassName;
 		// TODO: If Teds View items will be bound to the CB size and can change their size, this would need to be bound to that as well
-		// ThumbnailConfig.GenericThumbnailSize = 64.f;
+		// ThumbnailConfig.GenericThumbnailSize = 64;
 		// TODO: This data will need to be integrated in TEDS and later on remove this/update it to use TEDS instead to retrieve the data needed
 		// ThumbnailConfig.AssetSystemInfoProvider = nullptr;
 		ThumbnailConfig.bShowAssetChip = true;
diff --git a/Engine/Plugins/Experimental/EditorDataStorageFeatures/Source/TedsAssetData/TedsAssetData.build.cs b/Engine/Plugins/Experimental/EditorDataStorageFeatures/Source/TedsAssetData/TedsAssetData.build.cs
--- a/Engine/Plugins/Experimental/EditorDataStorageFeatures/Source/TedsAssetData/TedsAssetData.build.cs
+++ b/Engine/Plugins/Experimental/EditorDataStorageFeatures/Source/TedsAssetData/TedsAssetData.build.cs
@@ -7,6 +7,9 @@ public class TedsAssetData : ModuleRules
 {
 	public TedsAssetData(ReadOnlyTargetRules Target) : base(Target)
 	{
+		// Enable truncation warnings in this module
+		CppCompileWarningSettings.UnsafeTypeCastWarningLevel = WarningLevel.Warning;
+
 		ShortName = "TEDSAssetD";
 
 		PCHUsage = ModuleRules.PCHUsageMode.UseExplicitOrSharedPCHs;
diff --git a/Engine/Plugins/Experimental/EditorDataStorageFeatures/Source/TedsContentBrowser/TedsContentBrowser.Build.cs b/Engine/Plugins/Experimental/EditorDataStorageFeatures/Source/TedsContentBrowser/TedsContentBrowser.Build.cs
--- a/Engine/Plugins/Experimental/EditorDataStorageFeatures/Source/TedsContentBrowser/TedsContentBrowser.Build.cs
+++ b/Engine/Plugins/Experimental/EditorDataStorageFeatures/Source/TedsContentBrowser/TedsContentBrowser.Build.cs
@@ -6,6 +6,9 @@ public class TedsContentBrowser : ModuleRules
 {
 	public TedsContentBrowser(ReadOnlyTargetRules Target) : base(Target)
 	{
+		// Enable truncation warnings in this module
+		CppCompileWarningSettings.UnsafeTypeCastWarningLevel = WarningLevel.Warning;
+
 		PCHUsage = ModuleRules.PCHUsageMode.UseExplicitOrSharedPCHs;
 
 		if (Target.bBuildEditor)
diff --git a/Engine/Plugins/Experimental/EditorDataStorageFeatures/Source/TedsDebugger/TedsDebugger.Build.cs b/Engine/Plugins/Experimental/EditorDataStorageFeatures/Source/TedsDebugger/TedsDebugger.Build.cs
--- a/Engine/Plugins/Experimental/EditorDataStorageFeatures/Source/TedsDebugger/TedsDebugger.Build.cs
+++ b/Engine/Plugins/Experimental/EditorDataStorageFeatures/Source/TedsDebugger/TedsDebugger.Build.cs
@@ -6,6 +6,9 @@ public class TedsDebugger : ModuleRules
 {
 	public TedsDebugger(ReadOnlyTargetRules Target) : base(Target)
 	{
+		// Enable truncation warnings in this module
+		CppCompileWarningSettings.UnsafeTypeCastWarningLevel = WarningLevel.Warning;
+
 		PCHUsage = ModuleRules.PCHUsageMode.UseExplicitOrSharedPCHs;
 
 		if (Target.bBuildEditor)
diff --git a/Engine/Plugins/Experimental/EditorDataStorageFeatures/Source/TedsOutliner/Private/Compatibility/SceneOutlinerTedsBridge.cpp b/Engine/Plugins/Experimental/EditorDataStorageFeatures/Source/TedsOutliner/Private/Compatibility/SceneOutlinerTedsBridge.cpp
--- a/Engine/Plugins/Experimental/EditorDataStorageFeatures/Source/TedsOutliner/Private/Compatibility/SceneOutlinerTedsBridge.cpp
+++ b/Engine/Plugins/Experimental/EditorDataStorageFeatures/Source/TedsOutliner/Private/Compatibility/SceneOutlinerTedsBridge.cpp
@@ -507,7 +507,7 @@ FTreeItemIDDealiaser FSceneOutlinerTedsBridge::GetDealiaser()
 void FSceneOutlinerTedsBridge::AssignQuery(UE::Editor::DataStorage::QueryHandle Query, const TConstArrayView<FName> InCellWidgetPurposes)
 {
 	using MatchApproach = IEditorDataStorageUiProvider::EMatchApproach;
-	constexpr int32 DefaultPriorityIndex = 100;
+	constexpr uint8 DefaultPriorityIndex = 100;
 	FSceneOutlinerTedsQueryBinder& Binder = FSceneOutlinerTedsQueryBinder::GetInstance();
 	CellWidgetPurposes = InCellWidgetPurposes;
 
@@ -538,7 +538,7 @@ void FSceneOutlinerTedsBridge::AssignQuery(UE::Editor::DataStorage::QueryHandle
 				 	 */
 					FName FallbackColumn = Binder.FindOutlinerColumnFromTEDSColumns(ColumnTypes);
 					const FSceneOutlinerColumnInfo* FallbackColumnInfo = OutlinerPinned->GetSharedData().ColumnMap.Find(FallbackColumn);
-					int32 ColumnPriority = FallbackColumnInfo ? FallbackColumnInfo->PriorityIndex : DefaultPriorityIndex + IndexOffset;
+					uint8 ColumnPriority = FallbackColumnInfo ? FallbackColumnInfo->PriorityIndex : static_cast<uint8>(FMath::Clamp(DefaultPriorityIndex + static_cast<uint8>(FMath::Clamp(IndexOffset, 0, 255)), 0, 255));
 
 					OutlinerPinned->RemoveColumn(FallbackColumn);
 
@@ -578,9 +578,10 @@ void FSceneOutlinerTedsBridge::AssignQuery(UE::Editor::DataStorage::QueryHandle
 				{
 					TSharedPtr<FTypedElementWidgetConstructor> CellConstructor(Constructor.Release());
 					FName NameId = FName(ColumnType->GetDisplayNameText().ToString());
+					uint8 ColumnPriority = static_cast<uint8>(FMath::Clamp(DefaultPriorityIndex + static_cast<uint8>(FMath::Clamp(IndexOffset, 0, 255)), 0, 255));
 					AddedColumns.Add(NameId);
 					OutlinerPinned->AddColumn(NameId,
-						FSceneOutlinerColumnInfo(ESceneOutlinerColumnVisibility::Visible, DefaultPriorityIndex + IndexOffset,
+						FSceneOutlinerColumnInfo(ESceneOutlinerColumnVisibility::Visible, ColumnPriority,
 							FCreateSceneOutlinerColumn::CreateLambda(
 								[this, Query, MetaDataView, NameId, ColumnType, CellConstructor, &OutlinerPinned, FallbackColumn](ISceneOutliner&)
 								{
diff --git a/Engine/Plugins/Experimental/EditorDataStorageFeatures/Source/TedsOutliner/TedsOutliner.Build.cs b/Engine/Plugins/Experimental/EditorDataStorageFeatures/Source/TedsOutliner/TedsOutliner.Build.cs
--- a/Engine/Plugins/Experimental/EditorDataStorageFeatures/Source/TedsOutliner/TedsOutliner.Build.cs
+++ b/Engine/Plugins/Experimental/EditorDataStorageFeatures/Source/TedsOutliner/TedsOutliner.Build.cs
@@ -6,6 +6,9 @@ public class TedsOutliner : ModuleRules
 {
 	public TedsOutliner(ReadOnlyTargetRules Target) : base(Target)
 	{
+		// Enable truncation warnings in this module
+		CppCompileWarningSettings.UnsafeTypeCastWarningLevel = WarningLevel.Warning;
+
 		PCHUsage = ModuleRules.PCHUsageMode.UseExplicitOrSharedPCHs;
 
 		if (Target.bBuildEditor)
diff --git a/Engine/Plugins/Experimental/EditorDataStorageFeatures/Source/TedsPropertyEditor/TedsPropertyEditor.Build.cs b/Engine/Plugins/Experimental/EditorDataStorageFeatures/Source/TedsPropertyEditor/TedsPropertyEditor.Build.cs
--- a/Engine/Plugins/Experimental/EditorDataStorageFeatures/Source/TedsPropertyEditor/TedsPropertyEditor.Build.cs
+++ b/Engine/Plugins/Experimental/EditorDataStorageFeatures/Source/TedsPropertyEditor/TedsPropertyEditor.Build.cs
@@ -6,6 +6,9 @@ public class TedsPropertyEditor : ModuleRules
 {
 	public TedsPropertyEditor(ReadOnlyTargetRules Target) : base(Target)
 	{
+		// Enable truncation warnings in this module
+		CppCompileWarningSettings.UnsafeTypeCastWarningLevel = WarningLevel.Warning;
+
 		PCHUsage = ModuleRules.PCHUsageMode.UseExplicitOrSharedPCHs;
 
 		if (Target.bBuildEditor)
diff --git a/Engine/Plugins/Experimental/EditorDataStorageFeatures/Source/TedsQueryStack/TedsQueryStack.Build.cs b/Engine/Plugins/Experimental/EditorDataStorageFeatures/Source/TedsQueryStack/TedsQueryStack.Build.cs
--- a/Engine/Plugins/Experimental/EditorDataStorageFeatures/Source/TedsQueryStack/TedsQueryStack.Build.cs
+++ b/Engine/Plugins/Experimental/EditorDataStorageFeatures/Source/TedsQueryStack/TedsQueryStack.Build.cs
@@ -6,6 +6,9 @@ public class TedsQueryStack : ModuleRules
 {
 	public TedsQueryStack(ReadOnlyTargetRules Target) : base(Target)
 	{
+		// Enable truncation warnings in this module
+		CppCompileWarningSettings.UnsafeTypeCastWarningLevel = WarningLevel.Warning;
+
 		PCHUsage = ModuleRules.PCHUsageMode.UseExplicitOrSharedPCHs;
 
 		if (Target.bBuildEditor)
diff --git a/Engine/Plugins/Experimental/EditorDataStorageFeatures/Source/TedsRevisionControl/Private/Processors/RevisionControlProcessors.cpp b/Engine/Plugins/Experimental/EditorDataStorageFeatures/Source/TedsRevisionControl/Private/Processors/RevisionControlProcessors.cpp
--- a/Engine/Plugins/Experimental/EditorDataStorageFeatures/Source/TedsRevisionControl/Private/Processors/RevisionControlProcessors.cpp
+++ b/Engine/Plugins/Experimental/EditorDataStorageFeatures/Source/TedsRevisionControl/Private/Processors/RevisionControlProcessors.cpp
@@ -155,7 +155,7 @@ namespace UE::Editor::RevisionControl::Private
 	#if ENABLE_OVERLAY_DEBUG
 		if (int32 Force = CVarDebugForceColorOnAll.GetValueOnGameThread())
 		{
-			int32 Alpha = FMath::Lerp<float>(0.f, 255.f, gOverlayAlpha / 100.f);
+			uint8 Alpha = FMath::Lerp<uint8>(static_cast<uint8>(0), static_cast<uint8>(255), static_cast<float>(gOverlayAlpha) / 100.f);
 			switch(Force)
 			{
 			case 1:
@@ -176,7 +176,7 @@ namespace UE::Editor::RevisionControl::Private
 		if (!bIgnored && !bSelected)
 		{
 			// Convert CVar value from [0..100] to [0..255] range.
-			int32 Alpha = FMath::Lerp<float>(0.f, 255.f, CVarOverlayAlpha.GetValueOnGameThread() / 100.f);
+			uint8 Alpha = FMath::Lerp<uint8>(static_cast<uint8>(0), static_cast<uint8>(255), static_cast<float>(CVarOverlayAlpha.GetValueOnGameThread()) / 100.f);
 
 			// Check if the package is outdated because there is a newer version available.
 			if (SCCContext.HasColumn<FSCCNotCurrentTag>())
diff --git a/Engine/Plugins/Experimental/EditorDataStorageFeatures/Source/TedsRevisionControl/TedsRevisionControl.Build.cs b/Engine/Plugins/Experimental/EditorDataStorageFeatures/Source/TedsRevisionControl/TedsRevisionControl.Build.cs
--- a/Engine/Plugins/Experimental/EditorDataStorageFeatures/Source/TedsRevisionControl/TedsRevisionControl.Build.cs
+++ b/Engine/Plugins/Experimental/EditorDataStorageFeatures/Source/TedsRevisionControl/TedsRevisionControl.Build.cs
@@ -6,6 +6,9 @@ public class TedsRevisionControl : ModuleRules
 {
 	public TedsRevisionControl(ReadOnlyTargetRules Target) : base(Target)
 	{
+		// Enable truncation warnings in this module
+		CppCompileWarningSettings.UnsafeTypeCastWarningLevel = WarningLevel.Warning;
+
 		PCHUsage = ModuleRules.PCHUsageMode.UseExplicitOrSharedPCHs;
 
 		if (Target.bBuildEditor)
diff --git a/Engine/Plugins/Experimental/EditorDataStorageFeatures/Source/TedsSettings/TedsSettings.Build.cs b/Engine/Plugins/Experimental/EditorDataStorageFeatures/Source/TedsSettings/TedsSettings.Build.cs
--- a/Engine/Plugins/Experimental/EditorDataStorageFeatures/Source/TedsSettings/TedsSettings.Build.cs
+++ b/Engine/Plugins/Experimental/EditorDataStorageFeatures/Source/TedsSettings/TedsSettings.Build.cs
@@ -6,6 +6,9 @@ public class TedsSettings : ModuleRules
 {
 	public TedsSettings(ReadOnlyTargetRules Target) : base(Target)
 	{
+		// Enable truncation warnings in this module
+		CppCompileWarningSettings.UnsafeTypeCastWarningLevel = WarningLevel.Warning;
+
 		PCHUsage = ModuleRules.PCHUsageMode.UseExplicitOrSharedPCHs;
 
 		if (Target.bBuildEditor)
diff --git a/Engine/Plugins/Experimental/EditorDataStorageFeatures/Source/TedsStyling/TedsStyling.Build.cs b/Engine/Plugins/Experimental/EditorDataStorageFeatures/Source/TedsStyling/TedsStyling.Build.cs
--- a/Engine/Plugins/Experimental/EditorDataStorageFeatures/Source/TedsStyling/TedsStyling.Build.cs
+++ b/Engine/Plugins/Experimental/EditorDataStorageFeatures/Source/TedsStyling/TedsStyling.Build.cs
@@ -6,6 +6,9 @@ public class TedsStyling : ModuleRules
 {
 	public TedsStyling(ReadOnlyTargetRules Target) : base(Target)
 	{
+		// Enable truncation warnings in this module
+		CppCompileWarningSettings.UnsafeTypeCastWarningLevel = WarningLevel.Warning;
+
 		PCHUsage = ModuleRules.PCHUsageMode.UseExplicitOrSharedPCHs;
 
 		if (Target.bBuildEditor)
diff --git a/Engine/Plugins/Experimental/EditorDataStorageFeatures/Source/TedsTableViewer/TedsTableViewer.Build.cs b/Engine/Plugins/Experimental/EditorDataStorageFeatures/Source/TedsTableViewer/TedsTableViewer.Build.cs
--- a/Engine/Plugins/Experimental/EditorDataStorageFeatures/Source/TedsTableViewer/TedsTableViewer.Build.cs
+++ b/Engine/Plugins/Experimental/EditorDataStorageFeatures/Source/TedsTableViewer/TedsTableViewer.Build.cs
@@ -6,6 +6,8 @@ public class TedsTableViewer : ModuleRules
 {
 	public TedsTableViewer(ReadOnlyTargetRules Target) : base(Target)
 	{
+		// Enable truncation warnings in this module
+		CppCompileWarningSettings.UnsafeTypeCastWarningLevel = WarningLevel.Warning;
 
 		if (Target.bBuildEditor)
 		{
