TEDS Debugger and Table Viewer: Misc fixes after introducing dynamic columns

diff --git a/Engine/Plugins/Experimental/EditorDataStorageFeatures/Source/TedsDebugger/Private/QueryEditor/TedsQueryEditorModel.cpp b/Engine/Plugins/Experimental/EditorDataStorageFeatures/Source/TedsDebugger/Private/QueryEditor/TedsQueryEditorModel.cpp
--- a/Engine/Plugins/Experimental/EditorDataStorageFeatures/Source/TedsDebugger/Private/QueryEditor/TedsQueryEditorModel.cpp
+++ b/Engine/Plugins/Experimental/EditorDataStorageFeatures/Source/TedsDebugger/Private/QueryEditor/TedsQueryEditorModel.cpp
@@ -3,6 +3,7 @@
 
 #include "Elements/Common/TypedElementCommonTypes.h"
 #include "Elements/Common/TypedElementDataStorageLog.h"
+#include "Elements/Framework/TypedElementColumnUtils.h"
 #include "Elements/Interfaces/TypedElementDataStorageInterface.h"
 #include "UObject/UObjectIterator.h"
 
@@ -15,12 +16,12 @@
 
 	bool FConditionEntryHandle::IsValid() const
 	{
-		return Id != INDEX_NONE;
+		return Id != NAME_None;
 	}
 
 	void FConditionEntryHandle::Reset()
 	{
-		Id = INDEX_NONE;
+		Id = NAME_None;
 	}
 
 	FTedsQueryEditorModel::FTedsQueryEditorModel(IEditorDataStorageProvider& InDataStorageProvider)
@@ -30,64 +31,7 @@
 
 	void FTedsQueryEditorModel::Reset()
 	{
-		TArray<const UScriptStruct*> Columns;
-		TArray<const UScriptStruct*> Tags;
-
-		UScriptStruct* ColumnType = FColumn::StaticStruct();
-		UScriptStruct* TagType = FTag::StaticStruct();
-		// Not sure if there is a faster way to do this.  Would be nice to iterate only the derived classes
-		for(TObjectIterator< UScriptStruct > It; It; ++It)
-		{
-			UScriptStruct* Struct = *It;
-
-			if (Struct->IsChildOf(ColumnType) && Struct != ColumnType)
-			{
-				Columns.Add(Struct);
-			}
-			if (Struct->IsChildOf(TagType) && Struct != TagType)
-			{
-				Tags.Add(Struct);
-			}
-		}
-
-		Conditions.Empty();
-
-		Conditions.Reserve(Columns.Num() + Tags.Num());
-
-		auto NextId = [this]()
-		{
-			if (IdGenerator != TNumericLimits<decltype(IdGenerator)>::Max())
-			{
-				return IdGenerator++;
-			}
-			else
-			{
-				IdGenerator = 0;
-				return IdGenerator;
-			}
-		};
-
-		for (const UScriptStruct* Column : Columns)
-		{
-			Conditions.Emplace(
-			FConditionEntryInternal
-				{
-					.Id = NextId(),
-					.Struct = Column,
-					.OperatorType = EOperatorType::Unset
-				});
-		}
-		for (const UScriptStruct* Tag : Tags)
-		{
-			Conditions.Emplace(
-			FConditionEntryInternal
-				{
-					.Id = NextId(),
-					.Struct = Tag,
-					.OperatorType = EOperatorType::Unset
-				});
-		}
-
+		RegenerateColumnsList();
 		ModelChangedDelegate.Broadcast();
 	}
 
@@ -265,11 +209,18 @@
 		{
 			if (Entry.OperatorType == EOperatorType::Unset)
 			{
-				// Don't show Tags for Select operators... not valid query
-				if (OperatorType == EOperatorType::Select && Entry.Struct->IsChildOf(TagType))
+				// Don't show Tags or base dynamic templates for Select operators... not valid query
+				if (OperatorType == EOperatorType::Select &&
+						(Entry.Struct->IsChildOf(TagType) || ColumnUtils::IsDynamicTemplate(Entry.Struct)))
 				{
 					continue;
 				}
+				// Don't show base templates in All since it doesn't make much sense
+				if (OperatorType == EOperatorType::All && ColumnUtils::IsDynamicTemplate(Entry.Struct))
+				{
+					continue;
+				}
+				
 				FConditionEntryHandle Handle;
 				Handle.Id = Entry.Id;
 				Function(*this, Handle);
@@ -389,4 +340,85 @@
 	{
 		return Conditions.FindByPredicate([Handle](const FConditionEntryInternal& Entry){ return Entry.Id == Handle.Id; });
 	}
+
+	void FTedsQueryEditorModel::RegenerateColumnsList()
+	{
+		TArray<const UScriptStruct*> Columns;
+		TArray<const UScriptStruct*> Tags;
+
+		UScriptStruct* ColumnType = FColumn::StaticStruct();
+		UScriptStruct* TagType = FTag::StaticStruct();
+		// Not sure if there is a faster way to do this.  Would be nice to iterate only the derived classes
+		for(TObjectIterator< UScriptStruct > It; It; ++It)
+		{
+			UScriptStruct* Struct = *It;
+
+			if (Struct->IsChildOf(ColumnType) && Struct != ColumnType)
+			{
+				Columns.Add(Struct);
+			}
+			if (Struct->IsChildOf(TagType) && Struct != TagType)
+			{
+				Tags.Add(Struct);
+			}
+		}
+
+		// Save a map of any columns that previously had valid operators so we can restore them
+		TMap<FName, EOperatorType> PreviousValidOperators;
+		
+		ForEachCondition([&PreviousValidOperators](FTedsQueryEditorModel& Model, FConditionEntryHandle Handle)
+			{
+				switch(EOperatorType Operator = Model.GetOperatorType(Handle))
+				{
+				case EOperatorType::Select:
+				case EOperatorType::All:
+				case EOperatorType::Any:
+				case EOperatorType::None:
+					PreviousValidOperators.Emplace(Handle.Id, Operator);
+					break;
+				case EOperatorType::Unset:
+				case EOperatorType::Invalid:
+				default:
+					return;
+				};
+			});
+
+		Conditions.Empty();
+		Conditions.Reserve(Columns.Num() + Tags.Num());
+
+		auto GetPreviousOperator = [PreviousValidOperators] (const UScriptStruct* Column)
+		{
+			if (const EOperatorType* PreviousOp = PreviousValidOperators.Find(Column->GetFName()))
+			{
+				return *PreviousOp;
+			}
+
+			return EOperatorType::Unset;
+		};
+
+		for (const UScriptStruct* Column : Columns)
+		{
+			EOperatorType Operator = GetPreviousOperator(Column);
+			
+			Conditions.Emplace(
+			FConditionEntryInternal
+				{
+					.Id = Column->GetFName(),
+					.Struct = Column,
+					.OperatorType = Operator
+				});
+		}
+		for (const UScriptStruct* Tag : Tags)
+		{
+			EOperatorType Operator = GetPreviousOperator(Tag);
+			
+			Conditions.Emplace(
+			FConditionEntryInternal
+				{
+					.Id = Tag->GetFName(),
+					.Struct = Tag,
+					.OperatorType = Operator
+				});
+		}
+	}
 } // namespace UE::Editor::DataStorage::Debug::QueryEditor
diff --git a/Engine/Plugins/Experimental/EditorDataStorageFeatures/Source/TedsDebugger/Private/QueryEditor/TedsQueryEditorModel.h b/Engine/Plugins/Experimental/EditorDataStorageFeatures/Source/TedsDebugger/Private/QueryEditor/TedsQueryEditorModel.h
--- a/Engine/Plugins/Experimental/EditorDataStorageFeatures/Source/TedsDebugger/Private/QueryEditor/TedsQueryEditorModel.h
+++ b/Engine/Plugins/Experimental/EditorDataStorageFeatures/Source/TedsDebugger/Private/QueryEditor/TedsQueryEditorModel.h
@@ -40,7 +40,7 @@ namespace UE::Editor::DataStorage::Debug::QueryEditor
 		friend class FTedsQueryEditorModel;
 		// Unique ID for a condition
 		// Used to identify the entry by external code (ie. list views)
-		int32 Id = INDEX_NONE;
+		FName Id = NAME_None;
 	};
 
 	DECLARE_MULTICAST_DELEGATE(FTedsQueryEditorModel_ModelChanged);
@@ -88,6 +88,8 @@ namespace UE::Editor::DataStorage::Debug::QueryEditor
 		const UScriptStruct* GetColumnScriptStruct(FConditionEntryHandle Handle) const;
 		
 		FTedsQueryEditorModel_ModelChanged& GetModelChangedDelegate() { return ModelChangedDelegate; }
+		
+		void RegenerateColumnsList();
 	
 	private:
 		
@@ -95,7 +97,7 @@ namespace UE::Editor::DataStorage::Debug::QueryEditor
 		{
 			// Unique ID for a condition
 			// Used to identify the entry by external code (ie. list views)
-			int32 Id;
+			FName Id;
 			const UScriptStruct* Struct;
 			EOperatorType OperatorType;
 
@@ -107,8 +109,6 @@ namespace UE::Editor::DataStorage::Debug::QueryEditor
 
 		IEditorDataStorageProvider& EditorDataStorageProvider;
 
-		int32 IdGenerator = 0;
-
 		TArray<FConditionEntryInternal> Conditions;
 		FTedsQueryEditorModel_ModelChanged ModelChangedDelegate;
 
diff --git a/Engine/Plugins/Experimental/EditorDataStorageFeatures/Source/TedsDebugger/Private/Widgets/QueryEditor/TedsConditionSelectionComboWidget.cpp b/Engine/Plugins/Experimental/EditorDataStorageFeatures/Source/TedsDebugger/Private/Widgets/QueryEditor/TedsConditionSelectionComboWidget.cpp
--- a/Engine/Plugins/Experimental/EditorDataStorageFeatures/Source/TedsDebugger/Private/Widgets/QueryEditor/TedsConditionSelectionComboWidget.cpp
+++ b/Engine/Plugins/Experimental/EditorDataStorageFeatures/Source/TedsDebugger/Private/Widgets/QueryEditor/TedsConditionSelectionComboWidget.cpp
@@ -42,6 +42,11 @@ namespace UE::Editor::DataStorage::Debug::QueryEditor
 		
 		ComboBox = SNew(SComboBox<TSharedPtr<FComboItem>>)
 			.OptionsSource(&ComboItems)
+			.OnComboBoxOpening_Lambda([this]()
+			{
+				Model->RegenerateColumnsList();
+				PopulateComboItems();
+			})
 			.OnSelectionChanged(this, &SConditionComboWidget::OnSelectionChanged)
 			.OnGenerateWidget(this, &SConditionComboWidget::OnGenerateWidget)
 			.IsEnabled_Lambda([this](){ return !ComboItems.IsEmpty(); })
diff --git a/Engine/Plugins/Experimental/EditorDataStorageFeatures/Source/TedsTableViewer/Private/TedsTableViewerModel.cpp b/Engine/Plugins/Experimental/EditorDataStorageFeatures/Source/TedsTableViewer/Private/TedsTableViewerModel.cpp
--- a/Engine/Plugins/Experimental/EditorDataStorageFeatures/Source/TedsTableViewer/Private/TedsTableViewerModel.cpp
+++ b/Engine/Plugins/Experimental/EditorDataStorageFeatures/Source/TedsTableViewer/Private/TedsTableViewerModel.cpp
@@ -2,7 +2,6 @@
 
 #include "TedsTableViewerModel.h"
 
-
 #include "Elements/Columns/TypedElementUIColumns.h"
 #include "Elements/Common/EditorDataStorageFeatures.h"
 #include "Elements/Interfaces/TypedElementDataStorageInterface.h"
@@ -11,6 +10,8 @@
 #include "TedsTableViewerColumn.h"
 #include "TedsTableViewerUtils.h"
 
+DEFINE_LOG_CATEGORY_STATIC(LogTableViewer, Log, All)
+
 namespace UE::Editor::DataStorage
 {
 	FTedsTableViewerModel::FTedsTableViewerModel(const TSharedPtr<QueryStack::IRowNode>& InRowQueryStack,
@@ -61,6 +62,21 @@ namespace UE::Editor::DataStorage
 		return !Storage->HasColumns<FHideRowFromUITag>(InRowHandle);
 	}
 
+	void FTedsTableViewerModel::ValidateRequestedColumns()
+	{
+		RequestedTedsColumns.RemoveAll([](const TWeakObjectPtr<const UScriptStruct>& Column)
+		{
+			if (ColumnUtils::IsDynamicTemplate(Column.Get()))
+			{
+				UE_LOG(LogTableViewer, Log, TEXT("%s Column is a dynamic template which cannot be displayed in the table viewer and has been removed!"),
+						*Column->GetName());
+				return true;
+			}
+
+			return false;
+		});
+	}
+
 	bool FTedsTableViewerModel::Tick(float DeltaTime)
 	{
 		// If the revision ID has changed, refresh to update our rows
@@ -147,6 +163,8 @@ namespace UE::Editor::DataStorage
 
 	void FTedsTableViewerModel::GenerateColumns()
 	{
+		ValidateRequestedColumns();
+		
 		using MatchApproach = IEditorDataStorageUiProvider::EMatchApproach;
 		int32 IndexOffset = 0;
 
diff --git a/Engine/Plugins/Experimental/EditorDataStorageFeatures/Source/TedsTableViewer/Public/TedsTableViewerModel.h b/Engine/Plugins/Experimental/EditorDataStorageFeatures/Source/TedsTableViewer/Public/TedsTableViewerModel.h
--- a/Engine/Plugins/Experimental/EditorDataStorageFeatures/Source/TedsTableViewer/Public/TedsTableViewerModel.h
+++ b/Engine/Plugins/Experimental/EditorDataStorageFeatures/Source/TedsTableViewer/Public/TedsTableViewerModel.h
@@ -78,6 +78,8 @@ namespace UE::Editor::DataStorage
 		// Check whether a row is allowed to be displayed in the table viewer
 		bool IsRowDisplayable(RowHandle InRowHandle) const;
 
+		void ValidateRequestedColumns();
+
 	private:
 
 		// The row query stack used to supply the rows to display
