Added an std::initializer_list constructor to FRowHandleArrayView.

There were lifecycle issues when constructing an FRowHandleArrayView with an inline array, e.g. FRowHandleArrayView({1,2,3}, ...). With the initializer list and lifetime checks this should be avoided now.

diff --git a/Engine/Plugins/Experimental/EditorDataStorageFeatures/Source/TedsQueryStack/Tests/TedsRowMergeNodeTests.cpp b/Engine/Plugins/Experimental/EditorDataStorageFeatures/Source/TedsQueryStack/Tests/TedsRowMergeNodeTests.cpp
--- a/Engine/Plugins/Experimental/EditorDataStorageFeatures/Source/TedsQueryStack/Tests/TedsRowMergeNodeTests.cpp
+++ b/Engine/Plugins/Experimental/EditorDataStorageFeatures/Source/TedsQueryStack/Tests/TedsRowMergeNodeTests.cpp
@@ -22,7 +22,8 @@ TEST_CASE_NAMED(TEDS_QueryStack_RowMergeNode_Tests, "Editor::DataStorage::QueryS
 
 	SECTION("Append two nodes")
 	{
-		FRowHandleArrayView Values = FRowHandleArrayView( { 1, 2, 3 }, 
+		RowHandle ValueArray[] = {1, 2, 3};
+		FRowHandleArrayView Values = FRowHandleArrayView( ValueArray,
 			FRowHandleArrayView::EFlags::IsSorted | FRowHandleArrayView::EFlags::IsUnique);
 		TSharedPtr<IRowNode> View0 = MakeShared<FRowViewNode>(Values);
 		TSharedPtr<IRowNode> View1 = MakeShared<FRowViewNode>(Values);
@@ -42,7 +43,8 @@ TEST_CASE_NAMED(TEDS_QueryStack_RowMergeNode_Tests, "Editor::DataStorage::QueryS
 
 	SECTION("Merge two nodes sorted")
 	{
-		FRowHandleArrayView Values = FRowHandleArrayView( { 1, 2, 3 },
+		RowHandle ValueArray[] = { 1, 2, 3 }; 
+		FRowHandleArrayView Values = FRowHandleArrayView( ValueArray,
 			FRowHandleArrayView::EFlags::IsSorted | FRowHandleArrayView::EFlags::IsUnique);
 		TSharedPtr<IRowNode> View0 = MakeShared<FRowViewNode>(Values);
 		TSharedPtr<IRowNode> View1 = MakeShared<FRowViewNode>(Values);
@@ -62,7 +64,8 @@ TEST_CASE_NAMED(TEDS_QueryStack_RowMergeNode_Tests, "Editor::DataStorage::QueryS
 
 	SECTION("Uniquely merge two nodes")
 	{
-		FRowHandleArrayView Values = FRowHandleArrayView( { 1, 2, 3 },
+		RowHandle ValueArray[] = { 1, 2, 3 };
+		FRowHandleArrayView Values = FRowHandleArrayView( ValueArray,
 			FRowHandleArrayView::EFlags::IsSorted | FRowHandleArrayView::EFlags::IsUnique);
 		TSharedPtr<IRowNode> View0 = MakeShared<FRowViewNode>(Values);
 		TSharedPtr<IRowNode> View1 = MakeShared<FRowViewNode>(Values);
@@ -79,9 +82,11 @@ TEST_CASE_NAMED(TEDS_QueryStack_RowMergeNode_Tests, "Editor::DataStorage::QueryS
 
 	SECTION("Repeating of data in two nodes")
 	{
-		FRowHandleArrayView Values0 = FRowHandleArrayView( { 1, 2, 3 },
+		RowHandle ValueArray0[] = { 1, 2, 3 };
+		RowHandle ValueArray1[] = { 2, 3, 4 };
+		FRowHandleArrayView Values0 = FRowHandleArrayView( ValueArray0,
 			FRowHandleArrayView::EFlags::IsSorted | FRowHandleArrayView::EFlags::IsUnique);
-		FRowHandleArrayView Values1 = FRowHandleArrayView( { 2, 3, 4 },
+		FRowHandleArrayView Values1 = FRowHandleArrayView( ValueArray1,
 			FRowHandleArrayView::EFlags::IsSorted | FRowHandleArrayView::EFlags::IsUnique);
 		TSharedPtr<IRowNode> View0 = MakeShared<FRowViewNode>(Values0);
 		TSharedPtr<IRowNode> View1 = MakeShared<FRowViewNode>(Values1);
@@ -97,11 +102,15 @@ TEST_CASE_NAMED(TEDS_QueryStack_RowMergeNode_Tests, "Editor::DataStorage::QueryS
 
 	SECTION("Repeating of data in three nodes")
 	{
-		FRowHandleArrayView Values0 = FRowHandleArrayView({ 1, 2, 3 },
+		RowHandle ValueArray0[] = { 1, 2, 3 };
+		RowHandle ValueArray1[] = { 2, 3, 4 };
+		RowHandle ValueArray2[] = { 1, 3 };
+
+		FRowHandleArrayView Values0 = FRowHandleArrayView(ValueArray0,
 			FRowHandleArrayView::EFlags::IsSorted | FRowHandleArrayView::EFlags::IsUnique);
-		FRowHandleArrayView Values1 = FRowHandleArrayView({ 2, 3, 4 },
+		FRowHandleArrayView Values1 = FRowHandleArrayView(ValueArray1,
 			FRowHandleArrayView::EFlags::IsSorted | FRowHandleArrayView::EFlags::IsUnique);
-		FRowHandleArrayView Values2 = FRowHandleArrayView({ 1, 3 },
+		FRowHandleArrayView Values2 = FRowHandleArrayView(ValueArray2,
 			FRowHandleArrayView::EFlags::IsSorted | FRowHandleArrayView::EFlags::IsUnique);
 		TSharedPtr<IRowNode> View0 = MakeShared<FRowViewNode>(Values0);
 		TSharedPtr<IRowNode> View1 = MakeShared<FRowViewNode>(Values1);
diff --git a/Engine/Plugins/Experimental/EditorDataStorageFeatures/Source/TedsQueryStack/Tests/TedsRowViewNodeTests.cpp b/Engine/Plugins/Experimental/EditorDataStorageFeatures/Source/TedsQueryStack/Tests/TedsRowViewNodeTests.cpp
--- a/Engine/Plugins/Experimental/EditorDataStorageFeatures/Source/TedsQueryStack/Tests/TedsRowViewNodeTests.cpp
+++ b/Engine/Plugins/Experimental/EditorDataStorageFeatures/Source/TedsQueryStack/Tests/TedsRowViewNodeTests.cpp
@@ -19,7 +19,8 @@ TEST_CASE_NAMED(TEDS_QueryStack_RowViewNode_Tests, "Editor::DataStorage::QuerySt
 
 	SECTION("View with a few row handles")
 	{
-		FRowHandleArrayView RowHandles({ 1, 2, 3 }, FRowHandleArrayView::EFlags::IsSorted | FRowHandleArrayView::EFlags::IsUnique);
+		RowHandle ValueArray[] = { 1, 2, 3 };
+		FRowHandleArrayView RowHandles(ValueArray, FRowHandleArrayView::EFlags::IsSorted | FRowHandleArrayView::EFlags::IsUnique);
 		FRowViewNode View(RowHandles);
 		CHECK_EQUALS(TEXT("Revision"), View.GetRevision(), 0);
 		
diff --git a/Engine/Source/Runtime/TypedElementFramework/Private/Elements/Framework/TypedElementRowHandleArrayView.cpp b/Engine/Source/Runtime/TypedElementFramework/Private/Elements/Framework/TypedElementRowHandleArrayView.cpp
--- a/Engine/Source/Runtime/TypedElementFramework/Private/Elements/Framework/TypedElementRowHandleArrayView.cpp
+++ b/Engine/Source/Runtime/TypedElementFramework/Private/Elements/Framework/TypedElementRowHandleArrayView.cpp
@@ -20,6 +20,13 @@ namespace UE::Editor::DataStorage
 		, bIsUnique(EnumHasAllFlags(Flags, EFlags::IsUnique))
 	{}
 
+	FRowHandleArrayView::FRowHandleArrayView(std::initializer_list<RowHandle> Rows, EFlags Flags)
+		: Rows(std::data(Rows))
+		, RowCount(static_cast<int32>(Rows.size()))
+		, bIsSorted(EnumHasAllFlags(Flags, EFlags::IsSorted))
+		, bIsUnique(EnumHasAllFlags(Flags, EFlags::IsUnique))
+	{}
+
 	const RowHandle* FRowHandleArrayView::begin() const
 	{
 		return Rows;
diff --git a/Engine/Source/Runtime/TypedElementFramework/Public/Elements/Framework/TypedElementRowHandleArrayView.h b/Engine/Source/Runtime/TypedElementFramework/Public/Elements/Framework/TypedElementRowHandleArrayView.h
--- a/Engine/Source/Runtime/TypedElementFramework/Public/Elements/Framework/TypedElementRowHandleArrayView.h
+++ b/Engine/Source/Runtime/TypedElementFramework/Public/Elements/Framework/TypedElementRowHandleArrayView.h
@@ -2,8 +2,10 @@
 
 #pragma once
 
+#include <initializer_list>
 #include "Containers/ArrayView.h"
 #include "Elements/Common/TypedElementHandles.h"
+#include "HAL/Platform.h"
 #include "Misc/EnumClassFlags.h"
 
 namespace UE::Editor::DataStorage
@@ -25,8 +27,9 @@ namespace UE::Editor::DataStorage
 		};
 
 		FRowHandleArrayView() = default;
-		TYPEDELEMENTFRAMEWORK_API FRowHandleArrayView(const RowHandle* Rows, int32 RowCount, EFlags Flags);
+		TYPEDELEMENTFRAMEWORK_API FRowHandleArrayView(const RowHandle* Rows UE_LIFETIMEBOUND, int32 RowCount, EFlags Flags);
 		TYPEDELEMENTFRAMEWORK_API FRowHandleArrayView(TConstArrayView<RowHandle> Rows, EFlags Flags);
+		TYPEDELEMENTFRAMEWORK_API FRowHandleArrayView(std::initializer_list<RowHandle> Rows UE_LIFETIMEBOUND, EFlags Flags);
 
 		/** 
 		 * Returns a pointer to the first row in the view.
