[EditorDataStorageFeature] Adding Dynamic Asset Data Columns
Adding Dynamic Asset Data Columns for:
- Parent BP for BPs
- Skeleton on Skeletal Meshes
- Physics Asset on Skeletal Meshes
- Shadow Physics Asset on Skeletal Meshes
- Source Texture on Sprites

diff --git a/Engine/Plugins/Experimental/EditorDataStorageFeatures/Source/TedsAssetData/Private/Widgets/DynamicAssetDataColumnBaseWidget.cpp b/Engine/Plugins/Experimental/EditorDataStorageFeatures/Source/TedsAssetData/Private/Widgets/DynamicAssetDataColumnBaseWidget.cpp
new file mode 100644
--- /dev/null
+++ b/Engine/Plugins/Experimental/EditorDataStorageFeatures/Source/TedsAssetData/Private/Widgets/DynamicAssetDataColumnBaseWidget.cpp
@@ -0,0 +1,94 @@
+// Copyright Epic Games, Inc. All Rights Reserved.
+
+#include "DynamicAssetDataColumnBaseWidget.h"
+
+#include "AssetRegistry/AssetData.h"
+#include "AssetRegistry/AssetRegistryModule.h"
+#include "Elements/Framework/TypedElementAttributeBinding.h"
+#include "IContentBrowserSingleton.h"
+#include "Internationalization/Text.h"
+#include "SHyperlinkAssetPreviewWidget.h"
+#include "TedsAssetDataColumns.h"
+#include "UObject/ObjectMacros.h"
+#include "Widgets/SBoxPanel.h"
+
+namespace UE::DynamicColumn::Utilities
+{
+	FAssetData GetAssetData(const FString& InRawPath)
+	{
+		const FAssetRegistryModule& AssetRegistryModule = FModuleManager::LoadModuleChecked<FAssetRegistryModule>(TEXT("AssetRegistry"));
+		FSoftObjectPath ObjectPath = FSoftObjectPath(InRawPath);
+		FAssetData AssetData = AssetRegistryModule.Get().GetAssetByObjectPath(ObjectPath);
+		if (!AssetData.IsValid())
+		{
+			// If not valid try to get it again if it is a Generated BP
+			FString ObjectPathAsString = ObjectPath.ToString();
+			if (ObjectPathAsString.RemoveFromEnd(TEXT("_C")))
+			{
+				AssetData = AssetRegistryModule.Get().GetAssetByObjectPath(FSoftObjectPath(ObjectPathAsString));
+			}
+		}
+		return AssetData;
+	}	
+}
+
+void UDynamicAssetDataColumnBaseWidgetFactory::RegisterWidgetConstructors(IEditorDataStorageProvider& DataStorage,
+	IEditorDataStorageUiProvider& DataStorageUi) const
+{
+	using namespace UE::Editor::DataStorage::Queries;
+	DataStorageUi.RegisterWidgetFactory<FDynamicAssetDataColumnBaseWidgetConstructor>(TEXT("General.Cell"), TColumn<FItemStringAttributeColumn_Experimental>("ParentClass"));
+	DataStorageUi.RegisterWidgetFactory<FDynamicAssetDataColumnBaseWidgetConstructor>(TEXT("General.Cell"), TColumn<FItemStringAttributeColumn_Experimental>("Skeleton"));
+	DataStorageUi.RegisterWidgetFactory<FDynamicAssetDataColumnBaseWidgetConstructor>(TEXT("General.Cell"), TColumn<FItemStringAttributeColumn_Experimental>("SourceTexture"));
+	DataStorageUi.RegisterWidgetFactory<FDynamicAssetDataColumnBaseWidgetConstructor>(TEXT("General.Cell"), TColumn<FItemStringAttributeColumn_Experimental>("PhysicsAsset"));
+	DataStorageUi.RegisterWidgetFactory<FDynamicAssetDataColumnBaseWidgetConstructor>(TEXT("General.Cell"), TColumn<FItemStringAttributeColumn_Experimental>("ShadowPhysicsAsset"));
+
+}
+
+FDynamicAssetDataColumnBaseWidgetConstructor::FDynamicAssetDataColumnBaseWidgetConstructor()
+	: FSimpleWidgetConstructor(StaticStruct())
+{
+}
+
+TSharedPtr<SWidget> FDynamicAssetDataColumnBaseWidgetConstructor::CreateWidget(IEditorDataStorageProvider* DataStorage,
+	IEditorDataStorageUiProvider* DataStorageUi, UE::Editor::DataStorage::RowHandle TargetRow, UE::Editor::DataStorage::RowHandle WidgetRow,
+	const UE::Editor::DataStorage::FMetaDataView& Arguments)
+{
+	using namespace UE::Editor::DataStorage;
+	using namespace UE::Editor::DataStorage::Queries;
+
+	FAttributeBinder Binder(TargetRow, DataStorage);
+
+	FName* DynamicTemplates = MatchedDynamicTemplates.Find(FItemStringAttributeColumn_Experimental::StaticStruct());
+
+	if (!DynamicTemplates)
+	{
+		return SNullWidget::NullWidget;
+	}
+
+	const TAttribute<FAssetData>& AssetDataAttribute = Binder.BindData(*DynamicTemplates, &FItemStringAttributeColumn_Experimental::Value, [DataStorage] (const FString& InValue)
+		{
+			const FMapKey ReferencedAssetRowKey = FMapKey(FSoftObjectPath(InValue));
+			const RowHandle ReferencedRow = DataStorage->LookupMappedRow(ReferencedAssetRowKey);
+			if (ReferencedRow != InvalidRowHandle)
+			{
+				FAttributeBinder Binder(ReferencedRow, DataStorage);
+				return Binder.BindData(&FAssetDataColumn_Experimental::AssetData, FAssetData()).Get();
+			}
+			// Fallback to try to get the asset data from the Path
+			return UE::DynamicColumn::Utilities::GetAssetData(InValue);
+		});
+
+	const FOnNavigateAsset& OnNavigate = FOnNavigateAsset::CreateLambda([] (const FAssetData& InAssetData)
+		{
+			IContentBrowserSingleton::Get().SyncBrowserToAssets({ InAssetData });
+		});
+
+	return SNew(SHorizontalBox)
+		+SHorizontalBox::Slot()
+		.AutoWidth()
+		[
+			SNew(SHyperlinkAssetPreviewWidget)
+			.AssetData(AssetDataAttribute)
+			.OnNavigateAsset(OnNavigate)
+		];
+}
diff --git a/Engine/Plugins/Experimental/EditorDataStorageFeatures/Source/TedsAssetData/Private/Widgets/DynamicAssetDataColumnBaseWidget.h b/Engine/Plugins/Experimental/EditorDataStorageFeatures/Source/TedsAssetData/Private/Widgets/DynamicAssetDataColumnBaseWidget.h
new file mode 100644
--- /dev/null
+++ b/Engine/Plugins/Experimental/EditorDataStorageFeatures/Source/TedsAssetData/Private/Widgets/DynamicAssetDataColumnBaseWidget.h
@@ -0,0 +1,39 @@
+// Copyright Epic Games, Inc. All Rights Reserved.
+
+#pragma once
+
+#include "Elements/Interfaces/TypedElementDataStorageFactory.h"
+#include "Elements/Interfaces/TypedElementDataStorageUiInterface.h"
+
+#include "DynamicAssetDataColumnBaseWidget.generated.h"
+
+class IEditorDataStorageProvider;
+
+UCLASS()
+class UDynamicAssetDataColumnBaseWidgetFactory : public UEditorDataStorageFactory
+{
+	GENERATED_BODY()
+
+public:
+	~UDynamicAssetDataColumnBaseWidgetFactory() override = default;
+
+	TEDSASSETDATA_API void RegisterWidgetConstructors(IEditorDataStorageProvider& DataStorage,
+		IEditorDataStorageUiProvider& DataStorageUi) const override;
+};
+
+USTRUCT()
+struct FDynamicAssetDataColumnBaseWidgetConstructor : public FSimpleWidgetConstructor
+{
+	GENERATED_BODY()
+
+public:
+	TEDSASSETDATA_API FDynamicAssetDataColumnBaseWidgetConstructor();
+	~FDynamicAssetDataColumnBaseWidgetConstructor() override = default;
+
+	TEDSASSETDATA_API virtual TSharedPtr<SWidget> CreateWidget(
+		IEditorDataStorageProvider* DataStorage,
+		IEditorDataStorageUiProvider* DataStorageUi,
+		UE::Editor::DataStorage::RowHandle TargetRow,
+		UE::Editor::DataStorage::RowHandle WidgetRow,
+		const UE::Editor::DataStorage::FMetaDataView& Arguments) override;
+};
diff --git a/Engine/Plugins/Experimental/EditorDataStorageFeatures/Source/TedsAssetData/Private/Widgets/SHyperlinkAssetPreviewWidget.cpp b/Engine/Plugins/Experimental/EditorDataStorageFeatures/Source/TedsAssetData/Private/Widgets/SHyperlinkAssetPreviewWidget.cpp
new file mode 100644
--- /dev/null
+++ b/Engine/Plugins/Experimental/EditorDataStorageFeatures/Source/TedsAssetData/Private/Widgets/SHyperlinkAssetPreviewWidget.cpp
@@ -0,0 +1,136 @@
+// Copyright Epic Games, Inc. All Rights Reserved.
+
+#include "SHyperlinkAssetPreviewWidget.h"
+
+#include "AssetThumbnail.h"
+#include "IContentBrowserSingleton.h"
+#include "Widgets/Input/SHyperlink.h"
+
+#define LOCTEXT_NAMESPACE "SHyperlinkAssetPreviewWidget"
+
+/** HyperlinkAssetPreview tooltip */
+class SHyperlinkAssetToolTip : public SToolTip
+{
+public:
+	SLATE_BEGIN_ARGS(SHyperlinkAssetToolTip)
+		: _HyperlinkAssetPreview()
+	{ }
+
+		SLATE_ARGUMENT(TSharedPtr<SHyperlinkAssetPreviewWidget>, HyperlinkAssetPreview)
+
+	SLATE_END_ARGS()
+
+	void Construct(const FArguments& InArgs)
+	{
+		HyperlinkAssetPreview = InArgs._HyperlinkAssetPreview;
+
+		SToolTip::Construct(
+			SToolTip::FArguments()
+			.TextMargin(0.f)
+			.BorderImage(FAppStyle::GetBrush("AssetThumbnail.Tooltip.Border"))
+			);
+	}
+
+	virtual bool IsEmpty() const override
+	{
+		return !HyperlinkAssetPreview.IsValid();
+	}
+
+	virtual void OnOpening() override
+	{
+		TSharedPtr<SHyperlinkAssetPreviewWidget> AssetViewItemPin = HyperlinkAssetPreview.Pin();
+		if (AssetViewItemPin.IsValid())
+		{
+			SetContentWidget(AssetViewItemPin->GetThumbnailWidget());
+		}
+	}
+
+	virtual void OnClosed() override
+	{
+		ResetContentWidget();
+	}
+
+private:
+	TWeakPtr<SHyperlinkAssetPreviewWidget> HyperlinkAssetPreview;
+};
+
+void SHyperlinkAssetPreviewWidget::Construct(const FArguments& InArgs)
+{
+	AssetDataAttribute = InArgs._AssetData;
+	AssetThumbnailTooltip = MakeShared<FAssetThumbnail>(FAssetData(), 64, 64, UThumbnailManager::Get().GetSharedThumbnailPool());
+	OnNavigateAssetDelegate = InArgs._OnNavigateAsset;
+
+	// Set the Thumbnail tooltip
+	SetToolTip(SNew(SHyperlinkAssetToolTip).HyperlinkAssetPreview(SharedThis(this)));
+
+	SHyperlink::Construct(
+		SHyperlink::FArguments()
+		.Visibility(this, &SHyperlinkAssetPreviewWidget::GetHyperlinkVisibility)
+		.Text(this, &SHyperlinkAssetPreviewWidget::GetAssetDisplayName)
+		.OnNavigate(this, &SHyperlinkAssetPreviewWidget::OnNavigate_Internal)
+		.Style(FAppStyle::Get(), "Common.GotoBlueprintHyperlink"));
+}
+
+TSharedRef<SWidget> SHyperlinkAssetPreviewWidget::GetThumbnailWidget() const
+{
+	FAssetData AssetData = FAssetData();
+	if (AssetDataAttribute.IsSet())
+	{
+		AssetData = AssetDataAttribute.Get();
+	}
+
+	FAssetThumbnailConfig ThumbnailConfig;
+	ThumbnailConfig.bAllowFadeIn = true;
+	ThumbnailConfig.bAllowHintText = false;
+	ThumbnailConfig.bAllowRealTimeOnHovered = false;
+	ThumbnailConfig.bForceGenericThumbnail = !AssetData.IsValid();
+	ThumbnailConfig.bAllowAssetSpecificThumbnailOverlay = !ThumbnailConfig.bForceGenericThumbnail;
+	ThumbnailConfig.ThumbnailLabel = EThumbnailLabel::ClassName;
+	ThumbnailConfig.GenericThumbnailSize = 64.f;
+	ThumbnailConfig.bShowAssetChip = true;
+
+	if (AssetThumbnailTooltip->GetAssetData() != AssetData)
+	{
+		AssetThumbnailTooltip->SetAsset(AssetData);
+		AssetThumbnailTooltip->RefreshThumbnail();
+	}
+
+	return AssetThumbnailTooltip->MakeThumbnailWidget(ThumbnailConfig);
+}
+
+void SHyperlinkAssetPreviewWidget::OnNavigate_Internal() const
+{
+	if (AssetDataAttribute.IsSet())
+	{
+		FAssetData AssetData = AssetDataAttribute.Get();
+		if (AssetData.IsValid())
+		{
+			OnNavigateAssetDelegate.ExecuteIfBound(AssetData);
+		}
+	}
+}
+
+EVisibility SHyperlinkAssetPreviewWidget::GetHyperlinkVisibility() const
+{
+	if (AssetDataAttribute.IsSet())
+	{
+		FAssetData AssetData = AssetDataAttribute.Get();
+		return AssetData.IsValid() ? EVisibility::Visible : EVisibility::Collapsed;
+	}
+	return EVisibility::Collapsed;
+}
+
+FText SHyperlinkAssetPreviewWidget::GetAssetDisplayName() const
+{
+	if (AssetDataAttribute.IsSet())
+	{
+		const FAssetData& AssetData = AssetDataAttribute.Get();
+		if (AssetData.IsValid())
+		{
+			return FText::FromName(AssetData.AssetName);
+		}
+	}
+	return FText::GetEmpty();
+}
+
+#undef LOCTEXT_NAMESPACE
diff --git a/Engine/Plugins/Experimental/EditorDataStorageFeatures/Source/TedsAssetData/Private/Widgets/SHyperlinkAssetPreviewWidget.h b/Engine/Plugins/Experimental/EditorDataStorageFeatures/Source/TedsAssetData/Private/Widgets/SHyperlinkAssetPreviewWidget.h
new file mode 100644
--- /dev/null
+++ b/Engine/Plugins/Experimental/EditorDataStorageFeatures/Source/TedsAssetData/Private/Widgets/SHyperlinkAssetPreviewWidget.h
@@ -0,0 +1,51 @@
+// Copyright Epic Games, Inc. All Rights Reserved.
+
+#pragma once
+
+#include "AssetRegistry/AssetData.h"
+#include "Widgets/SCompoundWidget.h"
+#include "Widgets/Input/SHyperlink.h"
+
+class FAssetThumbnail;
+
+DECLARE_DELEGATE_OneParam(FOnNavigateAsset, const FAssetData&)
+
+/** Widget to display a Hyperlink with a Preview of the asset */
+class SHyperlinkAssetPreviewWidget : public SHyperlink
+{
+public:
+	SLATE_BEGIN_ARGS(SHyperlinkAssetPreviewWidget)
+	{}
+
+		SLATE_ATTRIBUTE(FAssetData, AssetData)
+
+		/** Will be called only if the AssetData is valid */
+		SLATE_EVENT(FOnNavigateAsset, OnNavigateAsset)
+
+	SLATE_END_ARGS()
+
+	void Construct(const FArguments& InArgs);
+
+	/** Get the preview tooltip for the hyperlink text */
+	TSharedRef<SWidget> GetThumbnailWidget() const;
+
+private:
+	/** Called when trying to navigate to the asset */
+	void OnNavigate_Internal() const;
+
+	/** Visible if the AssetData is valid, Collapsed otherwise */
+	EVisibility GetHyperlinkVisibility() const;
+
+	/** Retrieve the asset name to use */
+	FText GetAssetDisplayName() const;
+
+private:
+	/** AssetData to link */
+	TAttribute<FAssetData> AssetDataAttribute;
+
+	/** AssetThumbnail used for the tooltip */
+	TSharedPtr<FAssetThumbnail> AssetThumbnailTooltip;
+
+	/** Callback when trying to navigate to an asset */
+	FOnNavigateAsset OnNavigateAssetDelegate;
+};
diff --git a/Engine/Plugins/Experimental/EditorDataStorageFeatures/Source/TedsAssetData/Private/Widgets/SkeletalMeshSkeletonWidget.cpp b/Engine/Plugins/Experimental/EditorDataStorageFeatures/Source/TedsAssetData/Private/Widgets/SkeletalMeshSkeletonWidget.cpp
deleted file mode 100644
diff --git a/Engine/Plugins/Experimental/EditorDataStorageFeatures/Source/TedsAssetData/Private/Widgets/SkeletalMeshSkeletonWidget.h b/Engine/Plugins/Experimental/EditorDataStorageFeatures/Source/TedsAssetData/Private/Widgets/SkeletalMeshSkeletonWidget.h
deleted file mode 100644
diff --git a/Engine/Plugins/Experimental/EditorDataStorageFeatures/Source/TedsContentBrowser/Private/TedsContentBrowserModule.cpp b/Engine/Plugins/Experimental/EditorDataStorageFeatures/Source/TedsContentBrowser/Private/TedsContentBrowserModule.cpp
--- a/Engine/Plugins/Experimental/EditorDataStorageFeatures/Source/TedsContentBrowser/Private/TedsContentBrowserModule.cpp
+++ b/Engine/Plugins/Experimental/EditorDataStorageFeatures/Source/TedsContentBrowser/Private/TedsContentBrowserModule.cpp
@@ -192,7 +192,31 @@ namespace UE::Editor::ContentBrowser
 							.TemplateType = FItemStringAttributeColumn_Experimental::StaticStruct(),
 							.Identifier = "Skeleton"
 						});
-		
+
+		const UScriptStruct* DynamicBlueprintColumn = DataStorage->GenerateDynamicColumn(FDynamicColumnDescription
+						{
+							.TemplateType = FItemStringAttributeColumn_Experimental::StaticStruct(),
+							.Identifier = "ParentClass"
+						});
+
+		const UScriptStruct* DynamicSourceTextureColumn = DataStorage->GenerateDynamicColumn(FDynamicColumnDescription
+						{
+							.TemplateType = FItemStringAttributeColumn_Experimental::StaticStruct(),
+							.Identifier = "SourceTexture"
+						});
+
+		const UScriptStruct* DynamicPhysicsAssetColumn = DataStorage->GenerateDynamicColumn(FDynamicColumnDescription
+						{
+							.TemplateType = FItemStringAttributeColumn_Experimental::StaticStruct(),
+							.Identifier = "PhysicsAsset"
+						});
+
+		const UScriptStruct* DynamicShadowPhysicsAssetColumn = DataStorage->GenerateDynamicColumn(FDynamicColumnDescription
+						{
+							.TemplateType = FItemStringAttributeColumn_Experimental::StaticStruct(),
+							.Identifier = "ShadowPhysicsAsset"
+						});
+
 		// Create the table viewer widget
 		TableViewer = SNew(STedsTableViewer)
 					.QueryStack(RowQueryStack)
@@ -200,7 +224,9 @@ namespace UE::Editor::ContentBrowser
 					// Default list of columns to display
 					.Columns({ FNameColumn::StaticStruct(), FTedsAlertColumn::StaticStruct(),
 						FAssetClassColumn::StaticStruct(), FAssetTag::StaticStruct(), FFolderTag::StaticStruct(), FAssetPathColumn_Experimental::StaticStruct(),
-						FDiskSizeColumn::StaticStruct(), FVirtualPathColumn_Experimental::StaticStruct(), DynamicSkeletalMeshSkeletonColumn })
+						FDiskSizeColumn::StaticStruct(), FVirtualPathColumn_Experimental::StaticStruct(),
+						DynamicSkeletalMeshSkeletonColumn, DynamicBlueprintColumn, DynamicSourceTextureColumn,
+						DynamicPhysicsAssetColumn, DynamicShadowPhysicsAssetColumn })
 					.ListSelectionMode(ESelectionMode::Multi)
 					.OnSelectionChanged_Lambda([this](RowHandle Row)
 					{
