[TEDS] New Query Stack nodes to manage queries.

diff --git a/Engine/Plugins/Experimental/EditorDataStorageFeatures/Source/TedsQueryStack/Private/TedsQueryHandleNode.cpp b/Engine/Plugins/Experimental/EditorDataStorageFeatures/Source/TedsQueryStack/Private/TedsQueryHandleNode.cpp
new file mode 100644
--- /dev/null
+++ b/Engine/Plugins/Experimental/EditorDataStorageFeatures/Source/TedsQueryStack/Private/TedsQueryHandleNode.cpp
@@ -0,0 +1,32 @@
+// Copyright Epic Games, Inc. All Rights Reserved.
+
+#include "TedsQueryHandleNode.h"
+
+namespace UE::Editor::DataStorage::QueryStack
+{
+	FQueryHandleNode::FQueryHandleNode(QueryHandle Query)
+		: Query(Query)
+	{
+	}
+
+	void FQueryHandleNode::SetQuery(QueryHandle InQuery)
+	{
+		Query = InQuery;
+		Revision++;
+	}
+
+	INode::RevisionId FQueryHandleNode::GetRevision() const
+	{
+		return Revision;
+	}
+
+	void FQueryHandleNode::Update()
+	{
+		// Nothing to update.
+	}
+
+	QueryHandle FQueryHandleNode::GetQuery() const
+	{
+		return Query;
+	}
+} // namespace UE::Editor::DataStorage::QueryStack
diff --git a/Engine/Plugins/Experimental/EditorDataStorageFeatures/Source/TedsQueryStack/Private/TedsQueryNode.cpp b/Engine/Plugins/Experimental/EditorDataStorageFeatures/Source/TedsQueryStack/Private/TedsQueryNode.cpp
new file mode 100644
--- /dev/null
+++ b/Engine/Plugins/Experimental/EditorDataStorageFeatures/Source/TedsQueryStack/Private/TedsQueryNode.cpp
@@ -0,0 +1,56 @@
+// Copyright Epic Games, Inc. All Rights Reserved.
+
+#include "TedsQueryNode.h"
+
+#include "Elements/Interfaces/TypedElementDataStorageInterface.h"
+
+namespace UE::Editor::DataStorage::QueryStack
+{
+	FQueryNode::FQueryNode(IEditorDataStorageProvider& Storage)
+		: Storage(Storage)
+	{
+	}
+
+	FQueryNode::FQueryNode(IEditorDataStorageProvider& Storage, UE::Editor::DataStorage::FQueryDescription Query)
+		: QueryHandle(Storage.RegisterQuery(MoveTemp(Query)))
+		, Storage(Storage)
+	{
+	}
+
+	FQueryNode::~FQueryNode()
+	{
+		Storage.UnregisterQuery(QueryHandle);
+	}
+
+	void FQueryNode::SetQuery(UE::Editor::DataStorage::FQueryDescription Query)
+	{
+		Storage.UnregisterQuery(QueryHandle);
+		QueryHandle = Storage.RegisterQuery(MoveTemp(Query));
+		Revision++;
+	}
+
+	void FQueryNode::ClearQuery()
+	{
+		if (QueryHandle != InvalidQueryHandle)
+		{
+			Storage.UnregisterQuery(QueryHandle);
+			QueryHandle = InvalidQueryHandle;
+			Revision++;
+		}
+	}
+
+	INode::RevisionId FQueryNode::GetRevision() const
+	{
+		return Revision;
+	}
+
+	void FQueryNode::Update()
+	{
+		// Nothing to update.
+	}
+
+	QueryHandle FQueryNode::GetQuery() const
+	{
+		return QueryHandle;
+	}
+} // namespace UE::Editor::DataStorage::QueryStack
diff --git a/Engine/Plugins/Experimental/EditorDataStorageFeatures/Source/TedsQueryStack/Public/TedsQueryHandleNode.h b/Engine/Plugins/Experimental/EditorDataStorageFeatures/Source/TedsQueryStack/Public/TedsQueryHandleNode.h
new file mode 100644
--- /dev/null
+++ b/Engine/Plugins/Experimental/EditorDataStorageFeatures/Source/TedsQueryStack/Public/TedsQueryHandleNode.h
@@ -0,0 +1,29 @@
+// Copyright Epic Games, Inc. All Rights Reserved.
+
+#pragma once
+
+#include "TedsQueryStackInterfaces.h"
+
+namespace UE::Editor::DataStorage::QueryStack
+{
+	/**
+	 * Stores a handle to a query. This should be used to pass an externally created query into the Query Stack.
+	 */
+	class FQueryHandleNode : public IQueryNode
+	{
+	public:
+		TEDSQUERYSTACK_API FQueryHandleNode() = default;
+		TEDSQUERYSTACK_API explicit FQueryHandleNode(QueryHandle Query);
+		virtual ~FQueryHandleNode() override = default;
+
+		TEDSQUERYSTACK_API void SetQuery(QueryHandle InQuery);
+
+		TEDSQUERYSTACK_API virtual RevisionId GetRevision() const override;
+		TEDSQUERYSTACK_API virtual void Update() override;
+		TEDSQUERYSTACK_API virtual QueryHandle GetQuery() const override;
+
+	private:
+		QueryHandle Query = InvalidQueryHandle;
+		RevisionId Revision = 0;
+	};
+} // namespace UE::Editor::DataStorage::QueryStack
diff --git a/Engine/Plugins/Experimental/EditorDataStorageFeatures/Source/TedsQueryStack/Public/TedsQueryNode.h b/Engine/Plugins/Experimental/EditorDataStorageFeatures/Source/TedsQueryStack/Public/TedsQueryNode.h
new file mode 100644
--- /dev/null
+++ b/Engine/Plugins/Experimental/EditorDataStorageFeatures/Source/TedsQueryStack/Public/TedsQueryNode.h
@@ -0,0 +1,34 @@
+// Copyright Epic Games, Inc. All Rights Reserved.
+
+#pragma once
+
+#include "Elements/Common/TypedElementQueryDescription.h"
+#include "TedsQueryStackInterfaces.h"
+
+class IEditorDataStorageProvider;
+
+namespace UE::Editor::DataStorage::QueryStack
+{
+	/**
+	 * Stores a query and manages its registration.
+	 */
+	class FQueryNode : public IQueryNode
+	{
+	public:
+		TEDSQUERYSTACK_API explicit FQueryNode(IEditorDataStorageProvider& Storage);
+		TEDSQUERYSTACK_API FQueryNode(IEditorDataStorageProvider& Storage, UE::Editor::DataStorage::FQueryDescription Query);
+		TEDSQUERYSTACK_API virtual ~FQueryNode() override;
+
+		TEDSQUERYSTACK_API void SetQuery(UE::Editor::DataStorage::FQueryDescription Query);
+		TEDSQUERYSTACK_API void ClearQuery();
+
+		TEDSQUERYSTACK_API virtual RevisionId GetRevision() const override;
+		TEDSQUERYSTACK_API virtual void Update() override;
+		TEDSQUERYSTACK_API virtual QueryHandle GetQuery() const override;
+
+	private:
+		QueryHandle QueryHandle = InvalidQueryHandle;
+		IEditorDataStorageProvider& Storage;
+		RevisionId Revision = 0;
+	};
+} // namespace UE::Editor::DataStorage::QueryStack
