[TEDS] Alerts could attempt to check a parent that has become invalid, causing a crash deep in Mass due to an invalid archetype.

diff --git a/Engine/Plugins/Experimental/EditorDataStorageFeatures/Source/TedsAlerts/Private/TedsAlertsFactory.cpp b/Engine/Plugins/Experimental/EditorDataStorageFeatures/Source/TedsAlerts/Private/TedsAlertsFactory.cpp
--- a/Engine/Plugins/Experimental/EditorDataStorageFeatures/Source/TedsAlerts/Private/TedsAlertsFactory.cpp
+++ b/Engine/Plugins/Experimental/EditorDataStorageFeatures/Source/TedsAlerts/Private/TedsAlertsFactory.cpp
@@ -278,7 +278,10 @@ void UTedsAlertsFactory::RegisterChildAlertUpdatesQueries(IEditorDataStorageProv
 				.MakeActivatable(AlertConditionName),
 			[](IQueryContext& Context, RowHandle Row, FTedsAlertColumn& Alert, const FTableRowParentColumn& Parent)
 			{
-				AddChildAlertsToHierarchy(Context, Parent.Parent, 0);
+				if (Context.IsRowAssigned(Parent.Parent))
+				{
+					AddChildAlertsToHierarchy(Context, Parent.Parent, 0);
+				}
 			})
 		.Where()
 			.None<FTedsAlertChainTag>()
