[TEDS] New Query Stack node to manage a row handle array directly as well as a new interface call so nodes can share row arrays.

diff --git a/Engine/Plugins/Experimental/EditorDataStorageFeatures/Source/TedsQueryStack/Private/TedsRowArrayNode.cpp b/Engine/Plugins/Experimental/EditorDataStorageFeatures/Source/TedsQueryStack/Private/TedsRowArrayNode.cpp
new file mode 100644
--- /dev/null
+++ b/Engine/Plugins/Experimental/EditorDataStorageFeatures/Source/TedsQueryStack/Private/TedsRowArrayNode.cpp
@@ -0,0 +1,35 @@
+// Copyright Epic Games, Inc. All Rights Reserved.
+
+#include "TedsRowArrayNode.h"
+
+namespace UE::Editor::DataStorage::QueryStack
+{
+	FRowArrayNode::FRowArrayNode(FRowHandleArray Rows)
+		: Rows(MoveTemp(Rows))
+	{
+	}
+
+	void FRowArrayNode::MarkDirty()
+	{
+		Revision++;
+	}
+
+	INode::RevisionId FRowArrayNode::GetRevision() const
+	{
+		return Revision;
+	}
+
+	void FRowArrayNode::Update()
+	{
+	}
+
+	FRowHandleArrayView FRowArrayNode::GetRows() const
+	{
+		return Rows.GetRows();
+	}
+
+	FRowHandleArray* FRowArrayNode::GetMutableRows()
+	{
+		return &Rows;
+	}
+} // namespace UE::Editor::DataStorage::QueryStack
diff --git a/Engine/Plugins/Experimental/EditorDataStorageFeatures/Source/TedsQueryStack/Private/TedsRowMergeNode.cpp b/Engine/Plugins/Experimental/EditorDataStorageFeatures/Source/TedsQueryStack/Private/TedsRowMergeNode.cpp
--- a/Engine/Plugins/Experimental/EditorDataStorageFeatures/Source/TedsQueryStack/Private/TedsRowMergeNode.cpp
+++ b/Engine/Plugins/Experimental/EditorDataStorageFeatures/Source/TedsQueryStack/Private/TedsRowMergeNode.cpp
@@ -47,6 +47,11 @@ namespace UE::Editor::DataStorage::QueryStack
 		return Rows.GetRows();
 	}
 
+	FRowHandleArray* FRowMergeNode::GetMutableRows()
+	{
+		return &Rows;
+	}
+
 	void FRowMergeNode::Merge()
 	{
 		Rows.Reset();
diff --git a/Engine/Plugins/Experimental/EditorDataStorageFeatures/Source/TedsQueryStack/Private/TedsRowViewNode.cpp b/Engine/Plugins/Experimental/EditorDataStorageFeatures/Source/TedsQueryStack/Private/TedsRowViewNode.cpp
--- a/Engine/Plugins/Experimental/EditorDataStorageFeatures/Source/TedsQueryStack/Private/TedsRowViewNode.cpp
+++ b/Engine/Plugins/Experimental/EditorDataStorageFeatures/Source/TedsQueryStack/Private/TedsRowViewNode.cpp
@@ -32,4 +32,9 @@ namespace UE::Editor::DataStorage::QueryStack
 	{
 		return Rows;
 	}
+
+	FRowHandleArray* FRowViewNode::GetMutableRows()
+	{
+		return nullptr;
+	}
 } // namespace UE::Editor::DataStorage::QueryStack
diff --git a/Engine/Plugins/Experimental/EditorDataStorageFeatures/Source/TedsQueryStack/Public/TedsQueryStackInterfaces.h b/Engine/Plugins/Experimental/EditorDataStorageFeatures/Source/TedsQueryStack/Public/TedsQueryStackInterfaces.h
--- a/Engine/Plugins/Experimental/EditorDataStorageFeatures/Source/TedsQueryStack/Public/TedsQueryStackInterfaces.h
+++ b/Engine/Plugins/Experimental/EditorDataStorageFeatures/Source/TedsQueryStack/Public/TedsQueryStackInterfaces.h
@@ -4,6 +4,7 @@
 
 #include "Containers/ContainersFwd.h"
 #include "Elements/Common/TypedElementHandles.h"
+#include "Elements/Framework/TypedElementRowHandleArray.h"
 #include "Elements/Framework/TypedElementRowHandleArrayView.h"
 
 namespace UE::Editor::DataStorage::QueryStack
@@ -25,6 +26,7 @@ namespace UE::Editor::DataStorage::QueryStack
 	{
 	public:
 		virtual ~IQueryNode() = default;
+		/** Returns the handle to the query this node represents. */
 		virtual QueryHandle GetQuery() const = 0;
 	};
 
@@ -33,6 +35,9 @@ namespace UE::Editor::DataStorage::QueryStack
 	{
 	public:
 		virtual ~IRowNode() = default;
+		/** Retrieve access to the rows used by this node. */
 		virtual FRowHandleArrayView GetRows() const = 0;
+		/** Retrieve write access to the rows used by this node. Is allowed to return null if write access can't be granted. */
+		virtual FRowHandleArray* GetMutableRows() = 0;
 	};
 } // UE::Editor::DataStorage::QueryStack
diff --git a/Engine/Plugins/Experimental/EditorDataStorageFeatures/Source/TedsQueryStack/Public/TedsRowArrayNode.h b/Engine/Plugins/Experimental/EditorDataStorageFeatures/Source/TedsQueryStack/Public/TedsRowArrayNode.h
new file mode 100644
--- /dev/null
+++ b/Engine/Plugins/Experimental/EditorDataStorageFeatures/Source/TedsQueryStack/Public/TedsRowArrayNode.h
@@ -0,0 +1,34 @@
+// Copyright Epic Games, Inc. All Rights Reserved.
+
+#pragma once
+
+#include "Containers/ContainersFwd.h"
+#include "TedsQueryStackInterfaces.h"
+#include "Elements/Framework//TypedElementRowHandleArray.h"
+
+namespace UE::Editor::DataStorage::QueryStack
+{
+	/**
+	 * Stores row handle array. The array can be directly manipulated, but any changes
+	 * outside the Query Stack require MarkDirty to inform the Query Stack that it needs
+	 * to update dependent nodes.
+	 */
+	class FRowArrayNode : public IRowNode
+	{
+	public:
+		FRowArrayNode() = default;
+		TEDSQUERYSTACK_API explicit FRowArrayNode(FRowHandleArray Rows);
+		virtual ~FRowArrayNode() override = default;
+
+		TEDSQUERYSTACK_API void MarkDirty();
+		
+		TEDSQUERYSTACK_API virtual RevisionId GetRevision() const override;
+		TEDSQUERYSTACK_API virtual void Update() override;
+		TEDSQUERYSTACK_API virtual FRowHandleArrayView GetRows() const override;
+		TEDSQUERYSTACK_API virtual FRowHandleArray* GetMutableRows() override;
+
+	private:
+		FRowHandleArray Rows;
+		RevisionId Revision = 0;
+	};
+} // namespace UE::Editor::DataStorage::QueryStack
diff --git a/Engine/Plugins/Experimental/EditorDataStorageFeatures/Source/TedsQueryStack/Public/TedsRowMergeNode.h b/Engine/Plugins/Experimental/EditorDataStorageFeatures/Source/TedsQueryStack/Public/TedsRowMergeNode.h
--- a/Engine/Plugins/Experimental/EditorDataStorageFeatures/Source/TedsQueryStack/Public/TedsRowMergeNode.h
+++ b/Engine/Plugins/Experimental/EditorDataStorageFeatures/Source/TedsQueryStack/Public/TedsRowMergeNode.h
@@ -26,6 +26,7 @@ namespace UE::Editor::DataStorage::QueryStack
 		TEDSQUERYSTACK_API virtual RevisionId GetRevision() const override;
 		TEDSQUERYSTACK_API virtual void Update() override;
 		TEDSQUERYSTACK_API virtual FRowHandleArrayView GetRows() const override;
+		TEDSQUERYSTACK_API virtual FRowHandleArray* GetMutableRows() override;
 
 	private:
 		void Merge();
diff --git a/Engine/Plugins/Experimental/EditorDataStorageFeatures/Source/TedsQueryStack/Public/TedsRowViewNode.h b/Engine/Plugins/Experimental/EditorDataStorageFeatures/Source/TedsQueryStack/Public/TedsRowViewNode.h
--- a/Engine/Plugins/Experimental/EditorDataStorageFeatures/Source/TedsQueryStack/Public/TedsRowViewNode.h
+++ b/Engine/Plugins/Experimental/EditorDataStorageFeatures/Source/TedsQueryStack/Public/TedsRowViewNode.h
@@ -14,7 +14,7 @@ namespace UE::Editor::DataStorage::QueryStack
 	class FRowViewNode : public IRowNode
 	{
 	public:
-		TEDSQUERYSTACK_API FRowViewNode() = default;
+		FRowViewNode() = default;
 		TEDSQUERYSTACK_API explicit FRowViewNode(FRowHandleArrayView Rows);
 		virtual ~FRowViewNode() override = default;
 
@@ -24,6 +24,7 @@ namespace UE::Editor::DataStorage::QueryStack
 		TEDSQUERYSTACK_API virtual RevisionId GetRevision() const override;
 		TEDSQUERYSTACK_API virtual void Update() override;
 		TEDSQUERYSTACK_API virtual FRowHandleArrayView GetRows() const override;
+		TEDSQUERYSTACK_API virtual FRowHandleArray* GetMutableRows() override;
 
 	private:
 		FRowHandleArrayView Rows;
