TEDS-Outliner: Fix the parent change observer being hardcoded to FTableRowParent column instead of using the hierarchy interface
Also fix some const correctness

diff --git a/Engine/Plugins/ChaosVD/Source/ChaosVD/Private/ChaosVDWorldOutlinerTab.cpp b/Engine/Plugins/ChaosVD/Source/ChaosVD/Private/ChaosVDWorldOutlinerTab.cpp
--- a/Engine/Plugins/ChaosVD/Source/ChaosVD/Private/ChaosVDWorldOutlinerTab.cpp
+++ b/Engine/Plugins/ChaosVD/Source/ChaosVD/Private/ChaosVDWorldOutlinerTab.cpp
@@ -56,9 +56,9 @@ void FChaosVDWorldOutlinerTab::CreateWorldOutlinerWidget()
 	Params.bUseDefaultTedsFilters = false;
 	Params.bShowRowHandleColumn = false;
 
-	const FTedsOutlinerHierarchyData::FGetParentRowHandle RowHandleGetter = FTedsOutlinerHierarchyData::FGetParentRowHandle::CreateLambda([](void* InColumnData)
+	const FTedsOutlinerHierarchyData::FGetParentRowHandle RowHandleGetter = FTedsOutlinerHierarchyData::FGetParentRowHandle::CreateLambda([](const void* InColumnData)
 	{
-		if(const FTableRowParentColumn* ParentColumn = static_cast<FTableRowParentColumn *>(InColumnData))
+		if(const FTableRowParentColumn* ParentColumn = static_cast<const FTableRowParentColumn *>(InColumnData))
 		{
 			return ParentColumn->Parent;
 		}
diff --git a/Engine/Plugins/Experimental/EditorDataStorageFeatures/Source/TedsOutliner/Private/TedsOutlinerImpl.cpp b/Engine/Plugins/Experimental/EditorDataStorageFeatures/Source/TedsOutliner/Private/TedsOutlinerImpl.cpp
--- a/Engine/Plugins/Experimental/EditorDataStorageFeatures/Source/TedsOutliner/Private/TedsOutlinerImpl.cpp
+++ b/Engine/Plugins/Experimental/EditorDataStorageFeatures/Source/TedsOutliner/Private/TedsOutlinerImpl.cpp
@@ -433,7 +433,7 @@ void FTedsOutlinerImpl::CreateChildren(const FSceneOutlinerTreeItemPtr& Item, TA
 		{
 			for(RowHandle ChildEntityRowHandle : MatchedRowsWithParentColumn)
 			{
-				void* ParentColumnData = DataStorage->GetColumnData(ChildEntityRowHandle, InHierarchyData.GetValue().HierarchyColumn);
+				const void* ParentColumnData = DataStorage->GetColumnData(ChildEntityRowHandle, InHierarchyData.GetValue().HierarchyColumn);
 
 				if (ensureMsgf(ParentColumnData, TEXT("We should always the a parent column since we only grabbed rows with those ")))
 				{
@@ -479,7 +479,7 @@ DataStorage::RowHandle FTedsOutlinerImpl::GetParentRow(DataStorage::RowHandle In
 	}
 	
 	// If this entity does not have a parent entity, return InvalidRowHandle
-	void* ParentColumnData = Storage->GetColumnData(InRowHandle, HierarchyData.GetValue().HierarchyColumn);
+	const void* ParentColumnData = Storage->GetColumnData(InRowHandle, HierarchyData.GetValue().HierarchyColumn);
 	
 	if (!ParentColumnData)
 	{
@@ -591,39 +591,37 @@ void FTedsOutlinerImpl::RecompileQueries()
 		FQueryDescription UpdateParentQueryDescription =
 			Select(
 			TEXT("Update item parent"),
-			FProcessor(EQueryTickPhase::DuringPhysics, Storage->GetQueryTickGroupName(EQueryTickGroups::Update))
+			FProcessor(EQueryTickPhase::PostPhysics, Storage->GetQueryTickGroupName(EQueryTickGroups::SyncDataStorageToExternal))
 				.SetExecutionMode(EExecutionMode::GameThread),
-			[this](IQueryContext& Context, const DataStorage::RowHandle* Rows, const FTedsOutlinerColumn* TedsOutlinerColumns)
+			[this, ParentColumnType](IQueryContext& Context, const RowHandle Row)
 			{
-				const FTableRowParentColumn* ParentColumnBegin = Context.GetColumn<FTableRowParentColumn>();
+				const void* ParentColumn = Context.GetColumn(ParentColumnType);
+				const FTedsOutlinerColumn* TedsOutlinerColumn = Context.GetColumn<FTedsOutlinerColumn>();
 				
-				for (int32 RowIndex = 0; RowIndex < (int32)Context.GetRowCount(); ++RowIndex)
+				RowHandle ParentRowHandle = InvalidRowHandle;
+				
+				if (ParentColumn)
 				{
-					RowHandle Row = Rows[RowIndex];
-					DataStorage::RowHandle ParentRowHandle = InvalidRowHandle;
-					if (ParentColumnBegin)
-					{
-						ParentRowHandle = ParentColumnBegin[RowIndex].Parent;
-					}
-									
-					if (QueryUtils::HasItemParentChanged(Context, Row, ParentRowHandle, *SceneOutliner))
+					ParentRowHandle = HierarchyData.GetValue().GetParent.Execute(ParentColumn);
+				}
+								
+				if (QueryUtils::HasItemParentChanged(Context, Row, ParentRowHandle, *SceneOutliner))
+				{
+					if (TedsOutlinerColumn && !QueryUtils::CanDisplayRow(Context, *TedsOutlinerColumn, Row, *SceneOutliner))
 					{
-						const FTedsOutlinerColumn& TedsOutlinerColumn = TedsOutlinerColumns[RowIndex];
-						if (!QueryUtils::CanDisplayRow(Context, TedsOutlinerColumn, Row, *SceneOutliner))
-						{
-							return;
-						}
-						
-						FSceneOutlinerHierarchyChangedData EventData;
-						EventData.Type = FSceneOutlinerHierarchyChangedData::Moved;
-						EventData.ItemIDs.Add(Row);
-						HierarchyChangedEvent.Broadcast(EventData);
+						return;
 					}
+					
+					FSceneOutlinerHierarchyChangedData EventData;
+					EventData.Type = FSceneOutlinerHierarchyChangedData::Moved;
+					EventData.ItemIDs.Add(Row);
+					HierarchyChangedEvent.Broadcast(EventData);
 				}
 			})
-			.ReadOnly<FTableRowParentColumn>(EOptional::Yes)
+			.ReadOnly(ParentColumnType, EOptional::Yes)
+			.ReadOnly<FTedsOutlinerColumn>(EOptional::Yes)
 		.Where()
-			.All<FTypedElementSyncFromWorldTag>()
+			.Any<FTypedElementSyncFromWorldTag, FTypedElementSyncBackToWorldTag>()
 		.Compile();
 		
 		// Add the conditions from FinalQueryDescription to ensure we are the rows the user requested
diff --git a/Engine/Plugins/Experimental/EditorDataStorageFeatures/Source/TedsOutliner/Public/TedsOutlinerImpl.h b/Engine/Plugins/Experimental/EditorDataStorageFeatures/Source/TedsOutliner/Public/TedsOutlinerImpl.h
--- a/Engine/Plugins/Experimental/EditorDataStorageFeatures/Source/TedsOutliner/Public/TedsOutlinerImpl.h
+++ b/Engine/Plugins/Experimental/EditorDataStorageFeatures/Source/TedsOutliner/Public/TedsOutlinerImpl.h
@@ -18,7 +18,7 @@ namespace UE::Editor::Outliner
 struct FTedsOutlinerHierarchyData
 {
 	/** A delegate used to get the parent row handle for a given row */
-	DECLARE_DELEGATE_RetVal_OneParam(DataStorage::RowHandle, FGetParentRowHandle, void* /* InColumnData */);
+	DECLARE_DELEGATE_RetVal_OneParam(DataStorage::RowHandle, FGetParentRowHandle, const void* /* InColumnData */);
 
 	/** A delegate used to get the children row handles for a given row */
 	DECLARE_DELEGATE_RetVal_OneParam(TArrayView<DataStorage::RowHandle>, FGetChildrenRowsHandles, void* /* InColumnData */);
@@ -50,9 +50,9 @@ struct FTedsOutlinerHierarchyData
 	// Get the default hierarchy data for the TEDS Outliner that uses FTableRowParentColumn to get the parent
 	static FTedsOutlinerHierarchyData GetDefaultHierarchyData()
 	{
-		const FGetParentRowHandle RowHandleGetter = FGetParentRowHandle::CreateLambda([](void* InColumnData)
+		const FGetParentRowHandle RowHandleGetter = FGetParentRowHandle::CreateLambda([](const void* InColumnData)
 			{
-				if(const FTableRowParentColumn* ParentColumn = static_cast<FTableRowParentColumn *>(InColumnData))
+				if(const FTableRowParentColumn* ParentColumn = static_cast<const FTableRowParentColumn *>(InColumnData))
 				{
 					return ParentColumn->Parent;
 				}
