TEDS Styling:
Add support for a few other style types (fonts/widgets etc) through a cvar (TEDS.Feature.AddSlateStyleRows)
Add a new SlateGallery widget that displays a sample of the styles available across all known stylesets (TEDS.UI.OpenSlateGallery)

SlateStyle:
Keep track of known widget stye types and add a function to retrieve them
Add a function to retrieve a widget style by type name instead of template

TedsTileViewer:
Add init params for item width, item height and style

diff --git a/Engine/Plugins/Experimental/EditorDataStorageFeatures/Source/TedsStyling/Private/TedsStylingFactory.cpp b/Engine/Plugins/Experimental/EditorDataStorageFeatures/Source/TedsStyling/Private/TedsStylingFactory.cpp
--- a/Engine/Plugins/Experimental/EditorDataStorageFeatures/Source/TedsStyling/Private/TedsStylingFactory.cpp
+++ b/Engine/Plugins/Experimental/EditorDataStorageFeatures/Source/TedsStyling/Private/TedsStylingFactory.cpp
@@ -7,11 +7,15 @@
 #include "Elements/Columns/TypedElementSlateWidgetColumns.h"
 #include "Elements/Common/EditorDataStorageFeatures.h"
 #include "Elements/Interfaces/TypedElementDataStorageInterface.h"
+#include "Elements/Interfaces/TypedElementDataStorageUiInterface.h"
+#include "Fonts/SlateFontInfo.h"
 #include "HAL/IConsoleManager.h"
 #include "Styling/SlateBrush.h"
 #include "Styling/SlateStyle.h"
 #include "Styling/SlateStyleRegistry.h"
 
+#define LOCTEXT_NAMESPACE "TedsStylingFactory"
+
 namespace UE::TedsStylingFactory::Local
 {
 	static const FName TableName("Editor_StyleTable");
@@ -27,12 +31,19 @@ namespace UE::TedsStylingFactory::Local
 
 void UTedsStylingFactory::RegisterTables(IEditorDataStorageProvider& DataStorage)
 {
-	DataStorage.RegisterTable(TTypedElementColumnTypeList<FNameColumn, FSlateStyleSetColumn, FSlateStyleTag>(), UE::TedsStylingFactory::Local::TableName);
+	DataStorage.RegisterTable(TTypedElementColumnTypeList<FNameColumn>(), UE::TedsStylingFactory::Local::TableName);
 }
 
-void UTedsStylingFactory::RegisterQueries(IEditorDataStorageProvider& DataStorage)
+void UTedsStylingFactory::RegisterWidgetPurposes(IEditorDataStorageUiProvider& DataStorageUi) const
 {
-	Super::RegisterQueries(DataStorage);
+	DataStorageUi.RegisterWidgetPurpose(FName(TEXT("SlateGallery.Color")), IEditorDataStorageUiProvider::EPurposeType::UniqueByName,
+	LOCTEXT("SlateGalleryColorPurpose", "Specific purpose display a color in the slate gallery."));
+	
+	DataStorageUi.RegisterWidgetPurpose(FName(TEXT("SlateGallery.Icons")), IEditorDataStorageUiProvider::EPurposeType::UniqueByName,
+	LOCTEXT("SlateGalleryIconsPurpose", "Specific purpose display Icons in the slate gallery."));
+	
+	DataStorageUi.RegisterWidgetPurpose(FName(TEXT("SlateGallery.Font")), IEditorDataStorageUiProvider::EPurposeType::UniqueByName,
+	LOCTEXT("SlateGalleryFontsPurpose", "Specific purpose display Fonts in the slate gallery."));
 }
 
 void UTedsStylingFactory::RegisterAllKnownStyles()
@@ -44,34 +55,80 @@ void UTedsStylingFactory::RegisterAllKnownStyles()
 	{
 		// Get all styles belonging to this styleset
 		TSet<FName> StyleKeys = Style.GetStyleKeys();
+		const TSet<FName>& WidgetStyleNames = Style.GetWidgetStyleNames();
+
+		static auto IsValidBrush = [](const FSlateBrush* Brush, const ISlateStyle& Style)
+		{
+			if (Brush == Style.GetDefaultBrush() || Brush->GetResourceName() == Style.GetDefaultBrush()->GetResourceName())
+			{
+				return false;
+			}
+
+			if (Brush->GetResourceName().IsNone())
+			{
+				return false;
+			}
+
+			return true;
+		};
 
 		for (const FName& StyleKey : StyleKeys)
 		{
 			// Since there is no way to check which type of style a specific key is, we just go through each one by one and check for now
 			// This is not the most performant and can be improved in the future by exposing more internals from ISlateStyle if needed
 			// We also only support brushes and colors currently
-			bool bRegisteredSyleKey = false;
 
 			// Check if this style key is a brush
 			if (const FSlateBrush* Brush = Style.GetBrush(StyleKey))
 			{
-				if(Brush != Style.GetDefaultBrush())
+				if (IsValidBrush(Brush, Style))
 				{
 					RegisterBrush(Interface, StyleKey, Brush, Style);
-					bRegisteredSyleKey = true;
+					continue;
+				}
+			}
+
+			// Check if this style key is a color
+			FSlateColor Color = Style.GetSlateColor(StyleKey);
+				
+			if (Color != FStyleDefaults::GetSlateColor())
+			{
+				RegisterColor(Interface, StyleKey, Color, Style);
+				continue;
+			}
+
+			FSlateFontInfo Font = Style.GetFontStyle(StyleKey);
+
+			if (Font != FStyleDefaults::GetFontInfo())
+			{
+				RegisterFont(Interface, StyleKey, Font, Style);
+				continue;
+			}
+
+			// Ignore Float, Vector, Margin and LinearColor for now but still early out if we detect them to skip expensive widget style code
+			{
+				float FloatValue = Style.GetFloat(StyleKey);
+				FVector2D VectorValue = Style.GetVector(StyleKey);
+				FMargin MarginValue = Style.GetMargin(StyleKey);
+				FLinearColor ColorValue = Style.GetColor(StyleKey);
+				
+				if (FloatValue != FStyleDefaults::GetFloat() || VectorValue != FStyleDefaults::GetVector2D() ||
+					MarginValue != FStyleDefaults::GetMargin() || ColorValue != FStyleDefaults::GetColor())
+				{
+					continue;
 				}
 			}
 
-			if(!bRegisteredSyleKey)
+			// Check if this key is any of the known widget styles
+			for (const FName& WidgetStyleName : WidgetStyleNames)
 			{
-				// Check if this style key is a color
-				FSlateColor Color = Style.GetSlateColor(StyleKey);
-				if (Color != FStyleDefaults::GetSlateColor())
+				if (const FSlateWidgetStyle* WidgetStyle =
+						Style.GetWidgetStyle(WidgetStyleName, StyleKey, nullptr, nullptr, false))
 				{
-					RegisterColor(Interface, StyleKey, Color, Style);
-					bRegisteredSyleKey = true;
+					RegisterWidgetStyle(Interface, WidgetStyleName, WidgetStyle, StyleKey, Style);
 				}
 			}
+			
 		}
 	};
 
@@ -100,7 +157,17 @@ void UTedsStylingFactory::RegisterBrush(IEditorDataStorageProvider* DataStorage,
 			DataStorage->AddColumn(Row, FSlateStylePathColumn{.StylePath = ResourceName});
 		}
 
-		DataStorage->AddColumn<FSlateBrushTag>(Row);
+		if (Brush->GetDrawType() == ESlateBrushDrawType::Image)
+		{
+			DataStorage->AddColumn<FSlateStyleTag>(Row, "Icon");
+		}
+
+		if (!Brush->GetImageSize().IsZero())
+		{
+			DataStorage->AddColumn(Row, FIconSizeColumn{.IconSize = Brush->GetImageSize()});
+		}
+
+		DataStorage->AddColumn<FSlateStyleTag>(Row, "Brush");
 	}
 }
 
@@ -108,16 +175,52 @@ void UTedsStylingFactory::RegisterColor(IEditorDataStorageProvider* DataStorage,
 	const ISlateStyle& OwnerStyle)
 {
 	const UE::Editor::DataStorage::RowHandle Row = AddOrGetStyleRow(DataStorage, StyleName, OwnerStyle);
+
+	// Check if the StyleName starts with "Color" to see if it is a Default color
+	const FString DisplayName = FName::NameToDisplayString(StyleName.ToString().Replace(TEXT("."), TEXT("")), false);
 	
 	if (DataStorage->IsRowAssigned(Row))
 	{
 		DataStorage->AddColumn(Row, FSlateColorColumn{.Color = Color});
+		DataStorage->AddColumn<FSlateStyleTag>(Row, "Color");
+		DataStorage->AddColumn(Row, FDisplayNameColumn
+			{ 
+				.DisplayName = DisplayName.StartsWith("Colors") ? FText::FromString(DisplayName.RightChop(7)) : FText::FromString(DisplayName)
+			});
+		
+		if (DisplayName.StartsWith("Colors"))
+		{
+			DataStorage->AddColumn<FStyleColorTag>(Row);
+		}
 	}
+}
 
+void UTedsStylingFactory::RegisterFont(IEditorDataStorageProvider* DataStorage, const FName& StyleName, const FSlateFontInfo& Font,
+	const ISlateStyle& OwnerStyle)
+{
+	const UE::Editor::DataStorage::RowHandle Row = AddOrGetStyleRow(DataStorage, StyleName, OwnerStyle);
+	
+	if (DataStorage->IsRowAssigned(Row))
+	{
+		DataStorage->AddColumn(Row, FSlateFontColumn{.FontInfo = Font});
+		DataStorage->AddColumn<FSlateStyleTag>(Row, "Font");
+	}
+}
+
+void UTedsStylingFactory::RegisterWidgetStyle(IEditorDataStorageProvider* DataStorage, const FName& WidgetStyleName,
+	const FSlateWidgetStyle* WidgetStyle, const FName& StyleName, const ISlateStyle& OwnerStyle)
+{
+	const UE::Editor::DataStorage::RowHandle Row = AddOrGetStyleRow(DataStorage, StyleName, OwnerStyle);
+	
+	if (DataStorage->IsRowAssigned(Row))
+	{
+		DataStorage->AddColumn(Row, WidgetStyleName, FSlateWidgetStyleColumn{.WidgetStyle = WidgetStyle});
+		DataStorage->AddColumn<FSlateStyleTag>(Row, "WidgetStyle");
+	}
 }
 
 UE::Editor::DataStorage::RowHandle UTedsStylingFactory::AddOrGetStyleRow(IEditorDataStorageProvider* DataStorage,
-	const FName& StyleName, const ISlateStyle& OwnerStyle)
+                                                                         const FName& StyleName, const ISlateStyle& OwnerStyle)
 {
 	const FName StyleSetName = OwnerStyle.GetStyleSetName();
 	const UE::Editor::DataStorage::FMapKey Key = UE::Editor::DataStorage::FMapKey(ISlateStyle::Join(StyleSetName, TCHAR_TO_ANSI(*StyleName.ToString())));
@@ -130,9 +233,11 @@ UE::Editor::DataStorage::RowHandle UTedsStylingFactory::AddOrGetStyleRow(IEditor
 		DataStorage->MapRow(Key, Row);
 
 		DataStorage->GetColumn<FNameColumn>(Row)->Name = StyleName;
-		DataStorage->GetColumn<FSlateStyleSetColumn>(Row)->StyleSetName = StyleSetName;
+		DataStorage->AddColumn<FSlateStyleSetTag>(Row, StyleSetName);
+		DataStorage->AddColumn(Row, FSlateStyleSetColumn{.StyleSetName = StyleSetName});
 	}
 
 	return Row;
-
 }
+
+#undef LOCTEXT_NAMESPACE
diff --git a/Engine/Plugins/Experimental/EditorDataStorageFeatures/Source/TedsStyling/Private/TedsStylingFactory.h b/Engine/Plugins/Experimental/EditorDataStorageFeatures/Source/TedsStyling/Private/TedsStylingFactory.h
--- a/Engine/Plugins/Experimental/EditorDataStorageFeatures/Source/TedsStyling/Private/TedsStylingFactory.h
+++ b/Engine/Plugins/Experimental/EditorDataStorageFeatures/Source/TedsStyling/Private/TedsStylingFactory.h
@@ -9,7 +9,9 @@
 
 struct FSlateBrush;
 struct FSlateColor;
+struct FSlateFontInfo;
 class ISlateStyle;
+struct FSlateWidgetStyle;
 
 UCLASS()
 class UTedsStylingFactory : public UEditorDataStorageFactory
@@ -19,13 +21,15 @@ class UTedsStylingFactory : public UEditorDataStorageFactory
 public:
 	~UTedsStylingFactory() override = default;
 
-	void RegisterTables(IEditorDataStorageProvider& DataStorage) override;
-	virtual void RegisterQueries(IEditorDataStorageProvider& DataStorage) override;
-
+	virtual void RegisterTables(IEditorDataStorageProvider& DataStorage) override;
+	virtual void RegisterWidgetPurposes(IEditorDataStorageUiProvider& DataStorageUi) const override;
+	
 	static void RegisterAllKnownStyles();
 	
 private:
 	static void RegisterBrush(IEditorDataStorageProvider* DataStorage, const FName& StyleName, const FSlateBrush* Brush, const ISlateStyle& OwnerStyle);
 	static void RegisterColor(IEditorDataStorageProvider* DataStorage, const FName& StyleName, const FSlateColor& Color, const ISlateStyle& OwnerStyle);
+	static void RegisterFont(IEditorDataStorageProvider* DataStorage, const FName& StyleName, const FSlateFontInfo& Font, const ISlateStyle& OwnerStyle);
+	static void RegisterWidgetStyle(IEditorDataStorageProvider* DataStorage, const FName& WidgetStyleName, const FSlateWidgetStyle* WidgetStyle, const FName& StyleName, const ISlateStyle& OwnerStyle);
 	static UE::Editor::DataStorage::RowHandle AddOrGetStyleRow(IEditorDataStorageProvider* DataStorage, const FName& StyleName, const ISlateStyle& OwnerStyle);
 };
diff --git a/Engine/Plugins/Experimental/EditorDataStorageFeatures/Source/TedsStyling/Private/TedsStylingModule.cpp b/Engine/Plugins/Experimental/EditorDataStorageFeatures/Source/TedsStyling/Private/TedsStylingModule.cpp
--- a/Engine/Plugins/Experimental/EditorDataStorageFeatures/Source/TedsStyling/Private/TedsStylingModule.cpp
+++ b/Engine/Plugins/Experimental/EditorDataStorageFeatures/Source/TedsStyling/Private/TedsStylingModule.cpp
@@ -1,5 +1,93 @@
 // Copyright Epic Games, Inc. All Rights Reserved.
 
+#include "Framework/Application/SlateApplication.h"
+#include "Framework/Docking/TabManager.h"
 #include "Modules/ModuleManager.h"
+#include "WorkspaceMenuStructureModule.h"
+#include "Widgets/Docking/SDockTab.h"
+#include "Widgets/SSlateGallery.h"
+#include "WorkspaceMenuStructure.h"
 
-IMPLEMENT_MODULE(FDefaultModuleImpl, TedsStyling)
+#define LOCTEXT_NAMESPACE "TedsStylingModule"
+
+
+namespace TedsStylingModule::Private
+{
+	FName SlateGalleryTabName = TEXT("Slate Gallery");
+	
+	static FAutoConsoleCommand RegisterStylesCommand
+		(
+			TEXT("TEDS.UI.OpenSlateGallery"),
+			TEXT("Display the Slate Gallery displaying all styles in TEDS"),
+			FConsoleCommandDelegate::CreateLambda([]()
+			{
+				FGlobalTabmanager::Get()->TryInvokeTab(SlateGalleryTabName);
+			})
+		);
+}
+
+class FTedsStylingModule : public IModuleInterface, public FGCObject
+{
+public:
+	~FTedsStylingModule() override = default;
+
+	//
+	// IModuleInterface
+	//
+	virtual void StartupModule() override;
+	virtual void ShutdownModule() override;
+
+	//
+	// FGCObject
+	//
+	
+	void AddReferencedObjects(FReferenceCollector& Collector) override;
+	FString GetReferencerName() const override;
+	
+	TSharedRef<SDockTab> OpenSlateGallery(const FSpawnTabArgs& Args);
+};
+
+
+TSharedRef<SDockTab> FTedsStylingModule::OpenSlateGallery(const FSpawnTabArgs& Args)
+{
+	TSharedRef<SDockTab> Tab =  SNew(SDockTab)
+		.TabRole(ETabRole::MajorTab);
+
+	Tab->SetContent(SNew(SSlateGallery, Tab, Args.GetOwnerWindow()));
+
+	return Tab;
+}
+
+void FTedsStylingModule::StartupModule()
+{
+	FGlobalTabmanager::Get()->RegisterNomadTabSpawner(
+		TedsStylingModule::Private::SlateGalleryTabName,
+		FOnSpawnTab::CreateRaw(this, &FTedsStylingModule::OpenSlateGallery))
+		.SetGroup(WorkspaceMenu::GetMenuStructure().GetDeveloperToolsDebugCategory())
+		.SetDisplayName(LOCTEXT("SlateGalleryTabName", "Slate Gallery"))
+		.SetAutoGenerateMenuEntry(false)
+		.SetTooltipText(LOCTEXT("SlateGalleryTooltipName", "Opens the Slate Gallery"))
+		.SetIcon(FSlateIcon(FAppStyle::GetAppStyleSetName(), "Debug"));
+}
+
+void FTedsStylingModule::ShutdownModule()
+{
+	if (FSlateApplication::IsInitialized())
+	{
+		FGlobalTabmanager::Get()->UnregisterNomadTabSpawner(TedsStylingModule::Private::SlateGalleryTabName);
+	}
+}
+
+void FTedsStylingModule::AddReferencedObjects(FReferenceCollector& Collector)
+{
+	
+}
+
+FString FTedsStylingModule::GetReferencerName() const
+{
+	return TEXT("TEDS: Editor Data Storage Styling Module");
+}
+
+IMPLEMENT_MODULE(FTedsStylingModule, TedsStyling)
+
+#undef LOCTEXT_NAMESPACE
\ No newline at end of file
diff --git a/Engine/Plugins/Experimental/EditorDataStorageFeatures/Source/TedsStyling/Private/Widgets/SSlateGallery.cpp b/Engine/Plugins/Experimental/EditorDataStorageFeatures/Source/TedsStyling/Private/Widgets/SSlateGallery.cpp
new file mode 100644
--- /dev/null
+++ b/Engine/Plugins/Experimental/EditorDataStorageFeatures/Source/TedsStyling/Private/Widgets/SSlateGallery.cpp
@@ -0,0 +1,360 @@
+// Copyright Epic Games, Inc. All Rights Reserved.
+
+#include "SSlateGallery.h"
+
+#include "TedsRowViewNode.h"
+#include "Columns/TedsStylingColumns.h"
+#include "Elements/Columns/TypedElementSlateWidgetColumns.h"
+#include "Elements/Common/EditorDataStorageFeatures.h"
+#include "Elements/Common/TypedElementQueryDescription.h"
+#include "Elements/Framework/TypedElementQueryBuilder.h"
+#include "Elements/Interfaces/TypedElementDataStorageInterface.h"
+#include "Framework/Docking/LayoutService.h"
+#include "Framework/Docking/TabManager.h"
+#include "Widgets/Docking/SDockTab.h"
+#include "Widgets/Layout/SScrollBox.h"
+#include "Widgets/STedsTileViewer.h"
+
+#define LOCTEXT_NAMESPACE "SSlateGallery"
+
+namespace TedsStyling::Private
+{
+	FName ColorsTabName = TEXT("Colors");
+	FName IconsTabName = TEXT("Icons");
+	FName FontsTabName = TEXT("Fonts");
+	FName ButtonsTabName = TEXT("Buttons");
+	FName CheckboxesTabName = TEXT("Checkboxes");
+	
+	// Wrapper widget around the table viewer so we can manage the lifetime of the query and query stack manually for now
+	class STableViewerWrapper : public SCompoundWidget
+	{
+	public:
+		SLATE_BEGIN_ARGS(STableViewerWrapper)
+			: _ItemHeight(128)
+			, _ItemWidth(128)
+		{}
+		SLATE_ARGUMENT(UE::Editor::DataStorage::FQueryDescription, QueryDescription)
+		SLATE_ARGUMENT(TArray<TWeakObjectPtr<const UScriptStruct>>, Columns)
+		SLATE_ARGUMENT(FName, WidgetPurpose)
+		SLATE_ATTRIBUTE(float, ItemHeight)
+		SLATE_ATTRIBUTE(float, ItemWidth)
+	SLATE_END_ARGS()
+
+	void Construct(const FArguments& Arguments)
+	{
+		using namespace UE::Editor::DataStorage;
+		Storage = GetMutableDataStorageFeature<IEditorDataStorageProvider>(StorageFeatureName);
+
+		FQueryDescription QueryDescription = Arguments._QueryDescription;
+	
+		QueryHandle = Storage->RegisterQuery(MoveTemp(QueryDescription));
+	
+		RowView = MakeShared<QueryStack::FRowViewNode>(FRowHandleArrayView(Rows.GetData(), Rows.Num(),
+			FRowHandleArrayView::EFlags::IsSorted | FRowHandleArrayView::EFlags::IsUnique));
+
+		ChildSlot
+		[
+			SNew(UE::Editor::DataStorage::STedsTileViewer)
+					.TileStyle(&FAppStyle::GetWidgetStyle<FTableRowStyle>("PropertyWindow.PropertyRow"))
+					.Visibility(TAttribute<EVisibility>::CreateSP(this, &TedsStyling::Private::STableViewerWrapper::GetVisibility))
+					.QueryStack(RowView)
+					.ItemHeight(Arguments._ItemHeight)
+					.ItemWidth(Arguments._ItemWidth)
+					.Columns(Arguments._Columns)
+					.WidgetPurpose(Arguments._WidgetPurpose)
+					.SelectionMode(ESelectionMode::None)
+		];
+	}
+
+		virtual void Tick(const FGeometry& AllottedGeometry, const double InCurrentTime, const float InDeltaTime) override
+		{
+			// On tick, we run the query again and update the query stack if the rows changed
+		
+			TSet<UE::Editor::DataStorage::RowHandle> NewRows;
+
+			using namespace UE::Editor::DataStorage;
+
+			Storage->RunQuery(QueryHandle,
+				Queries::CreateDirectQueryCallbackBinding([&NewRows]
+					(const IEditorDataStorageProvider::IDirectQueryContext& Context, const RowHandle*)
+					{
+						NewRows.Append(Context.GetRowHandles());
+					}));
+
+			// Check if the two sets are equal, i.e no changes and no need to update the table viewer
+			const bool bSetsEqual = (Rows_Set.Num() == NewRows.Num()) && Rows_Set.Includes(NewRows);
+
+			if (!bSetsEqual)
+			{
+				Rows_Set = NewRows;
+				Rows = Rows_Set.Array();
+			
+				RowView->ResetView(FRowHandleArrayView(Rows.GetData(), Rows.Num(),
+					FRowHandleArrayView::EFlags::IsSorted | FRowHandleArrayView::EFlags::IsUnique));
+			}
+		}
+
+		virtual ~STableViewerWrapper() override
+		{
+			Storage->UnregisterQuery(QueryHandle);
+		}
+
+		EVisibility GetVisibility() const
+		{
+			return !Rows.Num() ? EVisibility::Collapsed : EVisibility::Visible;
+		}
+
+	private:
+		UE::Editor::DataStorage::QueryHandle QueryHandle = UE::Editor::DataStorage::InvalidQueryHandle;
+		TArray<UE::Editor::DataStorage::RowHandle> Rows;
+		TSet<UE::Editor::DataStorage::RowHandle> Rows_Set;
+		TSharedPtr<UE::Editor::DataStorage::QueryStack::FRowViewNode> RowView;
+		IEditorDataStorageProvider* Storage = nullptr;
+	};
+
+	// Widget that displays rows matching a specific style tag/column grouped by their stylesheets
+	// Each group is represented by a table viewer querying for the dynamic column representing a stylesheet
+	class SStyleSheetGrouper : public SCompoundWidget
+	{
+	public:
+		SLATE_BEGIN_ARGS(SStyleSheetGrouper)
+			: _ItemHeight(128)
+			, _ItemWidth(128)
+		{}
+		SLATE_ARGUMENT(TArray<TWeakObjectPtr<const UScriptStruct>>, Columns)
+		SLATE_ARGUMENT(FName, WidgetPurpose)
+		SLATE_ATTRIBUTE(float, ItemHeight)
+		SLATE_ATTRIBUTE(float, ItemWidth)
+		SLATE_END_ARGS()
+
+		// ColumnToQuery is the style tag the table viewers will query for and group based on stylesheet
+		void Construct(const FArguments& Arguments, const UE::Editor::DataStorage::Queries::FColumnBase& ColumnToQuery)
+		{
+			using namespace UE::Editor::DataStorage::Queries;
+			
+			const IEditorDataStorageProvider* Storage =
+				UE::Editor::DataStorage::GetDataStorageFeature<IEditorDataStorageProvider>(StorageFeatureName);
+
+			// Get all dynamic columns that derive from FSlatStyleSetTag so we can sort based on stylesheet
+			static const UScriptStruct* BaseStyleSetTag = FSlateStyleSetTag::StaticStruct();
+			TArray<const UScriptStruct*> DynamicColumns;
+			
+			Storage->ForEachDynamicColumn(*BaseStyleSetTag, [&DynamicColumns](const UScriptStruct& DynamicColumn)
+			{
+				DynamicColumns.Add(&DynamicColumn);
+			});
+			
+			TSharedRef<SScrollBox> Box = SNew(SScrollBox);
+
+			// Create an expandable area with a table viewer for each of those
+			for (const UScriptStruct* DynamicColumn : DynamicColumns)
+			{
+				const FName StyleSetName = UE::Editor::DataStorage::ColumnUtils::GetDynamicColumnIdentifier(DynamicColumn);
+
+				FConditions Conditions(ColumnToQuery && TColumn(BaseStyleSetTag, StyleSetName));
+
+				FQueryDescription QueryDescription =
+					Select()
+						.Where(Conditions)
+					.Compile();
+
+				TSharedRef<TedsStyling::Private::STableViewerWrapper> TableViewer = SNew(TedsStyling::Private::STableViewerWrapper)
+						.WidgetPurpose(Arguments._WidgetPurpose)
+						.ItemHeight(Arguments._ItemHeight)
+						.ItemWidth(Arguments._ItemWidth)
+						.Columns(Arguments._Columns)
+						.QueryDescription(QueryDescription);
+
+				Box->AddSlot()
+				.HAlign(HAlign_Left)
+				[
+					SNew(STextBlock)
+					.Visibility(TAttribute<EVisibility>::CreateSP(&TableViewer.Get(), &TedsStyling::Private::STableViewerWrapper::GetVisibility))
+					.Text(FText::FromName(StyleSetName))
+					.Font(FAppStyle::GetFontStyle("Font.Large.Bold"))
+				];
+
+				Box->AddSlot()
+				.HAlign(HAlign_Fill)
+				[
+					TableViewer
+				];
+			}
+
+			ChildSlot
+			[
+				Box
+			];
+		}
+
+	};
+}
+
+void SSlateGallery::Construct(const FArguments& InArgs, const TSharedRef<SDockTab>& ConstructUnderMajorTab, const TSharedPtr<SWindow>& ConstructUnderWindow)
+{
+	// Create the tab manager for our sub tabsa
+	TabManager = FGlobalTabmanager::Get()->NewTabManager(ConstructUnderMajorTab);
+	TabManager->SetAllowWindowMenuBar(true);
+
+	RegisterTabSpawners();
+
+	TSharedRef<FTabManager::FLayout> Layout = FTabManager::NewLayout("SlateGalleryTest3")
+		->AddArea
+		(
+			FTabManager::NewPrimaryArea()
+			->Split
+			(
+				FTabManager::NewStack()
+				->AddTab(TedsStyling::Private::ColorsTabName, ETabState::OpenedTab)
+				->AddTab(TedsStyling::Private::IconsTabName, ETabState::OpenedTab)
+				->AddTab(TedsStyling::Private::FontsTabName, ETabState::OpenedTab)
+				->AddTab(TedsStyling::Private::ButtonsTabName, ETabState::OpenedTab)
+				->AddTab(TedsStyling::Private::CheckboxesTabName, ETabState::OpenedTab)
+			)
+		);
+	
+	Layout = FLayoutSaveRestore::LoadFromConfig(GEditorLayoutIni, Layout);
+
+	ChildSlot
+	[
+		TabManager->RestoreFrom(Layout, ConstructUnderWindow).ToSharedRef()
+	];
+
+	TabManager->SetOnPersistLayout(
+		FTabManager::FOnPersistLayout::CreateStatic(
+			[](const TSharedRef<FTabManager::FLayout>& InLayout)
+			{
+				if (InLayout->GetPrimaryArea().Pin().IsValid())
+				{
+					FLayoutSaveRestore::SaveToConfig(GEditorLayoutIni, InLayout);
+				}
+			}
+		)
+	);
+}
+
+void SSlateGallery::RegisterTabSpawners()
+{
+	const TSharedRef<FWorkspaceItem> AppMenuGroup =
+		TabManager->AddLocalWorkspaceMenuCategory(LOCTEXT("SlateGalleryGroupName", "Slate Gallery"));
+
+	TabManager->RegisterTabSpawner(
+		TedsStyling::Private::ColorsTabName,
+		FOnSpawnTab::CreateRaw(this, &SSlateGallery::SpawnColorsTab))
+		.SetGroup(AppMenuGroup)
+		.SetDisplayName(LOCTEXT("ColorsTabName", "Colors"))
+		.SetAutoGenerateMenuEntry(false);
+	
+	TabManager->RegisterTabSpawner(
+		TedsStyling::Private::IconsTabName,
+		FOnSpawnTab::CreateRaw(this, &SSlateGallery::SpawnIconsTab))
+		.SetGroup(AppMenuGroup)
+		.SetDisplayName(LOCTEXT("IconsTabName", "Icons"));
+	
+	TabManager->RegisterTabSpawner(
+		TedsStyling::Private::FontsTabName,
+		FOnSpawnTab::CreateRaw(this, &SSlateGallery::SpawnFontsTab))
+		.SetGroup(AppMenuGroup)
+		.SetDisplayName(LOCTEXT("FontsTabName", "Fonts"));
+
+	TabManager->RegisterTabSpawner(
+		TedsStyling::Private::ButtonsTabName,
+		FOnSpawnTab::CreateRaw(this, &SSlateGallery::SpawnButtonsTab))
+		.SetGroup(AppMenuGroup)
+		.SetDisplayName(LOCTEXT("ButtonsTabName", "Buttons"));
+
+	TabManager->RegisterTabSpawner(
+		TedsStyling::Private::CheckboxesTabName,
+		FOnSpawnTab::CreateRaw(this, &SSlateGallery::SpawnCheckboxesTab))
+		.SetGroup(AppMenuGroup)
+		.SetDisplayName(LOCTEXT("CheckboxesTabName", "Checkboxes"));
+}
+
+TSharedRef<SDockTab> SSlateGallery::SpawnColorsTab(const FSpawnTabArgs& Args)
+{
+	using namespace UE::Editor::DataStorage::Queries;
+	
+	return SNew(SDockTab)
+		.TabRole(ETabRole::PanelTab)
+		[
+			SNew(TedsStyling::Private::SStyleSheetGrouper, TColumn<FSlateStyleTag>("Color"))
+				.WidgetPurpose(TEXT("SlateGallery.Color"))
+		];
+}
+
+TSharedRef<SDockTab> SSlateGallery::SpawnIconsTab(const FSpawnTabArgs& Args)
+{
+	using namespace UE::Editor::DataStorage::Queries;
+	
+	return SNew(SDockTab)
+		.TabRole(ETabRole::PanelTab)
+		[
+			SNew(TedsStyling::Private::SStyleSheetGrouper, TColumn<FSlateStyleTag>("Icon"))
+				.WidgetPurpose(TEXT("SlateGallery.Icons"))
+				.ItemHeight(64.0f)
+				.ItemWidth(64.0f)
+		];
+}
+
+TSharedRef<SDockTab> SSlateGallery::SpawnFontsTab(const FSpawnTabArgs& Args)
+{
+	using namespace UE::Editor::DataStorage::Queries;
+	
+	return SNew(SDockTab)
+		.TabRole(ETabRole::PanelTab)
+		[
+			SNew(TedsStyling::Private::SStyleSheetGrouper, TColumn<FSlateFontColumn>())
+				.WidgetPurpose(TEXT("SlateGallery.Font"))
+				.ItemHeight(200)
+				.ItemWidth(500)
+		];
+}
+
+TSharedRef<SDockTab> SSlateGallery::SpawnCheckboxesTab(const FSpawnTabArgs& Args)
+{
+	using namespace UE::Editor::DataStorage::Queries;
+	
+	IEditorDataStorageProvider* DataStorage = GetMutableDataStorageFeature<IEditorDataStorageProvider>(StorageFeatureName);
+	
+	const UScriptStruct* DynamicColumn = DataStorage->GenerateDynamicColumn(FDynamicColumnDescription
+					{
+						.TemplateType = FSlateWidgetStyleColumn::StaticStruct(),
+						.Identifier = FCheckBoxStyle::TypeName
+					});
+	
+	return SNew(SDockTab)
+		.TabRole(ETabRole::PanelTab)
+		[
+			SNew(TedsStyling::Private::SStyleSheetGrouper, TColumn<FSlateWidgetStyleColumn>(FCheckBoxStyle::TypeName))
+				.WidgetPurpose(TEXT("General.Cell"))
+				.ItemHeight(64)
+				.ItemWidth(320)
+				.Columns({DynamicColumn})
+		];
+}
+
+TSharedRef<SDockTab> SSlateGallery::SpawnButtonsTab(const FSpawnTabArgs& Args)
+{
+	using namespace UE::Editor::DataStorage::Queries;
+	
+	IEditorDataStorageProvider* DataStorage = GetMutableDataStorageFeature<IEditorDataStorageProvider>(StorageFeatureName);
+	
+	const UScriptStruct* DynamicColumn = DataStorage->GenerateDynamicColumn(FDynamicColumnDescription
+					{
+						.TemplateType = FSlateWidgetStyleColumn::StaticStruct(),
+						.Identifier = FButtonStyle::TypeName
+					});
+	
+	return SNew(SDockTab)
+		.TabRole(ETabRole::PanelTab)
+		[
+			SNew(TedsStyling::Private::SStyleSheetGrouper, TColumn<FSlateWidgetStyleColumn>(FButtonStyle::TypeName))
+				.WidgetPurpose(TEXT("General.Cell"))
+				.ItemHeight(64)
+				.ItemWidth(320)
+				.Columns({DynamicColumn})
+		];
+}
+
+
+#undef LOCTEXT_NAMESPACE
diff --git a/Engine/Plugins/Experimental/EditorDataStorageFeatures/Source/TedsStyling/Private/Widgets/SSlateGallery.h b/Engine/Plugins/Experimental/EditorDataStorageFeatures/Source/TedsStyling/Private/Widgets/SSlateGallery.h
new file mode 100644
--- /dev/null
+++ b/Engine/Plugins/Experimental/EditorDataStorageFeatures/Source/TedsStyling/Private/Widgets/SSlateGallery.h
@@ -0,0 +1,34 @@
+// Copyright Epic Games, Inc. All Rights Reserved.
+
+#pragma once
+
+#include "Widgets/SCompoundWidget.h"
+
+class SDockTab;
+class SWindow;
+class FTabManager;
+class FSpawnTabArgs;
+
+class SSlateGallery : public SCompoundWidget
+{
+public:
+
+	SLATE_BEGIN_ARGS(SSlateGallery) { }
+	SLATE_END_ARGS()
+	
+	void Construct(const FArguments& InArgs, const TSharedRef<SDockTab>& ConstructUnderMajorTab, const TSharedPtr<SWindow>& ConstructUnderWindow);
+
+private:
+
+	void RegisterTabSpawners();
+	TSharedRef<SDockTab> SpawnColorsTab(const FSpawnTabArgs& Args);
+	TSharedRef<SDockTab> SpawnIconsTab(const FSpawnTabArgs& Args);
+	TSharedRef<SDockTab> SpawnFontsTab(const FSpawnTabArgs& Args);
+	TSharedRef<SDockTab> SpawnButtonsTab(const FSpawnTabArgs& Args);
+	TSharedRef<SDockTab> SpawnCheckboxesTab(const FSpawnTabArgs& Args);
+
+private:
+
+	// Holds the tab manager that manages the front-end's tabs.
+	TSharedPtr<FTabManager> TabManager;
+};
\ No newline at end of file
diff --git a/Engine/Plugins/Experimental/EditorDataStorageFeatures/Source/TedsStyling/Private/Widgets/SlateBrushWidget.cpp b/Engine/Plugins/Experimental/EditorDataStorageFeatures/Source/TedsStyling/Private/Widgets/SlateBrushWidget.cpp
--- a/Engine/Plugins/Experimental/EditorDataStorageFeatures/Source/TedsStyling/Private/Widgets/SlateBrushWidget.cpp
+++ b/Engine/Plugins/Experimental/EditorDataStorageFeatures/Source/TedsStyling/Private/Widgets/SlateBrushWidget.cpp
@@ -5,6 +5,7 @@
 #include "Columns/TedsStylingColumns.h"
 #include "Elements/Columns/TypedElementMiscColumns.h"
 #include "Elements/Framework/TypedElementAttributeBinding.h"
+#include "Elements/Columns/TypedElementSlateWidgetColumns.h"
 #include "Styling/SlateStyleRegistry.h"
 #include "Widgets/Images/SImage.h"
 #include "Widgets/Layout/SBox.h"
@@ -15,8 +16,66 @@ void USlateStylePreviewWidget::RegisterWidgetConstructors(IEditorDataStorageProv
 	IEditorDataStorageUiProvider& DataStorageUi) const
 {
 	using namespace UE::Editor::DataStorage::Queries;
-	DataStorageUi.RegisterWidgetFactory<FSlateStylePreviewWidgetConstructor>(TEXT("General.RowLabel"), TColumn<FNameColumn>() && TColumn<FSlateStyleTag>());
-	DataStorageUi.RegisterWidgetFactory<FSlateStylePreviewWidgetConstructor>(TEXT("General.Cell"), TColumn<FNameColumn>() && TColumn<FSlateStyleTag>());
+	DataStorageUi.RegisterWidgetFactory<FSlateStylePreviewWidgetConstructor>(TEXT("General.RowLabel"),
+		TColumn<FNameColumn>() && TColumn<FSlateStyleTag>("Icon") && TColumn<FSlateStyleSetTag>());
+	
+	DataStorageUi.RegisterWidgetFactory<FSlateStylePreviewWidgetConstructor>(TEXT("General.Cell"),
+		TColumn<FNameColumn>() && TColumn<FSlateStyleTag>("Icon")  && TColumn<FSlateStyleSetTag>());
+
+	DataStorageUi.RegisterWidgetFactory<FSlateIconWidgetConstructor>(TEXT("SlateGallery.Icons"));
+}
+
+FSlateIconWidgetConstructor::FSlateIconWidgetConstructor()
+	: Super(StaticStruct())
+{
+	
+}
+
+TSharedPtr<SWidget> FSlateIconWidgetConstructor::CreateWidget(IEditorDataStorageProvider* DataStorage, IEditorDataStorageUiProvider* DataStorageUi,
+	UE::Editor::DataStorage::RowHandle TargetRow, UE::Editor::DataStorage::RowHandle WidgetRow,
+	const UE::Editor::DataStorage::FMetaDataView& Arguments)
+{
+	UE::Editor::DataStorage::FAttributeBinder Binder(TargetRow, DataStorage);
+
+	const FVector2D& SlateIconSize = Binder.BindData(&FIconSizeColumn::IconSize, [](const FVector2D& IconSize)
+		{
+			// If the size of the icon is (0,0)/empty, set it to the default of 16x16
+			return IconSize.IsZero() ? FVector2D(16.0f, 16.0f) : IconSize;
+		}).Get();
+	
+	TSharedPtr<SWidget> IconWidget = SNew(SBox)
+		.HAlign(HAlign_Center)
+		.VAlign(VAlign_Center)
+		[
+			SNew(SImage)
+				.ToolTipText(Binder.BindText(&FNameColumn::Name))
+				.DesiredSizeOverride(SlateIconSize)
+				.Image_Lambda([DataStorage, TargetRow]()
+					{
+						const FNameColumn* StyleNameColumn = DataStorage->GetColumn<FNameColumn>(TargetRow);
+
+						const FSlateStyleSetColumn* StyleSetColumn = DataStorage->GetColumn<FSlateStyleSetColumn>(TargetRow);
+						if (StyleNameColumn && StyleSetColumn)
+						{
+							if (const ISlateStyle* Style = FSlateStyleRegistry::FindSlateStyle(StyleSetColumn->StyleSetName))
+							{
+								if (const FSlateBrush* Brush = Style->GetBrush(StyleNameColumn->Name))
+								{
+									return Brush;
+								}
+							}
+						}
+						return FAppStyle::GetBrush("NoBrush");
+					})
+		];
+
+	return IconWidget;
+}
+
+TConstArrayView<const UScriptStruct*> FSlateIconWidgetConstructor::GetAdditionalColumnsList() const
+{
+	static TTypedElementColumnTypeList<FSlateIconWidgetTag> Columns;
+	return Columns;
 }
 
 FSlateStylePreviewWidgetConstructor::FSlateStylePreviewWidgetConstructor()
@@ -29,6 +88,8 @@ TSharedPtr<SWidget> FSlateStylePreviewWidgetConstructor::CreateWidget(IEditorDat
 	const UE::Editor::DataStorage::FMetaDataView& Arguments)
 {
 	UE::Editor::DataStorage::FAttributeBinder Binder(TargetRow, DataStorage);
+
+	FName* StyleSet = MatchedDynamicTemplates.Find(FSlateStyleSetTag::StaticStruct());
 	
 	return SNew(SHorizontalBox)
 			+ SHorizontalBox::Slot()
@@ -40,17 +101,14 @@ TSharedPtr<SWidget> FSlateStylePreviewWidgetConstructor::CreateWidget(IEditorDat
 				.MaxDesiredHeight(16.0f)
 				.MaxDesiredWidth(16.0f)
 				[
-					// Currently we cannot differentiate between image brushes and other brushes in TEDS so we only show image brushes as an SImage
-					// and don't support showing a preview of the rest
 					SNew(SImage)
-					.Image_Lambda([DataStorage, TargetRow]()
+					.Image_Lambda([DataStorage, TargetRow, StyleSet]()
 					{
 						const FNameColumn* StyleNameColumn = DataStorage->GetColumn<FNameColumn>(TargetRow);
-						const FSlateStyleSetColumn* StyleSetColumn = DataStorage->GetColumn<FSlateStyleSetColumn>(TargetRow);
 
-						if (StyleNameColumn && StyleSetColumn && DataStorage->HasColumns<FSlateBrushTag>(TargetRow))
+						if (StyleNameColumn && StyleSet)
 						{
-							if(const ISlateStyle* Style = FSlateStyleRegistry::FindSlateStyle(StyleSetColumn->StyleSetName))
+							if(const ISlateStyle* Style = FSlateStyleRegistry::FindSlateStyle(*StyleSet))
 							{
 								if(const FSlateBrush* Brush = Style->GetBrush(StyleNameColumn->Name))
 								{
diff --git a/Engine/Plugins/Experimental/EditorDataStorageFeatures/Source/TedsStyling/Private/Widgets/SlateBrushWidget.h b/Engine/Plugins/Experimental/EditorDataStorageFeatures/Source/TedsStyling/Private/Widgets/SlateBrushWidget.h
--- a/Engine/Plugins/Experimental/EditorDataStorageFeatures/Source/TedsStyling/Private/Widgets/SlateBrushWidget.h
+++ b/Engine/Plugins/Experimental/EditorDataStorageFeatures/Source/TedsStyling/Private/Widgets/SlateBrushWidget.h
@@ -25,6 +25,31 @@ public:
 
 // Widget to show a slate brush drawn as an SImage
 USTRUCT()
+struct FSlateIconWidgetConstructor : public FSimpleWidgetConstructor
+{
+	GENERATED_BODY()
+
+public:
+	FSlateIconWidgetConstructor();
+	~FSlateIconWidgetConstructor() override = default;
+
+	virtual TSharedPtr<SWidget> CreateWidget(
+		IEditorDataStorageProvider* DataStorage,
+		IEditorDataStorageUiProvider* DataStorageUi,
+		UE::Editor::DataStorage::RowHandle TargetRow,
+		UE::Editor::DataStorage::RowHandle WidgetRow,
+		const UE::Editor::DataStorage::FMetaDataView& Arguments) override;
+	TConstArrayView<const UScriptStruct*> GetAdditionalColumnsList() const override;
+};
+
+USTRUCT(meta = (DisplayName = "Icon"))
+struct FSlateIconWidgetTag : public FEditorDataStorageTag
+{
+	GENERATED_BODY()
+};
+
+// Widget to show a slate brush drawn as an SImage alongside the name
+USTRUCT()
 struct FSlateStylePreviewWidgetConstructor : public FSimpleWidgetConstructor
 {
 	GENERATED_BODY()
diff --git a/Engine/Plugins/Experimental/EditorDataStorageFeatures/Source/TedsStyling/Public/Columns/TedsStylingColumns.h b/Engine/Plugins/Experimental/EditorDataStorageFeatures/Source/TedsStyling/Public/Columns/TedsStylingColumns.h
--- a/Engine/Plugins/Experimental/EditorDataStorageFeatures/Source/TedsStyling/Public/Columns/TedsStylingColumns.h
+++ b/Engine/Plugins/Experimental/EditorDataStorageFeatures/Source/TedsStyling/Public/Columns/TedsStylingColumns.h
@@ -3,21 +3,35 @@
 #pragma once
 
 #include "Elements/Interfaces/TypedElementDataStorageInterface.h"
+#include "Fonts/SlateFontInfo.h"
+#include "Styling/SlateWidgetStyle.h"
 
 #include "TedsStylingColumns.generated.h"
 
 /**
- * Owner style set
+ * Dynamic column for stylesets (e.g FSlateStyleSetTag("CoreStyle")
  */
-USTRUCT(meta = (DisplayName = "Style Set"))
-struct FSlateStyleSetColumn final : public FEditorDataStorageColumn
+USTRUCT(meta=(EditorDataStorage_DynamicColumnTemplate))
+struct FSlateStyleSetTag : public FEditorDataStorageTag
 {
 	GENERATED_BODY()
+	
+	struct EditorDataStorage_DynamicColumnTemplate{};
+};
 
+/**
+ * Storing styleset as both a dynamic column and a regular column right now for two way lookup
+ */
+USTRUCT()
+struct FSlateStyleSetColumn : public FEditorDataStorageColumn
+{
+	GENERATED_BODY()
+	
 	UPROPERTY(meta = (Searchable))
 	FName StyleSetName;
 };
 
+
 /**
  * Path to style resource on disk
  */
@@ -30,21 +44,59 @@ struct FSlateStylePathColumn final : public FEditorDataStorageColumn
 	FName StylePath;
 };
 
+/**
+ * Stored size of an Icon
+ */
+USTRUCT(meta = (DisplayName = "Icon Size"))
+struct FIconSizeColumn final : public FEditorDataStorageColumn
+{
+	GENERATED_BODY()
+
+	UPROPERTY(meta = (Searchable))
+	FVector2D IconSize;
+};
+
 /**
  * Tag on rows that are part of a stylesheet
  */
-USTRUCT(meta = (DisplayName = "Slate Style"))
+USTRUCT(meta=(EditorDataStorage_DynamicColumnTemplate))
 struct FSlateStyleTag final : public FEditorDataStorageTag
 {
 	GENERATED_BODY()
+	
+	struct EditorDataStorage_DynamicColumnTemplate{};
 };
 
 /**
- * Tag on rows that are brushes in the stylesheet
- * TEDS UI TODO: Once TEDS UI supports dynamic columns, we can remove this tag and make FSlateStyleTagDynamic
+ * Tag on rows that are a default color
  */
-USTRUCT(meta = (DisplayName = "Slate Brush"))
-struct FSlateBrushTag final : public FEditorDataStorageTag
+USTRUCT()
+struct FStyleColorTag final : public FEditorDataStorageTag
 {
 	GENERATED_BODY()
 };
+
+/**
+ * Column containing font info for font styles
+ */
+USTRUCT(meta = (DisplayName = "Slate Font"))
+struct FSlateFontColumn final : public FEditorDataStorageColumn
+{
+	GENERATED_BODY()
+
+	UPROPERTY()
+	FSlateFontInfo FontInfo;
+};
+
+/**
+ * Dynamic column for widget styles (e.g FSlateWidgetStyleColumn("ButtonStyle")
+ */
+USTRUCT(meta=(EditorDataStorage_DynamicColumnTemplate))
+struct FSlateWidgetStyleColumn final : public FEditorDataStorageColumn
+{
+	GENERATED_BODY()
+
+	const FSlateWidgetStyle* WidgetStyle;
+	
+	struct EditorDataStorage_DynamicColumnTemplate{};
+};
\ No newline at end of file
diff --git a/Engine/Plugins/Experimental/EditorDataStorageFeatures/Source/TedsStyling/TedsStyling.Build.cs b/Engine/Plugins/Experimental/EditorDataStorageFeatures/Source/TedsStyling/TedsStyling.Build.cs
--- a/Engine/Plugins/Experimental/EditorDataStorageFeatures/Source/TedsStyling/TedsStyling.Build.cs
+++ b/Engine/Plugins/Experimental/EditorDataStorageFeatures/Source/TedsStyling/TedsStyling.Build.cs
@@ -19,11 +19,17 @@ public class TedsStyling : ModuleRules
 			PrivateDependencyModuleNames.AddRange(
 				new string[]
 				{
+					"AppFramework",
 					"Core",
 					"CoreUObject",
 					"Engine",
+					"InputCore",
 					"SlateCore",
-					"Slate"
+					"Slate",
+					"TedsTableViewer",
+					"TedsQueryStack",
+					"ToolWidgets",
+					"WorkspaceMenuStructure"
 				});
 
 			PublicDependencyModuleNames.AddRange(
diff --git a/Engine/Plugins/Experimental/EditorDataStorageFeatures/Source/TedsTableViewer/Private/Widgets/STedsTileViewer.cpp b/Engine/Plugins/Experimental/EditorDataStorageFeatures/Source/TedsTableViewer/Private/Widgets/STedsTileViewer.cpp
--- a/Engine/Plugins/Experimental/EditorDataStorageFeatures/Source/TedsTableViewer/Private/Widgets/STedsTileViewer.cpp
+++ b/Engine/Plugins/Experimental/EditorDataStorageFeatures/Source/TedsTableViewer/Private/Widgets/STedsTileViewer.cpp
@@ -1,12 +1,16 @@
 // Copyright Epic Games, Inc. All Rights Reserved.
 
+
 #include "Widgets/STedsTileViewer.h"
-#include "TedsTableViewerColumn.h"
+
+#include "Brushes/SlateColorBrush.h"
 #include "Columns/SlateDelegateColumns.h"
 #include "Elements/Columns/TypedElementUIColumns.h"
 #include "Elements/Common/EditorDataStorageFeatures.h"
 #include "Elements/Framework/TypedElementAttributeBinding.h"
 #include "Elements/Framework/TypedElementDataStorageWidget.h"
+#include "TedsTableViewerColumn.h"
+#include "Styling/StyleColors.h"
 #include "Widgets/STedsTableViewerRow.h"
 
 namespace UE::Editor::DataStorage
@@ -15,11 +19,15 @@ namespace UE::Editor::DataStorage
 	{
 		OnSelectionChanged = InArgs._OnSelectionChanged;
 		WidgetPurpose = InArgs._WidgetPurpose;
+		Columns = InArgs._Columns;
+		TableRowStyle = InArgs._TileStyle;
 
-		TArray<TWeakObjectPtr<const UScriptStruct>> Columns;
 		TArray<FName> Purposes;
+		// While we might have columns for the widget constructor, we don't want to display any UI columns since that doesn't make sense
+		// so we give FTedsTableViewerModel empty dummy columns
+		TArray<TWeakObjectPtr<const UScriptStruct>> DummyColumns;
 		
-		Model = MakeShared<FTedsTableViewerModel>(InArgs._QueryStack, Columns, Purposes,
+		Model = MakeShared<FTedsTableViewerModel>(InArgs._QueryStack, DummyColumns, Purposes,
 			FTedsTableViewerModel::FIsItemVisible::CreateSP(this, &STedsTileViewer::IsItemVisible));
 
 		IEditorDataStorageUiProvider* StorageUi = Model->GetDataStorageUiProvider();
@@ -43,6 +51,8 @@ namespace UE::Editor::DataStorage
 			.ListItemsSource(&Model->GetItems())
 			.OnGenerateTile(this, &STedsTileViewer::MakeTileWidget)
 			.OnSelectionChanged(this, &STedsTileViewer::OnListSelectionChanged)
+			.ItemWidth(InArgs._ItemWidth)
+			.ItemHeight(InArgs._ItemHeight)
 			.SelectionMode(InArgs._SelectionMode)
 			.OnContextMenuOpening(Binder.BindEvent(&FWidgetContextMenuColumn::OnContextMenuOpening))
 			.OnItemScrolledIntoView(Binder.BindEvent(&FWidgetRowScrolledIntoView::OnItemScrolledIntoView))
@@ -65,16 +75,25 @@ namespace UE::Editor::DataStorage
 		// of columns
 		TSharedPtr<FTedsTableViewerColumn> Column;
 		
-		auto AssignWidgetToColumn = [&Column](TUniquePtr<FTypedElementWidgetConstructor> Constructor, TConstArrayView<TWeakObjectPtr<const UScriptStruct>>)
+		auto AssignWidgetToColumn = [&Column](TUniquePtr<FTypedElementWidgetConstructor> Constructor, TConstArrayView<TWeakObjectPtr<const UScriptStruct>> InMatchedColumns)
 		{
+			TArray<TWeakObjectPtr<const UScriptStruct>> MatchedColumnsCopy(InMatchedColumns);
 			TSharedPtr<FTypedElementWidgetConstructor> WidgetConstructor(Constructor.Release());
-			Column = MakeShared<FTedsTableViewerColumn>(TEXT("TileView"), WidgetConstructor);
+			Column = MakeShared<FTedsTableViewerColumn>(TEXT("TileView"), WidgetConstructor, MatchedColumnsCopy);
 			return false;
 		};
 	
 		IEditorDataStorageUiProvider* StorageUi = GetMutableDataStorageFeature<IEditorDataStorageUiProvider>(UiFeatureName);
 
-		StorageUi->CreateWidgetConstructors(WidgetPurpose, FMetaDataView(), AssignWidgetToColumn);
+		if (Columns.IsEmpty())
+		{
+			StorageUi->CreateWidgetConstructors(WidgetPurpose, FMetaDataView(), AssignWidgetToColumn);
+		}
+		else
+		{
+			StorageUi->CreateWidgetConstructors(WidgetPurpose, IEditorDataStorageUiProvider::EMatchApproach::ExactMatch, Columns,
+				FMetaDataView(), AssignWidgetToColumn);
+		}
 
 		if (Column)
 		{
@@ -132,6 +151,8 @@ namespace UE::Editor::DataStorage
 		}
 
 		return SNew(STableRow<TableViewerItemPtr>, OwnerTable)
+			.Padding(4.f) // TODO: Make modular
+			.Style(TableRowStyle)
 			[
 				RowWidget.ToSharedRef()
 			];
diff --git a/Engine/Plugins/Experimental/EditorDataStorageFeatures/Source/TedsTableViewer/Public/Widgets/STedsTileViewer.h b/Engine/Plugins/Experimental/EditorDataStorageFeatures/Source/TedsTableViewer/Public/Widgets/STedsTileViewer.h
--- a/Engine/Plugins/Experimental/EditorDataStorageFeatures/Source/TedsTableViewer/Public/Widgets/STedsTileViewer.h
+++ b/Engine/Plugins/Experimental/EditorDataStorageFeatures/Source/TedsTableViewer/Public/Widgets/STedsTileViewer.h
@@ -34,6 +34,9 @@ namespace UE::Editor::DataStorage
 		SLATE_BEGIN_ARGS(STedsTileViewer)
 			: _WidgetPurpose(TEXT("General.Cell")) // TODO: Have a better "default" widget for tiles
 			, _SelectionMode(ESelectionMode::Type::Single)
+			, _ItemHeight(128)
+			, _ItemWidth(128)
+			, _TileStyle(&FAppStyle::Get().GetWidgetStyle<FTableRowStyle>("TableView.Row"))
 		{
 			
 		}
@@ -41,6 +44,9 @@ namespace UE::Editor::DataStorage
 		// Query Stack that will supply the rows to be displayed
 		SLATE_ARGUMENT(TSharedPtr<QueryStack::IRowNode>, QueryStack)
 
+		// Optional list of columns to give TEDS UI when creating the tile widget
+		SLATE_ARGUMENT(TArray<TWeakObjectPtr<const UScriptStruct>>, Columns)
+
 		// The widget purpose to use to create the widgets
 		SLATE_ARGUMENT(FName, WidgetPurpose)
 		
@@ -49,6 +55,12 @@ namespace UE::Editor::DataStorage
 			
 		// The selection mode for the table viewer (single/multi etc)
 		SLATE_ARGUMENT(ESelectionMode::Type, SelectionMode)
+			
+		SLATE_ATTRIBUTE( float, ItemHeight )
+
+		SLATE_ATTRIBUTE( float, ItemWidth )
+			
+		SLATE_STYLE_ARGUMENT( FTableRowStyle, TileStyle )
 
 		SLATE_END_ARGS()
 
@@ -102,5 +114,11 @@ namespace UE::Editor::DataStorage
 
 		// The widget purpose used to generate tiles
 		FName WidgetPurpose;
+
+		// Optional columns to use when matching the tile widget
+		TArray<TWeakObjectPtr<const UScriptStruct>> Columns;
+
+		// Style name to use for the rows (tiles)
+		const FTableRowStyle* TableRowStyle = nullptr;
 	};
 }
diff --git a/Engine/Source/Runtime/SlateCore/Private/Styling/SlateStyleSet.cpp b/Engine/Source/Runtime/SlateCore/Private/Styling/SlateStyleSet.cpp
--- a/Engine/Source/Runtime/SlateCore/Private/Styling/SlateStyleSet.cpp
+++ b/Engine/Source/Runtime/SlateCore/Private/Styling/SlateStyleSet.cpp
@@ -541,6 +541,10 @@ TSet<FName> FSlateStyleSet::GetStyleKeys() const
 
 }
 
+const TSet<FName>& FSlateStyleSet::GetWidgetStyleNames() const
+{
+	return KnownWidgetStyles;
+}
 
 void FSlateStyleSet::Log(ISlateStyle::EStyleMessageSeverity Severity, const FText& Message) const
 {
diff --git a/Engine/Source/Runtime/SlateCore/Public/Styling/ISlateStyle.h b/Engine/Source/Runtime/SlateCore/Public/Styling/ISlateStyle.h
--- a/Engine/Source/Runtime/SlateCore/Public/Styling/ISlateStyle.h
+++ b/Engine/Source/Runtime/SlateCore/Public/Styling/ISlateStyle.h
@@ -64,6 +64,18 @@ public:
 		const bool bWarnIfNotFound = false;
 		return GetWidgetStyleInternal( WidgetStyleType::TypeName, Join( PropertyName, Specifier ), nullptr, bWarnIfNotFound) != nullptr;
 	}
+	
+	const FSlateWidgetStyle* GetWidgetStyle( const FName& WidgetStyleTypeName, FName PropertyName, const ANSICHAR* Specifier = nullptr, const FSlateWidgetStyle* DefaultValue = nullptr, bool bWarnIfNotFound = true ) const 
+	{
+		const FSlateWidgetStyle* ExistingStyle = GetWidgetStyleInternal(WidgetStyleTypeName, Join( PropertyName, Specifier ), DefaultValue, bWarnIfNotFound);
+
+		if ( ExistingStyle == nullptr )
+		{
+			return DefaultValue;
+		}
+
+		return ExistingStyle ? ExistingStyle : DefaultValue;
+	}
 
 	/**
 	 * @param PropertyName   Name of the property to get
@@ -165,6 +177,12 @@ public:
 	}
 
 	virtual TSet<FName> GetStyleKeys() const = 0;
+
+	virtual const TSet<FName>& GetWidgetStyleNames() const
+	{
+		static TSet<FName> Dummy;
+		return Dummy;
+	}
 protected:
 
 	/** 
diff --git a/Engine/Source/Runtime/SlateCore/Public/Styling/SlateStyle.h b/Engine/Source/Runtime/SlateCore/Public/Styling/SlateStyle.h
--- a/Engine/Source/Runtime/SlateCore/Public/Styling/SlateStyle.h
+++ b/Engine/Source/Runtime/SlateCore/Public/Styling/SlateStyle.h
@@ -86,6 +86,7 @@ public:
 	{
 		LLM_SCOPE_BYTAG(UI_Style);
 		WidgetStyleValues.Add(PropertyName, MakeShared<DefinitionType>(InStyleDefintion));
+		KnownWidgetStyles.Add(DefinitionType::TypeName);
 	}
 
 	/**
@@ -212,7 +213,10 @@ protected:
 
 	SLATECORE_API virtual void LogUnusedBrushResources();
 
-	SLATECORE_API TSet<FName> GetStyleKeys() const;
+	SLATECORE_API virtual TSet<FName> GetStyleKeys() const override;
+
+	SLATECORE_API virtual const TSet<FName>& GetWidgetStyleNames() const override;
+
 protected:
 
 	SLATECORE_API bool IsBrushFromFile(const FString& FilePath, const FSlateBrush* Brush);
@@ -227,6 +231,7 @@ protected:
 	FString CoreContentRootDir;
 
 	TMap< FName, TSharedRef< struct FSlateWidgetStyle > > WidgetStyleValues;
+	TSet<FName> KnownWidgetStyles;
 
 	/** float property storage. */
 	TMap< FName, float > FloatValues;
