[TEDS] Added a Query Stack node that converts queries to row handle arrays.

This change also has a small behavior change to queries that only have a context. Instead of running them repeatedly for each row, they now only get called once. A query callback with this signature will have to do any work to access rows internally so it fits better to not attempt to do a row iteration.

diff --git a/Engine/Plugins/Experimental/EditorDataStorageFeatures/Source/TedsQueryStack/Private/TedsRowQueryResultsNode.cpp b/Engine/Plugins/Experimental/EditorDataStorageFeatures/Source/TedsQueryStack/Private/TedsRowQueryResultsNode.cpp
new file mode 100644
--- /dev/null
+++ b/Engine/Plugins/Experimental/EditorDataStorageFeatures/Source/TedsQueryStack/Private/TedsRowQueryResultsNode.cpp
@@ -0,0 +1,92 @@
+// Copyright Epic Games, Inc. All Rights Reserved.
+
+#include "TedsRowQueryResultsNode.h"
+
+#include "Elements/Framework/TypedElementQueryBuilder.h"
+#include "Elements/Interfaces/TypedElementDataStorageInterface.h"
+
+namespace UE::Editor::DataStorage::QueryStack
+{
+	FRowQueryResultsNode::FRowQueryResultsNode(IEditorDataStorageProvider& Storage, TSharedPtr<IQueryNode> InQueryNode, ESyncFlags SyncFlags)
+		: QueryNode(MoveTemp(InQueryNode))
+		, Storage(Storage)
+		, SyncFlags(SyncFlags)
+	{
+		Refresh();
+		QueryRevision = QueryNode->GetRevision();
+	}
+	
+	INode::RevisionId FRowQueryResultsNode::GetRevision() const
+	{
+		return Revision;
+	}
+
+	void FRowQueryResultsNode::Refresh()
+	{
+		if (EnumHasAllFlags(SyncFlags, ESyncFlags::IncrementWhenDifferent))
+		{
+			FRowHandleArray NewRows;
+			RefreshInternal(NewRows);
+			if (NewRows.Num() != Rows.Num())
+			{
+				Rows = MoveTemp(NewRows);
+				Revision++;
+			}
+			else
+			{
+				Rows.Sort();
+				NewRows.Sort();
+				FRowHandleArrayView NewRowView = NewRows.GetRows();
+				if (FMemory::Memcmp(NewRowView.GetData(), Rows.GetRows().GetData(), NewRowView.NumBytes()) != 0)
+				{
+					Rows = MoveTemp(NewRows);
+					Revision++;
+				}
+			}
+		}
+		else
+		{
+			Rows.Empty();
+			RefreshInternal(Rows);
+			Revision++;
+		}
+	}
+
+	void FRowQueryResultsNode::Update()
+	{
+		QueryNode->Update();
+
+		if (EnumHasAllFlags(SyncFlags, ESyncFlags::RefreshOnUpdate) ||
+			(EnumHasAllFlags(SyncFlags, ESyncFlags::RefreshOnQueryChange) && QueryRevision != QueryNode->GetRevision()))
+		{
+			Refresh();
+			QueryRevision = QueryNode->GetRevision();
+		}
+	}
+
+	FRowHandleArrayView FRowQueryResultsNode::GetRows() const
+	{
+		return Rows.GetRows();
+	}
+
+	FRowHandleArray* FRowQueryResultsNode::GetMutableRows()
+	{
+		return &Rows;
+	}
+
+	void FRowQueryResultsNode::RefreshInternal(FRowHandleArray& TargetRows)
+	{
+		using namespace UE::Editor::DataStorage::Queries;
+
+		FQueryResult Result = Storage.RunQuery(QueryNode->GetQuery()); // This is optimized to only collect the number of rows.
+		TargetRows.Reserve(Result.Count);
+
+		Storage.RunQuery(QueryNode->GetQuery(), 
+			EDirectQueryExecutionFlags::AllowBoundQueries | EDirectQueryExecutionFlags::IgnoreActivationCount,
+			CreateDirectQueryCallbackBinding(
+				[&TargetRows](IDirectQueryContext& Context)
+				{
+					TargetRows.Append(FRowHandleArrayView(Context.GetRowHandles(), FRowHandleArrayView::EFlags::IsUnique));
+				}));
+	}
+} // namespace UE::Editor::DataStorage::QueryStack
diff --git a/Engine/Plugins/Experimental/EditorDataStorageFeatures/Source/TedsQueryStack/Public/TedsRowQueryResultsNode.h b/Engine/Plugins/Experimental/EditorDataStorageFeatures/Source/TedsQueryStack/Public/TedsRowQueryResultsNode.h
new file mode 100644
--- /dev/null
+++ b/Engine/Plugins/Experimental/EditorDataStorageFeatures/Source/TedsQueryStack/Public/TedsRowQueryResultsNode.h
@@ -0,0 +1,62 @@
+// Copyright Epic Games, Inc. All Rights Reserved.
+
+#pragma once
+
+#include "Containers/ContainersFwd.h"
+#include "Elements/Framework/TypedElementRowHandleArray.h"
+#include "TedsQueryStackInterfaces.h"
+#include "Misc/EnumClassFlags.h"
+
+class IEditorDataStorageProvider;
+
+namespace UE::Editor::DataStorage::QueryStack
+{
+	/**
+	 * Used to convert an IQueryNode into an IRowNode by extracting the rows the query references optionally during update.
+	 * This node is cheap to setup, but has diminishing returns when the number of rows increases and when updates happen frequently
+	 * as it has to fully extract all rows from TEDS whenever Refresh is called. Use on small tables when the number of calls to 
+	 * Refresh can be minimized.
+	 */
+	class FRowQueryResultsNode : public IRowNode
+	{
+	public:
+		enum class ESyncFlags
+		{
+			None = 0,
+			/** 
+			 * Update the row list whenever "Update" is called. This can take several milliseconds for large lists. It's
+			 * recommended to use a monitoring node to detect changes for large lists.
+			 */
+			RefreshOnUpdate = 1 << 0,
+			/** Update the row list whenever the parent query changes. */
+			RefreshOnQueryChange = 1 << 1,
+			/** 
+			 * Compares the previous row list with the current and only updates the revision if there are differences.
+			 * This requires additional sorting and comparing, making the node substantially more expensive.
+			 */
+			IncrementWhenDifferent = 1 << 2
+		};
+
+		TEDSQUERYSTACK_API FRowQueryResultsNode(IEditorDataStorageProvider& Storage, TSharedPtr<IQueryNode> InQueryNode, ESyncFlags SyncFlags);
+		virtual ~FRowQueryResultsNode() override = default;
+
+		TEDSQUERYSTACK_API void Refresh();
+
+		TEDSQUERYSTACK_API virtual RevisionId GetRevision() const override;
+		TEDSQUERYSTACK_API virtual void Update() override;
+		TEDSQUERYSTACK_API virtual FRowHandleArrayView GetRows() const override;
+		TEDSQUERYSTACK_API virtual FRowHandleArray* GetMutableRows() override;
+
+	private:
+		void RefreshInternal(FRowHandleArray& TargetRows);
+
+		FRowHandleArray Rows;
+		TSharedPtr<IQueryNode> QueryNode;
+		IEditorDataStorageProvider& Storage;
+		RevisionId QueryRevision = 0;
+		RevisionId Revision = 0;
+		ESyncFlags SyncFlags = ESyncFlags::None;
+	};
+
+	ENUM_CLASS_FLAGS(FRowQueryResultsNode::ESyncFlags);
+} // namespace UE::Editor::DataStorage::QueryStack
diff --git a/Engine/Source/Runtime/TypedElementFramework/Public/Elements/Framework/TypedElementQueryBuilder.inl b/Engine/Source/Runtime/TypedElementFramework/Public/Elements/Framework/TypedElementQueryBuilder.inl
--- a/Engine/Source/Runtime/TypedElementFramework/Public/Elements/Framework/TypedElementQueryBuilder.inl
+++ b/Engine/Source/Runtime/TypedElementFramework/Public/Elements/Framework/TypedElementQueryBuilder.inl
@@ -585,11 +585,7 @@ namespace UE::Editor::DataStorage::Queries
 				}
 				else
 				{
-					const uint32 RowCount = Context.GetRowCount();
-					for (uint32 Index = 0; Index < RowCount; ++Index)
-					{
-						Caller(this->ContextWrapper);
-					}
+					Caller(this->ContextWrapper);
 				}
 			}
 		};
