TEDSStyling: Fix warning for uninit value in TEDS column

diff --git a/Engine/Plugins/Experimental/EditorDataStorageFeatures/Source/TedsStyling/Public/Columns/TedsStylingColumns.h b/Engine/Plugins/Experimental/EditorDataStorageFeatures/Source/TedsStyling/Public/Columns/TedsStylingColumns.h
--- a/Engine/Plugins/Experimental/EditorDataStorageFeatures/Source/TedsStyling/Public/Columns/TedsStylingColumns.h
+++ b/Engine/Plugins/Experimental/EditorDataStorageFeatures/Source/TedsStyling/Public/Columns/TedsStylingColumns.h
@@ -52,7 +52,7 @@ struct FIconSizeColumn final : public FEditorDataStorageColumn
 {
 	GENERATED_BODY()
 
-	UPROPERTY(meta = (Searchable))
+	UPROPERTY(meta = (Searchable, IgnoreForMemberInitializationTest))
 	FVector2D IconSize;
 };
 
