TEDSStyling: Remove searchable metadata from a property that isn't supported by universal search to fix ensure

diff --git a/Engine/Plugins/Experimental/EditorDataStorageFeatures/Source/TedsStyling/Public/Columns/TedsStylingColumns.h b/Engine/Plugins/Experimental/EditorDataStorageFeatures/Source/TedsStyling/Public/Columns/TedsStylingColumns.h
--- a/Engine/Plugins/Experimental/EditorDataStorageFeatures/Source/TedsStyling/Public/Columns/TedsStylingColumns.h
+++ b/Engine/Plugins/Experimental/EditorDataStorageFeatures/Source/TedsStyling/Public/Columns/TedsStylingColumns.h
@@ -52,7 +52,7 @@ struct FIconSizeColumn final : public FEditorDataStorageColumn
 {
 	GENERATED_BODY()
 
-	UPROPERTY(meta = (Searchable, IgnoreForMemberInitializationTest))
+	UPROPERTY(meta = (IgnoreForMemberInitializationTest))
 	FVector2D IconSize;
 };
 
