RowReferenceWidget: Update to use FSimpleWidgetConstructor and FAttributeBinder to fix broken queries that were operating on all rows with a row reference (instead of just widget rows belonging to RowReferenceWidget)

diff --git a/Engine/Plugins/Experimental/EditorDataStorageFeatures/Source/TedsDebugger/Private/Widgets/RowReferenceWidget.cpp b/Engine/Plugins/Experimental/EditorDataStorageFeatures/Source/TedsDebugger/Private/Widgets/RowReferenceWidget.cpp
--- a/Engine/Plugins/Experimental/EditorDataStorageFeatures/Source/TedsDebugger/Private/Widgets/RowReferenceWidget.cpp
+++ b/Engine/Plugins/Experimental/EditorDataStorageFeatures/Source/TedsDebugger/Private/Widgets/RowReferenceWidget.cpp
@@ -12,6 +12,7 @@
 #include "Modules/ModuleManager.h"
 #include "TedsDebuggerModule.h"
 #include "Elements/Columns/TypedElementSlateWidgetColumns.h"
+#include "Elements/Framework/TypedElementAttributeBinding.h"
 #include "Widgets/Input/SHyperlink.h"
 #include "Widgets/Layout/SBox.h"
 
@@ -19,7 +20,7 @@
 
 namespace UE::Editor::DataStorage::Debug::Private
 {
-	void OnNavigateHyperlink(const IEditorDataStorageProvider* DataStorage,  RowHandle TargetRowHandle, RowHandle UiRowHandle)
+	void OnNavigateHyperlink(IEditorDataStorageProvider* DataStorage,  RowHandle TargetRowHandle, RowHandle UiRowHandle)
 	{
 		const FTedsOutlinerColumn* TedsOutlinerColumn = DataStorage->GetColumn<FTedsOutlinerColumn>(UiRowHandle);
 
@@ -48,46 +49,6 @@ namespace UE::Editor::DataStorage::Debug::Private
 			return;
 		}
 	}
-	
-	void CreateInternalWidget(const TWeakPtr<SWidget>& InWidget, RowHandle UiRow, RowHandle TargetRow)
-	{
-		const TSharedPtr<SWidget> Widget = InWidget.Pin();
-
-		if(!Widget)
-		{
-			return;
-		}
-		
-		checkf(Widget->GetType() == SBox::StaticWidgetClass().GetWidgetType(),
-			TEXT("Stored widget with FTypedElementLabelWidgetConstructor doesn't match type %s, but was a %s."),
-			*(SBox::StaticWidgetClass().GetWidgetType().ToString()),
-			*(Widget->GetTypeAsString()));
-
-		SBox* WidgetInstance = static_cast<SBox*>(Widget.Get());
-		WidgetInstance->SetContent(SNullWidget::NullWidget);
-
-		const IEditorDataStorageProvider* DataStorage = GetMutableDataStorageFeature<IEditorDataStorageProvider>(StorageFeatureName);
-
-		// We only navigate to row references that have a label column
-		if (const FTypedElementLabelColumn* LabelColumn = DataStorage->GetColumn<FTypedElementLabelColumn>(TargetRow))
-		{
-			FNumberFormattingOptions NumberFormattingOptions;
-			NumberFormattingOptions.SetUseGrouping(false);
-
-			const FText Text = FText::AsNumber(TargetRow, &NumberFormattingOptions);
-			const FText TooltipText = FText::FromString(LabelColumn->Label);
-
-			TSharedRef<SWidget> HyperlinkWidget =
-				SNew(SHyperlink)
-					.Text(Text)
-					.Style(FAppStyle::Get(), "Common.GotoBlueprintHyperlink")
-					.ToolTipText(TooltipText)
-					.OnNavigate(FSimpleDelegate::CreateStatic(&OnNavigateHyperlink, DataStorage, TargetRow, UiRow));
-			
-			WidgetInstance->SetContent(HyperlinkWidget);
-		}
-	}
-
 
 } // namespace UE::Editor::DataStorage::Debug::Private
 
@@ -104,73 +65,43 @@ void URowReferenceWidgetFactory::RegisterWidgetConstructors(IEditorDataStoragePr
 		TColumn<FTypedElementRowReferenceColumn>());
 }
 
-void URowReferenceWidgetFactory::RegisterQueries(IEditorDataStorageProvider& DataStorage)
-{
-	using namespace UE::Editor::DataStorage::Queries;
-	
-	const QueryHandle UpdateRowReferenceWidget = DataStorage.RegisterQuery(
-		Select()
-			.ReadOnly<FTypedElementRowReferenceColumn>()
-		.Where()
-			.Any<FTypedElementSyncFromWorldTag, FTypedElementSyncBackToWorldTag>()
-		.Compile());
-
-	DataStorage.RegisterQuery(
-		Select(
-			TEXT("Sync row reference to widget"),
-			FProcessor(EQueryTickPhase::FrameEnd, DataStorage.GetQueryTickGroupName(EQueryTickGroups::SyncWidgets))
-				.SetExecutionMode(EExecutionMode::GameThread),
-			[](
-				IQueryContext& Context,
-				RowHandle UiRowHandle,
-				FTypedElementSlateWidgetReferenceColumn& Widget,
-				const FTypedElementRowReferenceColumn& Target)
-			{
-				Context.RunSubquery(0, Target.Row, CreateSubqueryCallbackBinding(
-					[&Widget, UiRowHandle](const FTypedElementRowReferenceColumn& Target)
-					{
-						Debug::Private::CreateInternalWidget(Widget.Widget, UiRowHandle, Target.Row);
-					}));
-			})
-		.DependsOn()
-			.SubQuery( UpdateRowReferenceWidget )
-		.Compile());
-}
-
 FRowReferenceWidgetConstructor::FRowReferenceWidgetConstructor()
 	: Super(FRowReferenceWidgetConstructor::StaticStruct())
 {
 }
 
-TSharedPtr<SWidget> FRowReferenceWidgetConstructor::CreateWidget(const UE::Editor::DataStorage::FMetaDataView& Arguments)
-{
-	return SNew(SBox)
-			.HAlign(HAlign_Left)
-			.VAlign(VAlign_Center);
-}
-
-bool FRowReferenceWidgetConstructor::FinalizeWidget(IEditorDataStorageProvider* DataStorage, IEditorDataStorageUiProvider* DataStorageUi, UE::Editor::DataStorage::RowHandle Row, const TSharedPtr<SWidget>& Widget)
+TSharedPtr<SWidget> FRowReferenceWidgetConstructor::CreateWidget(IEditorDataStorageProvider* DataStorage, IEditorDataStorageUiProvider* DataStorageUi,
+	UE::Editor::DataStorage::RowHandle TargetRow, UE::Editor::DataStorage::RowHandle WidgetRow,
+	const UE::Editor::DataStorage::FMetaDataView& Arguments)
 {
 	using namespace UE::Editor::DataStorage;
-	checkf(Widget, TEXT("Referenced widget is not valid. A constructed widget may not have been cleaned up. This can "
-		"also happen if this processor is running in the same phase as the processors responsible for cleaning up old "
-		"references."));
-
-	// The actual row we want to view in the widget
-	RowHandle TargetRowReference = InvalidRowHandle;
-
-	// The target row for which this widget was created
-	const RowHandle TargetRow = DataStorage->GetColumn<FTypedElementRowReferenceColumn>(Row)->Row;
-
-	// Check if the target row has a row reference column, if so we want to view the row in there.
-	if(const FTypedElementRowReferenceColumn* RowReferenceColumn = DataStorage->GetColumn<FTypedElementRowReferenceColumn>(TargetRow))
-	{
-		TargetRowReference = RowReferenceColumn->Row;
-	}
 	
-	Debug::Private::CreateInternalWidget(Widget, Row, TargetRowReference);
-
-	return true;
+	FAttributeBinder Binder(TargetRow, DataStorage);
+	
+	FNumberFormattingOptions NumberFormattingOptions;
+	NumberFormattingOptions.SetUseGrouping(false);
+	
+	return SNew(SBox)
+			.HAlign(HAlign_Left)
+			.VAlign(VAlign_Center)
+		[
+			SNew(SHyperlink)
+					.Text(Binder.BindData(&FTypedElementRowReferenceColumn::Row, [NumberFormattingOptions] (const RowHandle& Row)
+					{
+						return FText::AsNumber(Row, &NumberFormattingOptions);
+					}))
+					.Style(FAppStyle::Get(), "Common.GotoBlueprintHyperlink")
+					.ToolTipText(Binder.BindData(&FTypedElementRowReferenceColumn::Row, [DataStorage] (const RowHandle& Row)
+					{
+						if (FTypedElementLabelColumn* LabelColumn = DataStorage->GetColumn<FTypedElementLabelColumn>(Row))
+						{
+							return FText::FromString(LabelColumn->Label);
+						}
+
+						return FText::GetEmpty();
+					}))
+					.OnNavigate(FSimpleDelegate::CreateStatic(&Debug::Private::OnNavigateHyperlink, DataStorage, TargetRow, WidgetRow))
+		];
 }
 
 #undef LOCTEXT_NAMESPACE
\ No newline at end of file
diff --git a/Engine/Plugins/Experimental/EditorDataStorageFeatures/Source/TedsDebugger/Private/Widgets/RowReferenceWidget.h b/Engine/Plugins/Experimental/EditorDataStorageFeatures/Source/TedsDebugger/Private/Widgets/RowReferenceWidget.h
--- a/Engine/Plugins/Experimental/EditorDataStorageFeatures/Source/TedsDebugger/Private/Widgets/RowReferenceWidget.h
+++ b/Engine/Plugins/Experimental/EditorDataStorageFeatures/Source/TedsDebugger/Private/Widgets/RowReferenceWidget.h
@@ -23,8 +23,6 @@ public:
 
 	void RegisterWidgetConstructors(IEditorDataStorageProvider& DataStorage,
 		IEditorDataStorageUiProvider& DataStorageUi) const override;
-	
-	void RegisterQueries(IEditorDataStorageProvider& DataStorage) override;
 };
 
 USTRUCT()
@@ -36,8 +34,10 @@ public:
 	FRowReferenceWidgetConstructor();
 	~FRowReferenceWidgetConstructor() override = default;
 
-protected:
-	TSharedPtr<SWidget> CreateWidget(const UE::Editor::DataStorage::FMetaDataView& Arguments) override;
-	bool FinalizeWidget(IEditorDataStorageProvider* DataStorage, IEditorDataStorageUiProvider* DataStorageUi,
-		UE::Editor::DataStorage::RowHandle Row, const TSharedPtr<SWidget>& Widget) override;
+	virtual TSharedPtr<SWidget> CreateWidget(
+		IEditorDataStorageProvider* DataStorage,
+		IEditorDataStorageUiProvider* DataStorageUi,
+		UE::Editor::DataStorage::RowHandle TargetRow,
+		UE::Editor::DataStorage::RowHandle WidgetRow,
+		const UE::Editor::DataStorage::FMetaDataView& Arguments) override;
 };
\ No newline at end of file
