[TEDS] Added a Query Stack node that monitors for row adds/removes to keep a row list up to date.

The new Query Stack node takes either a fixed list of columns to monitor or extracts them from a provided query. It then setups up add and remove observers to collect changes that it will apply to the list of rows once per update. If a row node is provided as a parent this node will try to update the parents row array, otherwise it will keep a local array.
In order to support monitoring the Row Handle Array had to be extended with support for removing rows.

[RN] [TEDS] Added a Query Stack node that monitors for row adds/removes to keep a row list up to date.

diff --git a/Engine/Plugins/Experimental/EditorDataStorageFeatures/Source/TedsQueryStack/Private/TedsRowMonitorNode.cpp b/Engine/Plugins/Experimental/EditorDataStorageFeatures/Source/TedsQueryStack/Private/TedsRowMonitorNode.cpp
new file mode 100644
--- /dev/null
+++ b/Engine/Plugins/Experimental/EditorDataStorageFeatures/Source/TedsQueryStack/Private/TedsRowMonitorNode.cpp
@@ -0,0 +1,307 @@
+// Copyright Epic Games, Inc. All Rights Reserved.
+
+#include "TedsRowMonitorNode.h"
+
+#include "Elements/Framework/TypedElementQueryBuilder.h"
+#include "Elements/Interfaces/TypedElementDataStorageInterface.h"
+
+namespace UE::Editor::DataStorage::QueryStack
+{
+	FRowMonitorNode::FRowMonitorNode(
+		IEditorDataStorageProvider& Storage, TSharedPtr<IRowNode> ParentRow, TSharedPtr<IQueryNode> QueryNode)
+		: QueryNode(MoveTemp(QueryNode))
+		, ParentRow(MoveTemp(ParentRow))
+		, Storage(Storage)
+	{
+		UpdateRows();
+		UpdateColumnsFromQuery();
+		UpdateMonitoredColumns();
+	}
+
+	FRowMonitorNode::FRowMonitorNode(
+		IEditorDataStorageProvider& Storage, TSharedPtr<IRowNode> ParentRow, TArray<TObjectPtr<const UScriptStruct>> Columns)
+		: MonitoredColumns(MoveTemp(Columns))
+		, ParentRow(MoveTemp(ParentRow))
+		, Storage(Storage)
+	{
+		UpdateRows();
+		UpdateMonitoredColumns();
+	}
+
+	FRowMonitorNode::FRowMonitorNode(
+		IEditorDataStorageProvider& Storage, TSharedPtr<IQueryNode> QueryNode, TSharedPtr<IRowNode> ParentRow, 
+		TArray<TObjectPtr<const UScriptStruct>> MonitoredColumns)
+		: MonitoredColumns(MoveTemp(MonitoredColumns))
+		, QueryNode(MoveTemp(QueryNode))
+		, ParentRow(MoveTemp(ParentRow))
+		, Storage(Storage)
+		, bFixedColumns(true)
+	{
+		UpdateRows();
+		UpdateColumnsFromQuery();
+		UpdateMonitoredColumns();
+	}
+
+	FRowMonitorNode::FRowMonitorNode(IEditorDataStorageProvider& Storage, TSharedPtr<IQueryNode> QueryNode)
+		: QueryNode(MoveTemp(QueryNode))
+		, Storage(Storage)
+	{
+		UpdateColumnsFromQuery();
+		UpdateMonitoredColumns();
+	}
+	
+	FRowMonitorNode::FRowMonitorNode(IEditorDataStorageProvider& Storage, TArray<TObjectPtr<const UScriptStruct>> Columns)
+		: MonitoredColumns(MoveTemp(Columns))
+		, Storage(Storage)
+	{
+		UpdateMonitoredColumns();
+	}
+
+	FRowMonitorNode::FRowMonitorNode(
+		IEditorDataStorageProvider& Storage, TSharedPtr<IQueryNode> QueryNode, TArray<TObjectPtr<const UScriptStruct>> MonitoredColumns)
+		: MonitoredColumns(MoveTemp(MonitoredColumns))
+		, QueryNode(MoveTemp(QueryNode))
+		, Storage(Storage)
+		, bFixedColumns(true)
+	{
+		UpdateColumnsFromQuery();
+		UpdateMonitoredColumns();
+	}
+
+	FRowMonitorNode::~FRowMonitorNode()
+	{
+		for (QueryHandle Observer : Observers)
+		{
+			Storage.UnregisterQuery(Observer);
+		}
+	}
+
+	INode::RevisionId FRowMonitorNode::GetRevision() const
+	{
+		return Revision;
+	}
+
+	void FRowMonitorNode::Update()
+	{
+		if (QueryNode)
+		{
+			QueryNode->Update();
+			if (QueryNode->GetRevision() != QueryRevision)
+			{
+				UpdateColumnsFromQuery();
+				UpdateMonitoredColumns();
+				Revision++;
+			}
+		}
+		
+		bool MergeChanges = true;
+		if (ParentRow)
+		{
+			ParentRow->Update();
+			if (ParentRow->GetRevision() != ParentRevision)
+			{
+				UpdateRows();
+				Revision++;
+				MergeChanges = false;
+			}
+		}
+		
+		if (MergeChanges && (!AddedRows.IsEmpty() || !RemovedRows.IsEmpty()))
+		{
+			AddedRows.Sort();
+			AddedRows.MakeUnique();
+
+			RemovedRows.Sort();
+			RemovedRows.MakeUnique();
+
+			FRowHandleArray& TargetRows = ResolveRows();
+			TargetRows.SortedMerge(AddedRows);
+			TargetRows.Remove(RemovedRows.GetRows());
+
+			Revision++;
+		}
+		AddedRows.Empty();
+		RemovedRows.Empty();
+	}
+
+	FRowHandleArrayView FRowMonitorNode::GetRows() const
+	{
+		return ResolveRows().GetRows();
+	}
+
+	FRowHandleArray* FRowMonitorNode::GetMutableRows()
+	{
+		return &ResolveRows();
+	}
+
+	const FRowHandleArray& FRowMonitorNode::ResolveRows() const
+	{
+		return const_cast<FRowMonitorNode*>(this)->ResolveRows();
+	}
+	FRowHandleArray& FRowMonitorNode::ResolveRows()
+	{
+		FRowHandleArray* ParentRowArray = ParentRow ? ParentRow->GetMutableRows() : nullptr;
+		return ParentRowArray == nullptr ? Rows : *ParentRowArray;
+	}
+
+	void FRowMonitorNode::UpdateColumnsFromQuery()
+	{
+		using namespace UE::Editor::DataStorage::Queries;
+
+		if (!bFixedColumns)
+		{
+			if (QueryNode)
+			{
+				TSet<TObjectPtr<const UScriptStruct>> LocalColumns;
+				
+				const FQueryDescription& Query = Storage.GetQueryDescription(QueryNode->GetQuery());
+				
+				const FConditions* Conditions = Query.Conditions.GetPtrOrNull();
+				TConstArrayView<TWeakObjectPtr<const UScriptStruct>> ComplexConditionColumns;
+				if (Conditions && !Conditions->IsEmpty())
+				{
+					ComplexConditionColumns = Conditions->GetColumns();
+				}
+				
+				LocalColumns.Reserve(Query.SelectionTypes.Num() + Query.ConditionOperators.Num() + ComplexConditionColumns.Num());
+
+				// Collect all columns that are selected for access.
+				for (const TWeakObjectPtr<const UScriptStruct>& SelectionColumn : Query.SelectionTypes)
+				{
+					if (const UScriptStruct* SelectionColumnPtr = SelectionColumn.Get())
+					{
+						LocalColumns.Add(SelectionColumnPtr);
+					}
+				}
+
+				// Collect all columns that are used in simple conditions.
+				int32 NumConditions = Query.ConditionOperators.Num();
+				for (int32 Index = 0; Index < NumConditions; ++Index)
+				{
+					FQueryDescription::EOperatorType Type = Query.ConditionTypes[Index];
+					if (Type == FQueryDescription::EOperatorType::SimpleAll ||
+						Type == FQueryDescription::EOperatorType::SimpleAny)
+					{
+						if (const UScriptStruct* ConditionColumn = Query.ConditionOperators[Index].Type.Get())
+						{
+							LocalColumns.Add(ConditionColumn);
+						}
+					}
+				}
+
+				// Collect all columns that are used for complex conditions.
+				for (const TWeakObjectPtr<const UScriptStruct>& Column : ComplexConditionColumns)
+				{
+					if (const UScriptStruct* ColumnType = Column.Get())
+					{
+						LocalColumns.Add(ColumnType);
+					}
+				}
+
+				MonitoredColumns.Empty();
+				MonitoredColumns.Reserve(LocalColumns.Num());
+				for (TObjectPtr<const UScriptStruct>& Column : LocalColumns)
+				{
+					MonitoredColumns.Add(Column);
+				}
+				
+				QueryRevision = QueryNode->GetRevision();
+			}
+			else
+			{
+				MonitoredColumns.Empty();
+				QueryRevision = 0;
+			}
+		}
+	}
+
+	void FRowMonitorNode::UpdateRows()
+	{
+		if (FRowHandleArray* ParentRowArray = ParentRow->GetMutableRows(); ParentRowArray == nullptr)
+		{
+			Rows.Empty();
+			Rows.Append(ParentRow->GetRows());
+		}
+		ParentRevision = ParentRow->GetRevision();
+	}
+
+	void FRowMonitorNode::UpdateMonitoredColumns()
+	{
+		using namespace UE::Editor::DataStorage::Queries;
+
+		static FName Name = FName("QueryStack Row Monitor node");
+
+		for (QueryHandle Observer : Observers)
+		{
+			Storage.UnregisterQuery(Observer);
+		}
+
+		Observers.Reserve(MonitoredColumns.Num() * 2); // Add double, one for the OnAdd and one for the OnRemove.
+
+		if (QueryNode)
+		{
+			auto OnAdd = [this](const FQueryDescription&, IQueryContext& Context)
+				{
+					AddedRows.Append(FRowHandleArrayView(Context.GetRowHandles(), FRowHandleArrayView::EFlags::IsUnique));
+				};
+			auto OnRemove = [this](const FQueryDescription&, IQueryContext& Context)
+				{
+					RemovedRows.Append(FRowHandleArrayView(Context.GetRowHandles(), FRowHandleArrayView::EFlags::IsUnique));
+				};
+			
+			const FQueryDescription& QueryBase = Storage.GetQueryDescription(QueryNode->GetQuery());
+
+			for (const TObjectPtr<const UScriptStruct>& Column : MonitoredColumns)
+			{
+				FString OnAddName = TEXT("QueryStack Row Monitor node: OnAdd - ");
+				Column->GetFName().AppendString(OnAddName);
+				FString OnRemoveName = TEXT("QueryStack Row Monitor node: OnRemove - ");
+				Column->GetFName().AppendString(OnRemoveName);
+
+				FQueryDescription OnAddObserver = QueryBase;
+				OnAddObserver.Callback.Name = FName(OnAddName);
+				OnAddObserver.Callback.Type = EQueryCallbackType::ObserveAdd;
+				OnAddObserver.Callback.ExecutionMode = EExecutionMode::GameThread;
+				OnAddObserver.Callback.Function = OnAdd;
+				OnAddObserver.Callback.MonitoredType = Column;
+				Observers.Add(Storage.RegisterQuery(MoveTemp(OnAddObserver)));
+
+				FQueryDescription OnRemoveObserver = QueryBase;
+				OnRemoveObserver.Callback.Name = FName(OnRemoveName);
+				OnRemoveObserver.Callback.Type = EQueryCallbackType::ObserveRemove;
+				OnRemoveObserver.Callback.ExecutionMode = EExecutionMode::GameThread;
+				OnRemoveObserver.Callback.Function = OnRemove;
+				OnRemoveObserver.Callback.MonitoredType = Column;
+				Observers.Add(Storage.RegisterQuery(MoveTemp(OnRemoveObserver)));
+			}
+		}
+		else
+		{
+			auto OnAdd = [this](IQueryContext& Context)
+				{
+					AddedRows.Append(FRowHandleArrayView(Context.GetRowHandles(), FRowHandleArrayView::EFlags::IsUnique));
+				};
+			auto OnRemove = [this](IQueryContext& Context)
+				{
+					RemovedRows.Append(FRowHandleArrayView(Context.GetRowHandles(), FRowHandleArrayView::EFlags::IsUnique));
+				};
+
+			for (const TObjectPtr<const UScriptStruct>& Column : MonitoredColumns)
+			{
+				FString OnAddName = TEXT("QueryStack Row Monitor node: OnAdd - ");
+				Column->GetFName().AppendString(OnAddName);
+				FString OnRemoveName = TEXT("QueryStack Row Monitor node: OnRemove - ");
+				Column->GetFName().AppendString(OnRemoveName);
+
+				Observers.Add(Storage.RegisterQuery(Select(FName(OnAddName),
+					FObserver(FObserver::EEvent::Add, Column.Get())
+						.SetExecutionMode(EExecutionMode::GameThread), OnAdd)
+					.Compile()));
+				Observers.Add(Storage.RegisterQuery(Select(FName(OnRemoveName),
+					FObserver(FObserver::EEvent::Remove, Column.Get())
+						.SetExecutionMode(EExecutionMode::GameThread), OnRemove)
+					.Compile()));
+			}
+		}
+	}
+} // namespace UE::Editor::DataStorage::QueryStack
diff --git a/Engine/Plugins/Experimental/EditorDataStorageFeatures/Source/TedsQueryStack/Public/TedsRowMonitorNode.h b/Engine/Plugins/Experimental/EditorDataStorageFeatures/Source/TedsQueryStack/Public/TedsRowMonitorNode.h
new file mode 100644
--- /dev/null
+++ b/Engine/Plugins/Experimental/EditorDataStorageFeatures/Source/TedsQueryStack/Public/TedsRowMonitorNode.h
@@ -0,0 +1,63 @@
+// Copyright Epic Games, Inc. All Rights Reserved.
+
+#pragma once
+
+#include "Containers/ContainersFwd.h"
+#include "Elements/Framework/TypedElementRowHandleArray.h"
+#include "TedsQueryStackInterfaces.h"
+#include "UObject/ObjectPtr.h"
+
+class IEditorDataStorageProvider;
+
+namespace UE::Editor::DataStorage::QueryStack
+{
+	/**
+	 * Monitors tables for the addition and removal of one or more column types and updates the internal status if a change is detected.
+	 */
+	class FRowMonitorNode : public IRowNode
+	{
+	public:
+		TEDSQUERYSTACK_API FRowMonitorNode(
+			IEditorDataStorageProvider& Storage, TSharedPtr<IRowNode> ParentRow, TSharedPtr<IQueryNode> QueryNode);
+		TEDSQUERYSTACK_API FRowMonitorNode(
+			IEditorDataStorageProvider& Storage, TSharedPtr<IRowNode> ParentRow, TArray<TObjectPtr<const UScriptStruct>> Columns);
+		TEDSQUERYSTACK_API FRowMonitorNode(
+			IEditorDataStorageProvider& Storage, TSharedPtr<IQueryNode> QueryNode, TSharedPtr<IRowNode> ParentRow, 
+			TArray<TObjectPtr<const UScriptStruct>> MonitoredColumns);
+		TEDSQUERYSTACK_API FRowMonitorNode(
+			IEditorDataStorageProvider& Storage, TSharedPtr<IQueryNode> QueryNode);
+		TEDSQUERYSTACK_API FRowMonitorNode(
+			IEditorDataStorageProvider& Storage, TArray<TObjectPtr<const UScriptStruct>> Columns);
+		TEDSQUERYSTACK_API FRowMonitorNode(
+			IEditorDataStorageProvider& Storage, TSharedPtr<IQueryNode> QueryNode, TArray<TObjectPtr<const UScriptStruct>> MonitoredColumns);
+
+		TEDSQUERYSTACK_API virtual ~FRowMonitorNode() override;
+
+		TEDSQUERYSTACK_API virtual RevisionId GetRevision() const override;
+		TEDSQUERYSTACK_API virtual void Update() override;
+		TEDSQUERYSTACK_API virtual FRowHandleArrayView GetRows() const override;
+		TEDSQUERYSTACK_API virtual FRowHandleArray* GetMutableRows() override;
+
+	private:
+		const FRowHandleArray& ResolveRows() const;
+		FRowHandleArray& ResolveRows();
+		void UpdateColumnsFromQuery();
+		void UpdateRows();
+		void UpdateMonitoredColumns();
+
+		FRowHandleArray Rows;
+		FRowHandleArray AddedRows;
+		FRowHandleArray RemovedRows;
+
+		TArray<QueryHandle> Observers;
+		TArray<TObjectPtr<const UScriptStruct>> MonitoredColumns;
+		
+		TSharedPtr<IQueryNode> QueryNode;
+		TSharedPtr<IRowNode> ParentRow;
+		IEditorDataStorageProvider& Storage;
+		RevisionId QueryRevision = 0;
+		RevisionId ParentRevision = 0;
+		RevisionId Revision = 0;
+		bool bFixedColumns = false;
+	};
+} // namespace UE::Editor::DataStorage::QueryStack
diff --git a/Engine/Plugins/Experimental/EditorDataStorageFeatures/Source/TedsQueryStack/TedsQueryStack.Build.cs b/Engine/Plugins/Experimental/EditorDataStorageFeatures/Source/TedsQueryStack/TedsQueryStack.Build.cs
--- a/Engine/Plugins/Experimental/EditorDataStorageFeatures/Source/TedsQueryStack/TedsQueryStack.Build.cs
+++ b/Engine/Plugins/Experimental/EditorDataStorageFeatures/Source/TedsQueryStack/TedsQueryStack.Build.cs
@@ -20,6 +20,7 @@ public class TedsQueryStack : ModuleRules
 				new string[]
 				{
 					"Core",
+					"CoreUObject",
 					"TypedElementFramework",
 				});
 
diff --git a/Engine/Source/Runtime/TypedElementFramework/Private/Elements/Framework/TypedElementRowHandleArray.cpp b/Engine/Source/Runtime/TypedElementFramework/Private/Elements/Framework/TypedElementRowHandleArray.cpp
--- a/Engine/Source/Runtime/TypedElementFramework/Private/Elements/Framework/TypedElementRowHandleArray.cpp
+++ b/Engine/Source/Runtime/TypedElementFramework/Private/Elements/Framework/TypedElementRowHandleArray.cpp
@@ -5,6 +5,8 @@
 #include <algorithm>
 #include <iterator>
 #include <limits>
+#include "Algo/BinarySearch.h"
+#include "Algo/IndexOf.h"
 #include "HAL/UnrealMemory.h"
 #include "Templates/Sorting.h"
 
@@ -73,6 +75,122 @@ namespace UE::Editor::DataStorage
 		}
 	}
 
+	void FRowHandleArray::Remove(RowHandle Row)
+	{
+		if (bIsSorted)
+		{
+			if (int32 Index = Algo::BinarySearch(Rows, Row); Index >= 0)
+			{
+				Rows.RemoveAt(Index, EAllowShrinking::No);
+			}
+		}
+		else
+		{
+			Rows.RemoveSingleSwap(Row, EAllowShrinking::No);
+		}
+	}
+
+	void FRowHandleArray::Remove(TConstArrayView<RowHandle> RowsToRemove)
+	{
+		if (bIsSorted)
+		{
+			for (RowHandle Row : RowsToRemove)
+			{
+				if (int32 Index = Algo::BinarySearch(Rows, Row); Index >= 0)
+				{
+					Rows.RemoveAt(Index, EAllowShrinking::No);
+				}
+			}
+		}
+		else
+		{
+			for (RowHandle Row : RowsToRemove)
+			{
+				Rows.RemoveSingleSwap(Row, EAllowShrinking::No);
+			}
+		}
+	}
+
+	void FRowHandleArray::Remove(FRowHandleArrayView RowsToRemove)
+	{
+		if (bIsSorted && RowsToRemove.IsSorted() && !RowsToRemove.IsEmpty())
+		{
+			const RowHandle* RowsToRemoveIt = RowsToRemove.begin();
+			const RowHandle* RowsToRemoveEnd = RowsToRemove.end();
+			
+			// Search for the first row that can be deleted.
+			int32 StartIndex = -1;
+			do
+			{
+				StartIndex = Algo::BinarySearch(Rows, *RowsToRemoveIt);
+				RowsToRemoveIt++;
+			} while (StartIndex < 0 && RowsToRemoveIt < RowsToRemoveEnd);
+
+			if (StartIndex >= 0)
+			{
+				int32 DeletedCount = 1;
+				RowHandle* RowsInsert = Rows.GetData() + StartIndex;
+				RowHandle* RowsIt = RowsInsert + 1;
+				RowHandle* RowsEnd = Rows.GetData() + Rows.Num();
+
+				// Copy down the remaining list, increasing the gap whenever a new
+				// match in the rows to be removed has been found.
+				for (; RowsIt < RowsEnd; ++RowsIt)
+				{
+					// Catch delete pointer up, in case it encountered a row that wasn't in the array
+					// or a duplicate entry.
+					if (*RowsToRemoveIt < *RowsIt)
+					{
+						while (*RowsToRemoveIt < *RowsIt)
+						{
+							++RowsToRemoveIt;
+						}
+						if (RowsToRemoveIt == RowsToRemoveEnd)	
+						{
+							break;
+						}
+					}
+
+					// Skip if there's a match with in the list of rows to delete.
+					if (*RowsToRemoveIt == *RowsIt)
+					{
+						RowsToRemoveIt++;
+						DeletedCount++;
+
+						if (RowsToRemoveIt == RowsToRemoveEnd)
+						{
+							++RowsIt;
+							break;
+						}
+					}
+					else
+					{
+						*RowsInsert++ = *RowsIt;
+					}
+				}
+
+				// Copy the remainder down without any checks as there are no more rows to remove.
+				for (; RowsIt < RowsEnd; ++RowsIt)
+				{
+					*RowsInsert++ = *RowsIt;
+				}
+
+				// Remove the remaining tail of empty slots.
+				int32 RemoveStartIndex = static_cast<int32>(RowsInsert - Rows.GetData());
+				Rows.RemoveAt(RemoveStartIndex, DeletedCount, EAllowShrinking::No);
+			}
+		}
+		else
+		{
+			Remove(TConstArrayView<RowHandle>(RowsToRemove.GetData(), RowsToRemove.Num()));
+		}
+	}
+
+	bool FRowHandleArray::Contains(RowHandle Row) const
+	{
+		return Find(Row) >= 0;
+	}
+
 	void FRowHandleArray::Reset(int32 NewSize)
 	{
 		Rows.Reset(NewSize);
@@ -355,4 +473,11 @@ namespace UE::Editor::DataStorage
 		}
 		bIsUnique = false; // Assume there was overlap between the existing and merged list.
 	}
+
+	int32 FRowHandleArray::Find(RowHandle Row) const
+	{
+		return bIsSorted
+			? Algo::BinarySearch(Rows, Row)
+			: Algo::IndexOf(Rows, Row);
+	}
 } // namespace UE::Editor::DataStorage
diff --git a/Engine/Source/Runtime/TypedElementFramework/Public/Elements/Framework/TypedElementRowHandleArray.h b/Engine/Source/Runtime/TypedElementFramework/Public/Elements/Framework/TypedElementRowHandleArray.h
--- a/Engine/Source/Runtime/TypedElementFramework/Public/Elements/Framework/TypedElementRowHandleArray.h
+++ b/Engine/Source/Runtime/TypedElementFramework/Public/Elements/Framework/TypedElementRowHandleArray.h
@@ -36,12 +36,38 @@ namespace UE::Editor::DataStorage
 		/** Append the provided list of rows to the end of the array. This will make the array unsorted and not unique. */
 		TYPEDELEMENTFRAMEWORK_API void Append(TConstArrayView<RowHandle> AdditionalRows);
 
-		/** 
+		/**
 		 * Append the provided list of rows to the end of the array. If the current array and the additional rows are both sorted then
 		 * this function will try to append the additional rows and maintain the sorted status and if applicable the uniqueness state,
 		 * otherwise the array will be marked as unsorted and no longer unique.
 		 */
 		TYPEDELEMENTFRAMEWORK_API void Append(FRowHandleArrayView AdditionalRows);
+		
+		/**
+		 * Removes the provided row. If the array is sorted, a binary search will be used to find the row, but requires moving all
+		 * subsequent rows down. If the array is not sorted a linear search is required but the last row will be moved into the removed
+		 * location for quicker removal.
+		 */
+		TYPEDELEMENTFRAMEWORK_API void Remove(RowHandle Row);
+		/**
+		 * Removes the provided list of rows. If the array is sorted, a binary search will be used to find the row, but requires moving all
+		 * subsequent rows down. If the array is not sorted a linear search is required but the last row will be moved into the removed
+		 * location for quicker removals.
+		 */
+		TYPEDELEMENTFRAMEWORK_API void Remove(TConstArrayView<RowHandle> RowsToRemove);
+		/**
+		 * Removes the provided list of rows. If the array is sorted, a binary search will be used to find the row, but requires moving all
+		 * subsequent rows down. If the array is not sorted a linear search is required but the last row will be moved into the removed
+		 * location for quicker removals. If both the array and list of rows to remove are sorted than an optimized version can run that
+		 * quickly finds the first row to remove and then removes the other rows as it's moving the remainder down.
+		 */
+		TYPEDELEMENTFRAMEWORK_API void Remove(FRowHandleArrayView RowsToRemove);
+
+		/**
+		 * Checks if the provided row is stored in this array. If the array is sorted a faster binary search is used, otherwise a linear
+		 * search is required.
+		 */
+		TYPEDELEMENTFRAMEWORK_API bool Contains(RowHandle Row) const;
 
 		/**
 		 * Removes all row handles, but keeps the allocated memory around. If the provided size is larger than the number of rows
@@ -89,6 +115,7 @@ namespace UE::Editor::DataStorage
 
 	private:
 		void SortedMergeInternal(FRowHandleArrayView AddedRows);
+		int32 Find(RowHandle Row) const;
 
 		TArray<RowHandle> Rows;
 		bool bIsSorted = true;
diff --git a/Engine/Source/Runtime/TypedElementFramework/Tests/Elements/Framework/TypedElementRowHandleArrayTests.cpp b/Engine/Source/Runtime/TypedElementFramework/Tests/Elements/Framework/TypedElementRowHandleArrayTests.cpp
--- a/Engine/Source/Runtime/TypedElementFramework/Tests/Elements/Framework/TypedElementRowHandleArrayTests.cpp
+++ b/Engine/Source/Runtime/TypedElementFramework/Tests/Elements/Framework/TypedElementRowHandleArrayTests.cpp
@@ -166,6 +166,153 @@ TEST_CASE_NAMED(FRowHandleArray_Tests, "Editor::DataStorage::Row Handle Array (F
 		CHECK_MESSAGE(TEXT("Arrays should not be unique."), !Array.IsUnique());
 	}
 
+	SECTION("Remove row from sorted array.")
+	{
+		RowHandle ArrayData[] = { 1, 2, 3, 4, 5 };
+		FRowHandleArray Array(FRowHandleArrayView(ArrayData,
+			FRowHandleArrayView::EFlags::IsSorted | FRowHandleArrayView::EFlags::IsUnique));
+
+		Array.Remove(4);
+
+		CHECK_MESSAGE(TEXT("Not the expected number of rows."), Array.Num() == 4);
+		CHECK_MESSAGE(TEXT("Not the correct row removed."), Compare(Array, { 1, 2, 3, 5 }));
+	}
+
+	SECTION("Remove row from unsorted array.")
+	{
+		RowHandle ArrayData[] = { 3, 1, 5, 2, 4 };
+		FRowHandleArray Array(FRowHandleArrayView(ArrayData, FRowHandleArrayView::EFlags::IsUnique));
+
+		Array.Remove(2);
+
+		CHECK_MESSAGE(TEXT("Not the expected number of rows."), Array.Num() == 4);
+		CHECK_MESSAGE(TEXT("Not the correct row removed."), Compare(Array, { 3, 1, 5, 4 }));
+	}
+
+	SECTION("Remove rows from sorted array.")
+	{
+		RowHandle ArrayData[] = { 1, 2, 3, 4, 5 };
+		FRowHandleArray Array(FRowHandleArrayView(ArrayData,
+			FRowHandleArrayView::EFlags::IsSorted | FRowHandleArrayView::EFlags::IsUnique));
+
+		Array.Remove({2, 6, 5});
+
+		CHECK_MESSAGE(TEXT("Not the expected number of rows."), Array.Num() == 3);
+		CHECK_MESSAGE(TEXT("Not the correct row removed."), Compare(Array, { 1, 3, 4 }));
+	}
+
+	SECTION("Remove rows from unsorted array.")
+	{
+		RowHandle ArrayData[] = { 3, 1, 5, 2, 4 };
+		FRowHandleArray Array(FRowHandleArrayView(ArrayData, FRowHandleArrayView::EFlags::IsUnique));
+
+		Array.Remove({ 1, 6, 4 });
+
+		CHECK_MESSAGE(TEXT("Not the expected number of rows."), Array.Num() == 3);
+		Array.Sort(); // Sort because the order of the rows will be jumbled after the call to remove.
+		CHECK_MESSAGE(TEXT("Not the correct row removed."), Compare(Array, { 2, 3, 5 }));
+	}
+
+	SECTION("Remove rows from sorted array with a sorted list.")
+	{
+		RowHandle ArrayData[] = { 1, 2, 3, 4, 5 };
+		FRowHandleArray Array(FRowHandleArrayView(ArrayData,
+			FRowHandleArrayView::EFlags::IsSorted | FRowHandleArrayView::EFlags::IsUnique));
+
+		RowHandle RemoveRows[] = { 2, 5 };
+		FRowHandleArrayView RemoveRowsView(RemoveRows, FRowHandleArrayView::EFlags::IsSorted | FRowHandleArrayView::EFlags::IsUnique);
+
+		Array.Remove(RemoveRowsView);
+
+		CHECK_MESSAGE(TEXT("Not the expected number of rows."), Array.Num() == 3);
+		CHECK_MESSAGE(TEXT("Not the correct row removed."), Compare(Array, { 1, 3, 4 }));
+	}
+
+	SECTION("Remove rows from sorted array with a sorted list and additional row.")
+	{
+		RowHandle ArrayData[] = { 1, 2, 3, 4, 5 };
+		FRowHandleArray Array(FRowHandleArrayView(ArrayData,
+			FRowHandleArrayView::EFlags::IsSorted | FRowHandleArrayView::EFlags::IsUnique));
+
+		RowHandle RemoveRows[] = { 2, 5, 6 };
+		FRowHandleArrayView RemoveRowsView (RemoveRows, FRowHandleArrayView::EFlags::IsSorted | FRowHandleArrayView::EFlags::IsUnique);
+
+		Array.Remove(RemoveRowsView);
+
+		CHECK_MESSAGE(TEXT("Not the expected number of rows."), Array.Num() == 3);
+		CHECK_MESSAGE(TEXT("Not the correct row removed."), Compare(Array, { 1, 3, 4 }));
+	}
+
+	SECTION("Remove rows from sorted array with an empty list.")
+	{
+		RowHandle ArrayData[] = { 1, 2, 3, 4, 5 };
+		FRowHandleArray Array(FRowHandleArrayView(ArrayData,
+			FRowHandleArrayView::EFlags::IsSorted | FRowHandleArrayView::EFlags::IsUnique));
+
+		FRowHandleArray RemoveRows;
+		FRowHandleArrayView RemoveRowsView(RemoveRows.GetRows());
+
+		Array.Remove(RemoveRowsView);
+
+		CHECK_MESSAGE(TEXT("Not the expected number of rows."), Array.Num() == 5);
+		CHECK_MESSAGE(TEXT("Not the correct row removed."), Compare(Array, { 1, 2, 3, 4, 5 }));
+	}
+
+	SECTION("Remove rows from sorted array with a sorted list and duplicates.")
+	{
+		RowHandle ArrayData[] = { 1, 2, 3, 4, 5 };
+		FRowHandleArray Array(FRowHandleArrayView(ArrayData,
+			FRowHandleArrayView::EFlags::IsSorted | FRowHandleArrayView::EFlags::IsUnique));
+
+		RowHandle RemoveRows[] = { 2, 2, 3, 5, 5, 5 };
+		FRowHandleArrayView RemoveRowsView(RemoveRows, FRowHandleArrayView::EFlags::IsSorted | FRowHandleArrayView::EFlags::IsUnique);
+
+		Array.Remove(RemoveRowsView);
+
+		CHECK_MESSAGE(TEXT("Not the expected number of rows."), Array.Num() == 2);
+		CHECK_MESSAGE(TEXT("Not the correct row removed."), Compare(Array, { 1, 4 }));
+	}
+
+	SECTION("Remove rows from sorted array with a sorted list without any matching entries.")
+	{
+		RowHandle ArrayData[] = { 1, 2, 3, 4, 5 };
+		FRowHandleArray Array(FRowHandleArrayView(ArrayData,
+			FRowHandleArrayView::EFlags::IsSorted | FRowHandleArrayView::EFlags::IsUnique));
+
+		RowHandle RemoveRows[] = { 6, 7, 8, 9 };
+		FRowHandleArrayView RemoveRowsView(RemoveRows, FRowHandleArrayView::EFlags::IsSorted | FRowHandleArrayView::EFlags::IsUnique);
+
+		Array.Remove(RemoveRowsView);
+
+		CHECK_MESSAGE(TEXT("Not the expected number of rows."), Array.Num() == 5);
+		CHECK_MESSAGE(TEXT("Not the correct row removed."), Compare(Array, { 1, 2, 3, 4, 5 }));
+	}
+
+	SECTION("Contains row sorted")
+	{
+		RowHandle ArrayData[] = { 1, 2, 3 };
+		FRowHandleArray Array(FRowHandleArrayView(ArrayData,
+			FRowHandleArrayView::EFlags::IsSorted | FRowHandleArrayView::EFlags::IsUnique));
+
+		CHECK_MESSAGE(TEXT("Failed to find row."), Array.Contains(2));
+	}
+
+	SECTION("Contains row unsorted")
+	{
+		RowHandle ArrayData[] = { 3, 1, 2 };
+		FRowHandleArray Array(FRowHandleArrayView(ArrayData, FRowHandleArrayView::EFlags::IsUnique));
+
+		CHECK_MESSAGE(TEXT("Failed to find row."), Array.Contains(1));
+	}
+
+	SECTION("No contains row")
+	{
+		RowHandle ArrayData[] = { 3, 1, 2 };
+		FRowHandleArray Array(FRowHandleArrayView(ArrayData, FRowHandleArrayView::EFlags::IsUnique));
+
+		CHECK_MESSAGE(TEXT("Found a non-existing row."), !Array.Contains(4));
+	}
+
 	SECTION("Sort array low bits")
 	{
 		FRowHandleArray Array(FRowHandleArrayView({ 3, 1, 2 }, FRowHandleArrayView::EFlags::IsUnique));
@@ -582,5 +729,4 @@ TEST_CASE_NAMED(FRowHandleArray_Tests, "Editor::DataStorage::Row Handle Array (F
 		CHECK_MESSAGE(TEXT("Sorted merge not correct."), Compare(Array, TestRange));
 	}
 }
-
 #endif // #if WITH_TESTS
