Update EditorDataStorageFeatures plugin to use GetAssetAccessSpecifier() instead of PKG_NotExternallyReferenceable and add support for EpicInternal

#rnx

diff --git a/Engine/Plugins/Experimental/EditorDataStorageFeatures/Source/TedsAssetData/Private/CB/TedsAssetDataCBDataSource.cpp b/Engine/Plugins/Experimental/EditorDataStorageFeatures/Source/TedsAssetData/Private/CB/TedsAssetDataCBDataSource.cpp
--- a/Engine/Plugins/Experimental/EditorDataStorageFeatures/Source/TedsAssetData/Private/CB/TedsAssetDataCBDataSource.cpp
+++ b/Engine/Plugins/Experimental/EditorDataStorageFeatures/Source/TedsAssetData/Private/CB/TedsAssetDataCBDataSource.cpp
@@ -229,14 +229,18 @@ void FTedsAssetDataCBDataSource::AddAssetDataColumns(UE::Editor::DataStorage::IQ
 {
 	// For now just add the columns one by one but this should be rework to work in batch
 	// Not optimized at all but we would like to have the data in sooner for testing purposes.
-	if (AssetData.HasAnyPackageFlags(PKG_NotExternallyReferenceable))
+	const EAssetAccessSpecifier AssetAccessSpecifier = AssetData.GetAssetAccessSpecifier();
+	if (AssetAccessSpecifier == EAssetAccessSpecifier::Public)
 	{
-		// Private Asset
-		Context.AddColumns<FAssetTag, FPrivateAssetTag>(Row);
+		Context.AddColumns<FAssetTag, FPublicAssetTag>(Row);
+	}
+	else if (AssetAccessSpecifier == EAssetAccessSpecifier::EpicInternal)
+	{
+		Context.AddColumns<FAssetTag, FEpicInternalAssetTag>(Row);
 	}
 	else
 	{
-		Context.AddColumns<FAssetTag, FPublicAssetTag>(Row);
+		Context.AddColumns<FAssetTag, FPrivateAssetTag>(Row);
 	}
 
 	if (OptionalPackageData)
diff --git a/Engine/Plugins/Experimental/EditorDataStorageFeatures/Source/TedsAssetData/Private/Debug/TedsAssetDataDebugCommands.cpp b/Engine/Plugins/Experimental/EditorDataStorageFeatures/Source/TedsAssetData/Private/Debug/TedsAssetDataDebugCommands.cpp
--- a/Engine/Plugins/Experimental/EditorDataStorageFeatures/Source/TedsAssetData/Private/Debug/TedsAssetDataDebugCommands.cpp
+++ b/Engine/Plugins/Experimental/EditorDataStorageFeatures/Source/TedsAssetData/Private/Debug/TedsAssetDataDebugCommands.cpp
@@ -206,8 +206,19 @@ namespace UE::Editor::DataStorage::Debug::Private
 			ClassTypeInfoColumn->DiskSize = FMath::RandRange(1024, 32768);
 		}
 
-		// Randomly make this a public or a private asset
-		DataStorage->AddColumn(Row, FMath::RandRange(0, 1) == 0 ? FPrivateAssetTag::StaticStruct() : FPublicAssetTag::StaticStruct());
+		// Randomly make this a public, EpicInternal or a private asset
+		int32 AssetAccessSpecifierChoiceInt = FMath::RandRange(0, 2);
+		UScriptStruct* AssetAccessTag = FPrivateAssetTag::StaticStruct();
+		if (AssetAccessSpecifierChoiceInt == 0)
+		{
+			AssetAccessTag = FPublicAssetTag::StaticStruct();
+		}
+		else if (AssetAccessSpecifierChoiceInt == 1)
+		{
+			AssetAccessTag = FEpicInternalAssetTag::StaticStruct();
+		}
+
+		DataStorage->AddColumn(Row, AssetAccessTag);
 	}
 	
 	FAutoConsoleCommand CreateDebugAssetRowsCommand(
diff --git a/Engine/Plugins/Experimental/EditorDataStorageFeatures/Source/TedsAssetData/Public/TedsAssetDataColumns.h b/Engine/Plugins/Experimental/EditorDataStorageFeatures/Source/TedsAssetData/Public/TedsAssetDataColumns.h
--- a/Engine/Plugins/Experimental/EditorDataStorageFeatures/Source/TedsAssetData/Public/TedsAssetDataColumns.h
+++ b/Engine/Plugins/Experimental/EditorDataStorageFeatures/Source/TedsAssetData/Public/TedsAssetDataColumns.h
@@ -97,6 +97,13 @@ struct FPrivateAssetTag final : public FEditorDataStorageTag
 	GENERATED_BODY()
 };
 
+// Tag to identify assets with epic internal visibility
+USTRUCT(meta = (DisplayName = "Epic Internal Asset"))
+struct FEpicInternalAssetTag final : public FEditorDataStorageTag
+{
+	GENERATED_BODY()
+};
+
 // Tag to identify assets with public visibility
 USTRUCT(meta = (DisplayName = "Public Asset"))
 struct FPublicAssetTag final : public FEditorDataStorageTag
