TedsSettings : Add a GetSettingsSectionFromRow method to get settings section details for a given row.

diff --git a/Engine/Plugins/Experimental/EditorDataStorageFeatures/Source/TedsSettings/Private/TedsSettingsEditorSubsystem.cpp b/Engine/Plugins/Experimental/EditorDataStorageFeatures/Source/TedsSettings/Private/TedsSettingsEditorSubsystem.cpp
--- a/Engine/Plugins/Experimental/EditorDataStorageFeatures/Source/TedsSettings/Private/TedsSettingsEditorSubsystem.cpp
+++ b/Engine/Plugins/Experimental/EditorDataStorageFeatures/Source/TedsSettings/Private/TedsSettingsEditorSubsystem.cpp
@@ -36,6 +36,11 @@ UE::Editor::DataStorage::RowHandle UTedsSettingsEditorSubsystem::FindOrAddSettin
 	return SettingsManager->FindOrAddSettingsSection(ContainerName, CategoryName, SectionName);
 }
 
+bool UTedsSettingsEditorSubsystem::GetSettingsSectionFromRow(UE::Editor::DataStorage::RowHandle Row, FName& OutContainerName, FName& OutCategoryName, FName& OutSectionName)
+{
+	return SettingsManager->GetSettingsSectionFromRow(Row, OutContainerName, OutCategoryName, OutSectionName);
+}
+
 void UTedsSettingsEditorSubsystem::Initialize(FSubsystemCollectionBase& Collection)
 {
 	Super::Initialize(Collection);
diff --git a/Engine/Plugins/Experimental/EditorDataStorageFeatures/Source/TedsSettings/Private/TedsSettingsManager.cpp b/Engine/Plugins/Experimental/EditorDataStorageFeatures/Source/TedsSettings/Private/TedsSettingsManager.cpp
--- a/Engine/Plugins/Experimental/EditorDataStorageFeatures/Source/TedsSettings/Private/TedsSettingsManager.cpp
+++ b/Engine/Plugins/Experimental/EditorDataStorageFeatures/Source/TedsSettings/Private/TedsSettingsManager.cpp
@@ -117,8 +117,7 @@ UE::Editor::DataStorage::RowHandle FTedsSettingsManager::FindOrAddSettingsSectio
 
 	if (bIsInitialized)
 	{
-		IEditorDataStorageProvider* DataStorage = GetMutableDataStorageFeature<IEditorDataStorageProvider>(StorageFeatureName);
-		if (ensure(DataStorage))
+		if (IEditorDataStorageProvider* DataStorage = GetMutableDataStorageFeature<IEditorDataStorageProvider>(StorageFeatureName))
 		{
 			// Ensure that the SettingsInactiveSectionTable is registered since it is possible for OnDataStorageInterfacesSet to be called
 			// for a delegate bound in another module before it has been called for our bound delegate.
@@ -147,6 +146,35 @@ UE::Editor::DataStorage::RowHandle FTedsSettingsManager::FindOrAddSettingsSectio
 	return InvalidRowHandle;
 }
 
+bool FTedsSettingsManager::GetSettingsSectionFromRow(UE::Editor::DataStorage::RowHandle Row, FName& OutContainerName, FName& OutCategoryName, FName& OutSectionName)
+{
+	using namespace UE::Editor::DataStorage;
+
+	if (bIsInitialized)
+	{
+		if (const IEditorDataStorageProvider* DataStorage = GetDataStorageFeature<IEditorDataStorageProvider>(StorageFeatureName))
+		{
+			if (DataStorage->IsRowAvailable(Row) && (DataStorage->HasColumns<FSettingsSectionTag>(Row) || DataStorage->HasColumns<FSettingsInactiveSectionTag>(Row)))
+			{
+				const FSettingsContainerReferenceColumn* ContainerReferenceColumn = DataStorage->GetColumn<FSettingsContainerReferenceColumn>(Row);
+				const FSettingsCategoryReferenceColumn* CategoryReferenceColumn = DataStorage->GetColumn<FSettingsCategoryReferenceColumn>(Row);
+				const FNameColumn* NameColumn = DataStorage->GetColumn<FNameColumn>(Row);
+
+				if (ContainerReferenceColumn && CategoryReferenceColumn && NameColumn)
+				{
+					OutContainerName = ContainerReferenceColumn->ContainerName;
+					OutCategoryName = CategoryReferenceColumn->CategoryName;
+					OutSectionName = NameColumn->Name;
+
+					return true;
+				}
+			}
+		}
+	}
+
+	return false;
+}
+
 void FTedsSettingsManager::RegisterTables(IEditorDataStorageProvider& DataStorage)
 {
 	if (SettingsContainerTable == UE::Editor::DataStorage::InvalidTableHandle)
diff --git a/Engine/Plugins/Experimental/EditorDataStorageFeatures/Source/TedsSettings/Private/TedsSettingsManager.h b/Engine/Plugins/Experimental/EditorDataStorageFeatures/Source/TedsSettings/Private/TedsSettingsManager.h
--- a/Engine/Plugins/Experimental/EditorDataStorageFeatures/Source/TedsSettings/Private/TedsSettingsManager.h
+++ b/Engine/Plugins/Experimental/EditorDataStorageFeatures/Source/TedsSettings/Private/TedsSettingsManager.h
@@ -23,6 +23,8 @@ public:
 
 	UE::Editor::DataStorage::RowHandle FindOrAddSettingsSection(const FName& ContainerName, const FName& CategoryName, const FName& SectionName);
 
+	bool GetSettingsSectionFromRow(UE::Editor::DataStorage::RowHandle Row, FName& OutContainerName, FName& OutCategoryName, FName& OutSectionName);
+
 private:
 
 	void RegisterTables(IEditorDataStorageProvider& DataStorage);
diff --git a/Engine/Plugins/Experimental/EditorDataStorageFeatures/Source/TedsSettings/Public/TedsSettingsEditorSubsystem.h b/Engine/Plugins/Experimental/EditorDataStorageFeatures/Source/TedsSettings/Public/TedsSettingsEditorSubsystem.h
--- a/Engine/Plugins/Experimental/EditorDataStorageFeatures/Source/TedsSettings/Public/TedsSettingsEditorSubsystem.h
+++ b/Engine/Plugins/Experimental/EditorDataStorageFeatures/Source/TedsSettings/Public/TedsSettingsEditorSubsystem.h
@@ -23,9 +23,21 @@ public:
 	DECLARE_MULTICAST_DELEGATE(FOnEnabledChanged)
 	FOnEnabledChanged& OnEnabledChanged();
 
-	// Finds an existing row (may be active or inactive) or adds a new inactive settings section row if no existing row is found.
+	/** Finds an existing row (may be active or inactive) or adds a new inactive settings section row if no existing row is found. */
 	UE::Editor::DataStorage::RowHandle FindOrAddSettingsSection(const FName& ContainerName, const FName& CategoryName, const FName& SectionName);
 
+	/**
+	 * Gets the settings section details for the given row.
+	 *
+	 * @param	Row					The row from which to read settings details.
+	 * @param	OutContainerName	Reference to the ContainerName for the given row, only set if this method returns true.
+	 * @param	OutCategoryName		Reference to the CategoryName for the given row, only set if this method returns true.
+	 * @param	OutSectionName		Reference to the SectionName for the given row, only set if this method returns true.
+	 *
+	 * @return	True if the settings section details are successfully returned in the out parameters.
+	 */
+	bool GetSettingsSectionFromRow(UE::Editor::DataStorage::RowHandle Row, FName& OutContainerName, FName& OutCategoryName, FName& OutSectionName);
+
 protected: // USubsystem
 
 	virtual void Initialize(FSubsystemCollectionBase& Collection) override;
