OutlinerLabelWidget is now subdivided into two widgets with separate purposes (icon & text).

diff --git a/Engine/Plugins/Experimental/EditorDataStorageFeatures/Source/TedsOutliner/Private/Widgets/OutlinerIconWidget.cpp b/Engine/Plugins/Experimental/EditorDataStorageFeatures/Source/TedsOutliner/Private/Widgets/OutlinerIconWidget.cpp
new file mode 100644
--- /dev/null
+++ b/Engine/Plugins/Experimental/EditorDataStorageFeatures/Source/TedsOutliner/Private/Widgets/OutlinerIconWidget.cpp
@@ -0,0 +1,142 @@
+// Copyright Epic Games, Inc. All Rights Reserved.
+
+#include "Widgets/OutlinerIconWidget.h"
+
+#include "Elements/Columns/TypedElementOverrideColumns.h"
+#include "Elements/Columns/TypedElementTypeInfoColumns.h"
+#include "Elements/Framework/TypedElementAttributeBinding.h"
+#include "TedsTableViewerUtils.h"
+#include "Widgets/Images/SImage.h"
+#include "Widgets/Images/SLayeredImage.h"
+
+#define LOCTEXT_NAMESPACE "OutlinerIconWidget"
+
+namespace UE::OutlinerIconWidget::Local
+{
+	static const FSlateBrush* GetOverrideBadgeFirstLayer(const EOverriddenState& OverriddenState)
+	{
+		switch (OverriddenState)
+		{
+		case EOverriddenState::Added:
+			return FAppStyle::GetBrush("SceneOutliner.OverrideAddedBase");
+
+		case EOverriddenState::AllOverridden:
+			// Not implemented yet
+			break;
+
+		case EOverriddenState::HasOverrides:
+			return FAppStyle::GetBrush("SceneOutliner.OverrideInsideBase");
+
+		case EOverriddenState::NoOverrides:
+			// No icon for no overrides
+			break;
+
+		case EOverriddenState::SubObjectsHasOverrides:
+			// Not implemented yet
+			break;
+		}
+
+		return FAppStyle::GetBrush("NoBrush");
+	}
+
+	static const FSlateBrush* GetOverrideBadgeSecondLayer(const EOverriddenState& OverriddenState)
+	{
+		switch (OverriddenState)
+		{
+		case EOverriddenState::Added:
+			return FAppStyle::GetBrush("SceneOutliner.OverrideAdded");
+
+		case EOverriddenState::AllOverridden:
+			// Not implemented yet
+			break;
+
+		case EOverriddenState::HasOverrides:
+			return FAppStyle::GetBrush("SceneOutliner.OverrideInside");
+
+		case EOverriddenState::NoOverrides:
+			// No icon for no overrides
+			break;
+
+		case EOverriddenState::SubObjectsHasOverrides:
+			// Not implemented yet
+			break;
+		}
+
+		return FAppStyle::GetBrush("NoBrush");
+	}
+
+	static FText GetOverrideTooltip(const EOverriddenState& OverriddenState)
+	{
+		switch (OverriddenState)
+		{
+		case EOverriddenState::Added:
+			return LOCTEXT("OverrideAddedTooltip", "This entity has been added.");
+
+		case EOverriddenState::AllOverridden:
+			// Not implemented yet
+			break;
+
+		case EOverriddenState::HasOverrides:
+			return LOCTEXT("OverrideInsideTooltip", "At least one property or child has an override.");
+
+		case EOverriddenState::NoOverrides:
+			// No icon for no overrides
+			break;
+
+		case EOverriddenState::SubObjectsHasOverrides:
+			// Not implemented yet
+			break;
+		}
+
+		return FText::GetEmpty();
+	}
+}
+
+void UOutlinerIconWidgetFactory::RegisterWidgetConstructors(IEditorDataStorageProvider& DataStorage,
+	IEditorDataStorageUiProvider& DataStorageUi) const
+{
+	using namespace UE::Editor::DataStorage::Queries;
+
+	DataStorageUi.RegisterWidgetFactory<FOutlinerIconWidgetConstructor>(TEXT("SceneOutliner.RowLabel.Icon"),
+		TColumn<FTypedElementClassTypeInfoColumn>() || TColumn<FTypedElementScriptStructTypeInfoColumn>());
+}
+
+FOutlinerIconWidgetConstructor::FOutlinerIconWidgetConstructor()
+	: Super(FOutlinerIconWidgetConstructor::StaticStruct())
+{
+}
+
+TSharedPtr<SWidget> FOutlinerIconWidgetConstructor::CreateWidget(IEditorDataStorageProvider* DataStorage,
+	IEditorDataStorageUiProvider* DataStorageUi, UE::Editor::DataStorage::RowHandle TargetRow,
+	UE::Editor::DataStorage::RowHandle WidgetRow, const UE::Editor::DataStorage::FMetaDataView& Arguments)
+{
+	using namespace UE::Editor::DataStorage;
+	if (!DataStorage->IsRowAvailable(TargetRow))
+	{
+		return SNullWidget::NullWidget;
+	}
+
+	FAttributeBinder Binder(TargetRow, DataStorage);
+
+	TSharedRef<SLayeredImage> LayeredImageWidget = SNew(SLayeredImage)
+		.Image(TableViewerUtils::GetIconForRow(DataStorage, TargetRow))
+		.ToolTipText(Binder.BindData(&FObjectOverrideColumn::OverriddenState, [](const EOverriddenState& OverriddenState)
+			{
+				return UE::OutlinerIconWidget::Local::GetOverrideTooltip(OverriddenState);
+			}))
+		.ColorAndOpacity(FSlateColor::UseForeground());
+
+	LayeredImageWidget->AddLayer(Binder.BindData(&FObjectOverrideColumn::OverriddenState, [](const EOverriddenState& OverriddenState)
+		{
+			return UE::OutlinerIconWidget::Local::GetOverrideBadgeFirstLayer(OverriddenState);
+		}));
+
+	LayeredImageWidget->AddLayer(Binder.BindData(&FObjectOverrideColumn::OverriddenState, [](const EOverriddenState& OverriddenState)
+		{
+			return UE::OutlinerIconWidget::Local::GetOverrideBadgeSecondLayer(OverriddenState);
+		}));
+
+	return LayeredImageWidget;
+}
+
+#undef LOCTEXT_NAMESPACE
diff --git a/Engine/Plugins/Experimental/EditorDataStorageFeatures/Source/TedsOutliner/Private/Widgets/OutlinerLabelWidget.cpp b/Engine/Plugins/Experimental/EditorDataStorageFeatures/Source/TedsOutliner/Private/Widgets/OutlinerLabelWidget.cpp
--- a/Engine/Plugins/Experimental/EditorDataStorageFeatures/Source/TedsOutliner/Private/Widgets/OutlinerLabelWidget.cpp
+++ b/Engine/Plugins/Experimental/EditorDataStorageFeatures/Source/TedsOutliner/Private/Widgets/OutlinerLabelWidget.cpp
@@ -2,236 +2,125 @@
 
 #include "Widgets/OutlinerLabelWidget.h"
 
-#include "ActorEditorUtils.h"
-#include "Elements/Columns/TypedElementCompatibilityColumns.h"
-#include "Elements/Columns/TypedElementLabelColumns.h"
-#include "Elements/Framework/TypedElementAttributeBinding.h"
-#include "TedsTableViewerUtils.h"
-#include "Columns/SlateDelegateColumns.h"
-#include "Elements/Columns/TypedElementMiscColumns.h"
-#include "Elements/Columns/TypedElementOverrideColumns.h"
 #include "Elements/Columns/TypedElementTypeInfoColumns.h"
-#include "Elements/Interfaces/Capabilities/TypedElementUiEditableCapability.h"
-#include "Elements/Interfaces/Capabilities/TypedElementUiStyleOverrideCapability.h"
-#include "Settings/EditorStyleSettings.h"
-#include "Widgets/Images/SImage.h"
-#include "Widgets/Layout/SSpacer.h"
-#include "Widgets/Text/STextBlock.h"
+#include "Elements/Columns/TypedElementLabelColumns.h"
+#include "Layout/Margin.h"
+#include "Widgets/SNullWidget.h"
 #include "Widgets/SBoxPanel.h"
-#include "Widgets/Images/SLayeredImage.h"
-#include "Widgets/Text/SInlineEditableTextBlock.h"
-
-#define LOCTEXT_NAMESPACE "FOutlinerLabelWidgetConstructor"
-
-namespace UE::OutlinerLabelWidget::Local
-{
-	const FSlateBrush* GetOverrideBadgeFirstLayer(const EOverriddenState& OverriddenState)
-	{
-		switch(OverriddenState)
-		{
-		case EOverriddenState::Added:
-			return FAppStyle::GetBrush("SceneOutliner.OverrideAddedBase");
-			
-		case EOverriddenState::AllOverridden:
-			// Not implemented yet
-			break;
-			
-		case EOverriddenState::HasOverrides:
-			return FAppStyle::GetBrush("SceneOutliner.OverrideInsideBase");
-			
-		case EOverriddenState::NoOverrides:
-			// No icon for no overrides
-			break;
-			
-		case EOverriddenState::SubObjectsHasOverrides:
-			// Not implemented yet
-			break;
-		}
-		
-		return FAppStyle::GetBrush("NoBrush");
-	}
-	
-	const FSlateBrush* GetOverrideBadgeSecondLayer(const EOverriddenState& OverriddenState)
-	{
-		switch(OverriddenState)
-		{
-		case EOverriddenState::Added:
-			return FAppStyle::GetBrush("SceneOutliner.OverrideAdded");
-			
-		case EOverriddenState::AllOverridden:
-			// Not implemented yet
-				break;
-			
-		case EOverriddenState::HasOverrides:
-			return FAppStyle::GetBrush("SceneOutliner.OverrideInside");
-			
-		case EOverriddenState::NoOverrides:
-			// No icon for no overrides
-				break;
-			
-		case EOverriddenState::SubObjectsHasOverrides:
-			// Not implemented yet
-				break;
-		}
-
-		return FAppStyle::GetBrush("NoBrush");
-	}
+#include "Widgets/Layout/SSpacer.h"
 
-	FText GetOverrideTooltip(const EOverriddenState& OverriddenState)
-	{
-		switch(OverriddenState)
-		{
-		case EOverriddenState::Added:
-			return LOCTEXT("OverrideAddedTooltip", "This entity has been added.");
-			
-		case EOverriddenState::AllOverridden:
-			// Not implemented yet
-				break;
-			
-		case EOverriddenState::HasOverrides:
-			return LOCTEXT("OverrideInsideTooltip", "At least one property or child has an override.");
-			
-		case EOverriddenState::NoOverrides:
-			// No icon for no overrides
-				break;
-			
-		case EOverriddenState::SubObjectsHasOverrides:
-			// Not implemented yet
-				break;
-		}
+#define LOCTEXT_NAMESPACE "OutlinerLabelWidget"
 
-		return FText::GetEmpty();
-	}
-}
+static const FName IconWidgetPurpose = TEXT("SceneOutliner.RowLabel.Icon");
+static const FName TextWidgetPurpose = TEXT("SceneOutliner.RowLabel.Text");
 
 void UOutlinerLabelWidgetFactory::RegisterWidgetConstructors(IEditorDataStorageProvider& DataStorage,
 	IEditorDataStorageUiProvider& DataStorageUi) const
 {
 	using namespace UE::Editor::DataStorage::Queries;
-	DataStorageUi.RegisterWidgetFactory<FOutlinerLabelWidgetConstructor>(
-		TEXT("SceneOutliner.RowLabel"),
+	DataStorageUi.RegisterWidgetFactory<FOutlinerLabelWidgetConstructor>(TEXT("SceneOutliner.RowLabel"),
 		TColumn<FTypedElementLabelColumn>() && (TColumn<FTypedElementClassTypeInfoColumn>() || TColumn<FTypedElementScriptStructTypeInfoColumn>()));
 }
 
+void UOutlinerLabelWidgetFactory::RegisterWidgetPurposes(IEditorDataStorageUiProvider& DataStorageUi) const
+{
+	DataStorageUi.RegisterWidgetPurpose(IconWidgetPurpose, IEditorDataStorageUiProvider::EPurposeType::UniqueByNameAndColumn,
+		LOCTEXT("IconItemCellWidgetPurpose", "The icon widget to use in cells for the Scene Outliner specific to the Item label column.")
+	);
+
+	DataStorageUi.RegisterWidgetPurpose(TextWidgetPurpose, IEditorDataStorageUiProvider::EPurposeType::UniqueByNameAndColumn,
+		LOCTEXT("TextItemCellWidgetPurpose", "The text widget to use in cells for the Scene Outliner specific to the Item label column.")
+	);
+}
+
+static void GetAllColumns(TArray<TWeakObjectPtr<const UScriptStruct>>& OutColumns, const IEditorDataStorageProvider& DataStorage, UE::Editor::DataStorage::RowHandle Row)
+{
+	OutColumns.Empty();
+	DataStorage.ListColumns(Row, [&OutColumns](const UScriptStruct& ColumnType)
+		{
+			OutColumns.Emplace(&ColumnType);
+			return true;
+		});
+}
+
 FOutlinerLabelWidgetConstructor::FOutlinerLabelWidgetConstructor()
 	: FSimpleWidgetConstructor(StaticStruct())
 {
 }
 
 TSharedPtr<SWidget> FOutlinerLabelWidgetConstructor::CreateWidget(IEditorDataStorageProvider* DataStorage,
-	IEditorDataStorageUiProvider* DataStorageUi, UE::Editor::DataStorage::RowHandle TargetRow, 
+	IEditorDataStorageUiProvider* DataStorageUi, UE::Editor::DataStorage::RowHandle TargetRow,
 	UE::Editor::DataStorage::RowHandle WidgetRow, const UE::Editor::DataStorage::FMetaDataView& Arguments)
 {
-	using namespace UE::Editor::DataStorage;
-	if(DataStorage->IsRowAvailable(TargetRow))
+	if (DataStorage->IsRowAvailable(TargetRow))
 	{
-		FAttributeBinder Binder(TargetRow, DataStorage);
+		TSharedPtr<FTypedElementWidgetConstructor> OutIconWidgetConstructorPtr;
+		TSharedPtr<FTypedElementWidgetConstructor> OutTextWidgetConstructorPtr;
+
+		auto AssignIconWidgetToColumn = [&OutIconWidgetConstructorPtr](TUniquePtr<FTypedElementWidgetConstructor> Constructor, TConstArrayView<TWeakObjectPtr<const UScriptStruct>>)
+			{
+				OutIconWidgetConstructorPtr = TSharedPtr<FTypedElementWidgetConstructor>(Constructor.Release());
+				return false;
+			};
+
+		auto AssignTextWidgetToColumn = [&OutTextWidgetConstructorPtr](TUniquePtr<FTypedElementWidgetConstructor> Constructor, TConstArrayView<TWeakObjectPtr<const UScriptStruct>>)
+			{
+				OutTextWidgetConstructorPtr = TSharedPtr<FTypedElementWidgetConstructor>(Constructor.Release());
+				return false;
+			};
+
+		TArray<TWeakObjectPtr<const UScriptStruct>> Columns;
+		GetAllColumns(Columns, *DataStorage, TargetRow);
+		DataStorageUi->CreateWidgetConstructors(IconWidgetPurpose, IEditorDataStorageUiProvider::EMatchApproach::LongestMatch, Columns, Arguments, AssignIconWidgetToColumn);
+		GetAllColumns(Columns, *DataStorage, TargetRow);
+		DataStorageUi->CreateWidgetConstructors(TextWidgetPurpose, IEditorDataStorageUiProvider::EMatchApproach::LongestMatch, Columns, Arguments, AssignTextWidgetToColumn);
+
+		static const FMargin ColumnItemPadding(8, 0);
+		constexpr float IconNameHorizontalPadding = 8.f;
+
+		TSharedPtr<SWidget> IconWidget = SNullWidget::NullWidget;
+		TSharedPtr<SWidget> TextWidget = SNullWidget::NullWidget;
+
+		if (OutIconWidgetConstructorPtr)
+		{
+			IconWidget = DataStorageUi->ConstructWidget(WidgetRow, *OutIconWidgetConstructorPtr, Arguments);
+		}
 
-		TSharedRef<SLayeredImage> LayeredImageWidget = SNew(SLayeredImage)
-					.Image(TableViewerUtils::GetIconForRow(DataStorage, TargetRow))
-					.ToolTipText(Binder.BindData(&FObjectOverrideColumn::OverriddenState, [](const EOverriddenState& OverriddenState)
-					{
-						return UE::OutlinerLabelWidget::Local::GetOverrideTooltip(OverriddenState);
-					}))
-					.ColorAndOpacity(FSlateColor::UseForeground());
+		if (OutTextWidgetConstructorPtr)
+		{
+			TextWidget = DataStorageUi->ConstructWidget(WidgetRow, *OutTextWidgetConstructorPtr, Arguments);
+		}
 
-		LayeredImageWidget->AddLayer(Binder.BindData(&FObjectOverrideColumn::OverriddenState, [](const EOverriddenState& OverriddenState)
+		if (!IconWidget.IsValid())
 		{
-			return UE::OutlinerLabelWidget::Local::GetOverrideBadgeFirstLayer(OverriddenState);
-		}));
-		
-		LayeredImageWidget->AddLayer(Binder.BindData(&FObjectOverrideColumn::OverriddenState, [](const EOverriddenState& OverriddenState)
+			IconWidget = SNullWidget::NullWidget;
+		}
+
+		if (!TextWidget.IsValid())
 		{
-			return UE::OutlinerLabelWidget::Local::GetOverrideBadgeSecondLayer(OverriddenState);
-		}));
+			TextWidget = SNullWidget::NullWidget;
+		}
 
 		return SNew(SHorizontalBox)
-			+SHorizontalBox::Slot()
-				.AutoWidth()
+			+ SHorizontalBox::Slot()
+			.AutoWidth()
 			[
-				LayeredImageWidget
+				IconWidget.ToSharedRef()
 			]
-			+SHorizontalBox::Slot()
-				.AutoWidth()
+			+ SHorizontalBox::Slot()
+			.AutoWidth()
 			[
 				SNew(SSpacer)
 					.Size(FVector2D(5.0f, 0.0f))
 			]
-			+SHorizontalBox::Slot()
-				.FillWidth(1.0f)
+			+ SHorizontalBox::Slot()
+			.FillWidth(1.0f)
 			[
-				CreateLabel(DataStorage, DataStorageUi, TargetRow, WidgetRow, Arguments)
+				TextWidget.ToSharedRef()
 			];
 	}
-	else
-	{
-		return SNew(STextBlock)
-			.Text(LOCTEXT("MissingRowReferenceColumn", "Unable to retrieve row reference."));
-	}
-	
-	
-}
-
-TSharedRef<SWidget> FOutlinerLabelWidgetConstructor::CreateLabel(IEditorDataStorageProvider* DataStorage,
-	IEditorDataStorageUiProvider* DataStorageUi, UE::Editor::DataStorage::RowHandle TargetRow, 
-	UE::Editor::DataStorage::RowHandle WidgetRow, const UE::Editor::DataStorage::FMetaDataView& Arguments)
-{
-	using namespace UE::Editor::DataStorage;
-	TSharedRef<SWidget> Result = SNullWidget::NullWidget;
-	
-	const bool* IsEditable = Arguments.FindForColumn<FTypedElementLabelColumn>(IsEditableName).TryGetExact<bool>();
-
-	if (IsEditable && *IsEditable)
-	{
-		FAttributeBinder TargetRowBinder(TargetRow, DataStorage);
-		FAttributeBinder WidgetRowBinder(WidgetRow, DataStorage);
-
-		const bool bShouldUseMiddleEllipsis = GetDefault<UEditorStyleSettings>()->bEnableMiddleEllipsis;
-
-		TSharedPtr<SInlineEditableTextBlock> TextBlock = SNew(SInlineEditableTextBlock)
-			.OnTextCommitted_Lambda(
-				[DataStorage, TargetRow](const FText& NewText, ETextCommit::Type CommitInfo)
-				{
-					// This callback happens on the game thread so it's safe to directly call into the data storage.
-					FString NewLabelText = NewText.ToString();
-					if (FTypedElementLabelHashColumn* LabelHashColumn = DataStorage->GetColumn<FTypedElementLabelHashColumn>(TargetRow))
-					{
-						LabelHashColumn->LabelHash = CityHash64(reinterpret_cast<const char*>(*NewLabelText), NewLabelText.Len() * sizeof(**NewLabelText));
-					}
-					if (FTypedElementLabelColumn* LabelColumn = DataStorage->GetColumn<FTypedElementLabelColumn>(TargetRow))
-					{
-						LabelColumn->Label = MoveTemp(NewLabelText);
-					}
-					DataStorage->AddColumn<FTypedElementSyncBackToWorldTag>(TargetRow);
-				})
-			.OnVerifyTextChanged_Lambda([](const FText& Label, FText& ErrorMessage)
-				{
-					// Note: The use of actor specific functionality should be minimized, but this function acts generic enough that the 
-					// use of actor is just in names.
-					return FActorEditorUtils::ValidateActorName(Label, ErrorMessage);
-				})
-			.Text(TargetRowBinder.BindText(&FTypedElementLabelColumn::Label))
-			.ToolTipText(TargetRowBinder.BindText(&FTypedElementLabelColumn::Label))
-			.OverflowPolicy(bShouldUseMiddleEllipsis ? ETextOverflowPolicy::MiddleEllipsis : TOptional<ETextOverflowPolicy>())
-			.IsSelected(WidgetRowBinder.BindEvent(&FExternalWidgetSelectionColumn::IsSelected));
-
-		TextBlock->AddMetadata(MakeShared<TTypedElementUiEditableCapability<SInlineEditableTextBlock>>(*TextBlock));
-		TextBlock->AddMetadata(MakeShared<TTypedElementUiStyleOverrideCapability<SInlineEditableTextBlock>>(*TextBlock));
-		Result = TextBlock.ToSharedRef();
-	}
-	else
-	{
-		FAttributeBinder TargetRowBinder(TargetRow, DataStorage);
-		
-		TSharedPtr<STextBlock> TextBlock = SNew(STextBlock)
-			.IsEnabled(false)
-			.Text(TargetRowBinder.BindText(&FTypedElementLabelColumn::Label))
-			.ToolTipText(TargetRowBinder.BindText(&FTypedElementLabelColumn::Label));
-		
-		Result = TextBlock.ToSharedRef();
-	}
 
-	return Result;
+	return SNullWidget::NullWidget;
 }
 
-#undef LOCTEXT_NAMESPACE //"FOutlinerLabelWidgetConstructor"
+#undef LOCTEXT_NAMESPACE
diff --git a/Engine/Plugins/Experimental/EditorDataStorageFeatures/Source/TedsOutliner/Private/Widgets/OutlinerLabelWidget.h b/Engine/Plugins/Experimental/EditorDataStorageFeatures/Source/TedsOutliner/Private/Widgets/OutlinerLabelWidget.h
--- a/Engine/Plugins/Experimental/EditorDataStorageFeatures/Source/TedsOutliner/Private/Widgets/OutlinerLabelWidget.h
+++ b/Engine/Plugins/Experimental/EditorDataStorageFeatures/Source/TedsOutliner/Private/Widgets/OutlinerLabelWidget.h
@@ -4,13 +4,11 @@
 
 #include "Elements/Interfaces/TypedElementDataStorageFactory.h"
 #include "Elements/Interfaces/TypedElementDataStorageUiInterface.h"
-#include "Internationalization/Text.h"
 #include "UObject/ObjectMacros.h"
 
 #include "OutlinerLabelWidget.generated.h"
 
 class IEditorDataStorageProvider;
-class UScriptStruct;
 
 UCLASS()
 class UOutlinerLabelWidgetFactory : public UEditorDataStorageFactory
@@ -22,6 +20,8 @@ public:
 
 	TEDSOUTLINER_API void RegisterWidgetConstructors(IEditorDataStorageProvider& DataStorage,
 		IEditorDataStorageUiProvider& DataStorageUi) const override;
+
+	TEDSOUTLINER_API void RegisterWidgetPurposes(IEditorDataStorageUiProvider& DataStorageUi) const override;
 };
 
 // Label widget for the Scene Outliner that shows an icon (with optional override information) + a text label
@@ -40,10 +40,4 @@ public:
 		UE::Editor::DataStorage::RowHandle TargetRow,
 		UE::Editor::DataStorage::RowHandle WidgetRow, 
 		const UE::Editor::DataStorage::FMetaDataView& Arguments) override;
-
-protected:
-
-	TSharedRef<SWidget> CreateLabel(IEditorDataStorageProvider* DataStorage,
-		IEditorDataStorageUiProvider* DataStorageUi, UE::Editor::DataStorage::RowHandle TargetRow, 
-		UE::Editor::DataStorage::RowHandle WidgetRow, const UE::Editor::DataStorage::FMetaDataView& Arguments);
 };
\ No newline at end of file
diff --git a/Engine/Plugins/Experimental/EditorDataStorageFeatures/Source/TedsOutliner/Private/Widgets/OutlinerTextWidget.cpp b/Engine/Plugins/Experimental/EditorDataStorageFeatures/Source/TedsOutliner/Private/Widgets/OutlinerTextWidget.cpp
new file mode 100644
--- /dev/null
+++ b/Engine/Plugins/Experimental/EditorDataStorageFeatures/Source/TedsOutliner/Private/Widgets/OutlinerTextWidget.cpp
@@ -0,0 +1,100 @@
+// Copyright Epic Games, Inc. All Rights Reserved.
+
+#include "Widgets/OutlinerTextWidget.h"
+
+#include "ActorEditorUtils.h"
+#include "Columns/SlateDelegateColumns.h"
+#include "Elements/Columns/TypedElementLabelColumns.h"
+#include "Elements/Columns/TypedElementMiscColumns.h"
+#include "Elements/Framework/TypedElementAttributeBinding.h"
+#include "Elements/Interfaces/Capabilities/TypedElementUiEditableCapability.h"
+#include "Elements/Interfaces/Capabilities/TypedElementUiStyleOverrideCapability.h"
+#include "Settings/EditorStyleSettings.h"
+#include "Widgets/Text/SInlineEditableTextBlock.h"
+#include "Widgets/Text/STextBlock.h"
+
+#define LOCTEXT_NAMESPACE "OutlinerTextWidget"
+
+void UOutlinerTextWidgetFactory::RegisterWidgetConstructors(IEditorDataStorageProvider& DataStorage,
+	IEditorDataStorageUiProvider& DataStorageUi) const
+{
+	using namespace UE::Editor::DataStorage::Queries;
+
+	DataStorageUi.RegisterWidgetFactory<FOutlinerTextWidgetConstructor>(TEXT("SceneOutliner.RowLabel.Text"),
+		TColumn<FTypedElementLabelColumn>());
+}
+
+FOutlinerTextWidgetConstructor::FOutlinerTextWidgetConstructor()
+	: Super(FOutlinerTextWidgetConstructor::StaticStruct())
+{
+}
+
+TSharedPtr<SWidget> FOutlinerTextWidgetConstructor::CreateWidget(IEditorDataStorageProvider* DataStorage,
+	IEditorDataStorageUiProvider* DataStorageUi, UE::Editor::DataStorage::RowHandle TargetRow,
+	UE::Editor::DataStorage::RowHandle WidgetRow, const UE::Editor::DataStorage::FMetaDataView& Arguments)
+{
+	using namespace UE::Editor::DataStorage;
+	if (!DataStorage->IsRowAvailable(TargetRow))
+	{
+		return SNew(STextBlock)
+			.Text(LOCTEXT("MissingRowReferenceColumn", "Unable to retrieve row reference."));
+	}
+
+	TSharedRef<SWidget> Result = SNullWidget::NullWidget;
+
+	const bool* IsEditable = Arguments.FindForColumn<FTypedElementLabelColumn>(IsEditableName).TryGetExact<bool>();
+
+	if (IsEditable && *IsEditable)
+	{
+		FAttributeBinder TargetRowBinder(TargetRow, DataStorage);
+		FAttributeBinder WidgetRowBinder(WidgetRow, DataStorage);
+
+		const bool bShouldUseMiddleEllipsis = GetDefault<UEditorStyleSettings>()->bEnableMiddleEllipsis;
+
+		TSharedPtr<SInlineEditableTextBlock> TextBlock = SNew(SInlineEditableTextBlock)
+			.OnTextCommitted_Lambda(
+				[DataStorage, TargetRow](const FText& NewText, ETextCommit::Type CommitInfo)
+				{
+					// This callback happens on the game thread so it's safe to directly call into the data storage.
+					FString NewLabelText = NewText.ToString();
+					if (FTypedElementLabelHashColumn* LabelHashColumn = DataStorage->GetColumn<FTypedElementLabelHashColumn>(TargetRow))
+					{
+						LabelHashColumn->LabelHash = CityHash64(reinterpret_cast<const char*>(*NewLabelText), NewLabelText.Len() * sizeof(**NewLabelText));
+					}
+					if (FTypedElementLabelColumn* LabelColumn = DataStorage->GetColumn<FTypedElementLabelColumn>(TargetRow))
+					{
+						LabelColumn->Label = MoveTemp(NewLabelText);
+					}
+					DataStorage->AddColumn<FTypedElementSyncBackToWorldTag>(TargetRow);
+				})
+			.OnVerifyTextChanged_Lambda([](const FText& Label, FText& ErrorMessage)
+				{
+					// Note: The use of actor specific functionality should be minimized, but this function acts generic enough that the 
+					// use of actor is just in names.
+					return FActorEditorUtils::ValidateActorName(Label, ErrorMessage);
+				})
+			.Text(TargetRowBinder.BindText(&FTypedElementLabelColumn::Label))
+			.ToolTipText(TargetRowBinder.BindText(&FTypedElementLabelColumn::Label))
+			.OverflowPolicy(bShouldUseMiddleEllipsis ? ETextOverflowPolicy::MiddleEllipsis : TOptional<ETextOverflowPolicy>())
+			.IsSelected(WidgetRowBinder.BindEvent(&FExternalWidgetSelectionColumn::IsSelected));
+
+		TextBlock->AddMetadata(MakeShared<TTypedElementUiEditableCapability<SInlineEditableTextBlock>>(*TextBlock));
+		TextBlock->AddMetadata(MakeShared<TTypedElementUiStyleOverrideCapability<SInlineEditableTextBlock>>(*TextBlock));
+		Result = TextBlock.ToSharedRef();
+	}
+	else
+	{
+		FAttributeBinder TargetRowBinder(TargetRow, DataStorage);
+
+		TSharedPtr<STextBlock> TextBlock = SNew(STextBlock)
+			.IsEnabled(false)
+			.Text(TargetRowBinder.BindText(&FTypedElementLabelColumn::Label))
+			.ToolTipText(TargetRowBinder.BindText(&FTypedElementLabelColumn::Label));
+
+		Result = TextBlock.ToSharedRef();
+	}
+
+	return Result;
+}
+
+#undef LOCTEXT_NAMESPACE
diff --git a/Engine/Plugins/Experimental/EditorDataStorageFeatures/Source/TedsOutliner/Public/Widgets/OutlinerIconWidget.h b/Engine/Plugins/Experimental/EditorDataStorageFeatures/Source/TedsOutliner/Public/Widgets/OutlinerIconWidget.h
new file mode 100644
--- /dev/null
+++ b/Engine/Plugins/Experimental/EditorDataStorageFeatures/Source/TedsOutliner/Public/Widgets/OutlinerIconWidget.h
@@ -0,0 +1,41 @@
+// Copyright Epic Games, Inc. All Rights Reserved.
+
+#pragma once
+
+#include "Elements/Interfaces/TypedElementDataStorageFactory.h"
+#include "Elements/Interfaces/TypedElementDataStorageUiInterface.h"
+#include "UObject/ObjectMacros.h"
+
+#include "OutlinerIconWidget.generated.h"
+
+struct FSlateBrush;
+struct FTypedElementClassTypeInfoColumn;
+
+UCLASS()
+class UOutlinerIconWidgetFactory : public UEditorDataStorageFactory
+{
+	GENERATED_BODY()
+
+public:
+	~UOutlinerIconWidgetFactory() override = default;
+
+	void RegisterWidgetConstructors(IEditorDataStorageProvider& DataStorage,
+		IEditorDataStorageUiProvider& DataStorageUi) const override;
+};
+
+USTRUCT()
+struct FOutlinerIconWidgetConstructor : public FSimpleWidgetConstructor
+{
+	GENERATED_BODY()
+
+public:
+	TEDSOUTLINER_API FOutlinerIconWidgetConstructor();
+	~FOutlinerIconWidgetConstructor() override = default;
+
+	TEDSOUTLINER_API virtual TSharedPtr<SWidget> CreateWidget(
+		IEditorDataStorageProvider* DataStorage,
+		IEditorDataStorageUiProvider* DataStorageUi,
+		UE::Editor::DataStorage::RowHandle TargetRow,
+		UE::Editor::DataStorage::RowHandle WidgetRow,
+		const UE::Editor::DataStorage::FMetaDataView& Arguments) override;
+};
diff --git a/Engine/Plugins/Experimental/EditorDataStorageFeatures/Source/TedsOutliner/Public/Widgets/OutlinerTextWidget.h b/Engine/Plugins/Experimental/EditorDataStorageFeatures/Source/TedsOutliner/Public/Widgets/OutlinerTextWidget.h
new file mode 100644
--- /dev/null
+++ b/Engine/Plugins/Experimental/EditorDataStorageFeatures/Source/TedsOutliner/Public/Widgets/OutlinerTextWidget.h
@@ -0,0 +1,41 @@
+// Copyright Epic Games, Inc. All Rights Reserved.
+
+#pragma once
+
+#include "Elements/Interfaces/TypedElementDataStorageFactory.h"
+#include "Elements/Interfaces/TypedElementDataStorageUiInterface.h"
+#include "UObject/ObjectMacros.h"
+
+#include "OutlinerTextWidget.generated.h"
+
+struct FSlateBrush;
+struct FTypedElementClassTypeInfoColumn;
+
+UCLASS()
+class UOutlinerTextWidgetFactory : public UEditorDataStorageFactory
+{
+	GENERATED_BODY()
+
+public:
+	~UOutlinerTextWidgetFactory() override = default;
+
+	void RegisterWidgetConstructors(IEditorDataStorageProvider& DataStorage,
+		IEditorDataStorageUiProvider& DataStorageUi) const override;
+};
+
+USTRUCT()
+struct FOutlinerTextWidgetConstructor : public FSimpleWidgetConstructor
+{
+	GENERATED_BODY()
+
+public:
+	TEDSOUTLINER_API FOutlinerTextWidgetConstructor();
+	~FOutlinerTextWidgetConstructor() override = default;
+
+	TEDSOUTLINER_API virtual TSharedPtr<SWidget> CreateWidget(
+		IEditorDataStorageProvider* DataStorage,
+		IEditorDataStorageUiProvider* DataStorageUi,
+		UE::Editor::DataStorage::RowHandle TargetRow,
+		UE::Editor::DataStorage::RowHandle WidgetRow,
+		const UE::Editor::DataStorage::FMetaDataView& Arguments) override;
+};
