TedsOutliner: Fix parent update query incorrectly assuming Context.GetColumn() was retrieving the specific column for the row instead of an array of all columns in the table

diff --git a/Engine/Plugins/Experimental/EditorDataStorageFeatures/Source/TedsOutliner/Private/TedsOutlinerImpl.cpp b/Engine/Plugins/Experimental/EditorDataStorageFeatures/Source/TedsOutliner/Private/TedsOutlinerImpl.cpp
--- a/Engine/Plugins/Experimental/EditorDataStorageFeatures/Source/TedsOutliner/Private/TedsOutlinerImpl.cpp
+++ b/Engine/Plugins/Experimental/EditorDataStorageFeatures/Source/TedsOutliner/Private/TedsOutlinerImpl.cpp
@@ -593,29 +593,31 @@ void FTedsOutlinerImpl::RecompileQueries()
 			TEXT("Update item parent"),
 			FProcessor(EQueryTickPhase::PostPhysics, Storage->GetQueryTickGroupName(EQueryTickGroups::SyncDataStorageToExternal))
 				.SetExecutionMode(EExecutionMode::GameThread),
-			[this, ParentColumnType](IQueryContext& Context, const RowHandle Row)
+			[this, ParentColumnType](IQueryContext& Context, const RowHandle* Rows)
 			{
-				const void* ParentColumn = Context.GetColumn(ParentColumnType);
-				const FTedsOutlinerColumn* TedsOutlinerColumn = Context.GetColumn<FTedsOutlinerColumn>();
-				
-				RowHandle ParentRowHandle = InvalidRowHandle;
-				
-				if (ParentColumn)
-				{
-					ParentRowHandle = HierarchyData.GetValue().GetParent.Execute(ParentColumn);
-				}
-								
-				if (QueryUtils::HasItemParentChanged(Context, Row, ParentRowHandle, *SceneOutliner))
+				if (const char* ParentColumn = reinterpret_cast<const char*>(Context.GetColumn(ParentColumnType)))
 				{
-					if (TedsOutlinerColumn && !QueryUtils::CanDisplayRow(Context, *TedsOutlinerColumn, Row, *SceneOutliner))
+					int32 ColumnSize = ParentColumnType->GetCppStructOps() ? ParentColumnType->GetCppStructOps()->GetSize() : ParentColumnType->GetStructureSize();
+					const FTedsOutlinerColumn* TedsOutlinerColumns = Context.GetColumn<FTedsOutlinerColumn>();
+					uint32 RowCount = Context.GetRowCount();
+					
+					for(uint32 RowIndex = 0; RowIndex < RowCount; ++RowIndex, ParentColumn += ColumnSize)
 					{
-						return;
+						RowHandle ParentRowHandle = HierarchyData.GetValue().GetParent.Execute(ParentColumn);
+								
+						if (QueryUtils::HasItemParentChanged(Context, Rows[RowIndex], ParentRowHandle, *SceneOutliner))
+						{
+							if (TedsOutlinerColumns && !QueryUtils::CanDisplayRow(Context, TedsOutlinerColumns[RowIndex], Rows[RowIndex], *SceneOutliner))
+							{
+								continue;
+							}
+						
+							FSceneOutlinerHierarchyChangedData EventData;
+							EventData.Type = FSceneOutlinerHierarchyChangedData::Moved;
+							EventData.ItemIDs.Add(Rows[RowIndex]);
+							HierarchyChangedEvent.Broadcast(EventData);
+						}
 					}
-					
-					FSceneOutlinerHierarchyChangedData EventData;
-					EventData.Type = FSceneOutlinerHierarchyChangedData::Moved;
-					EventData.ItemIDs.Add(Row);
-					HierarchyChangedEvent.Broadcast(EventData);
 				}
 			})
 			.ReadOnly(ParentColumnType, EOptional::Yes)
