Introduced specialized Outliner text widget for Entity names to separate the universally unique id suffix.

diff --git a/Engine/Plugins/Experimental/EditorDataStorageFeatures/Source/TedsOutliner/Private/Widgets/OutlinerTextWidget.cpp b/Engine/Plugins/Experimental/EditorDataStorageFeatures/Source/TedsOutliner/Private/Widgets/OutlinerTextWidget.cpp
--- a/Engine/Plugins/Experimental/EditorDataStorageFeatures/Source/TedsOutliner/Private/Widgets/OutlinerTextWidget.cpp
+++ b/Engine/Plugins/Experimental/EditorDataStorageFeatures/Source/TedsOutliner/Private/Widgets/OutlinerTextWidget.cpp
@@ -29,6 +29,11 @@ FOutlinerTextWidgetConstructor::FOutlinerTextWidgetConstructor()
 {
 }
 
+FOutlinerTextWidgetConstructor::FOutlinerTextWidgetConstructor(const UScriptStruct* InTypeInfo)
+	: Super(InTypeInfo)
+{
+}
+
 TSharedPtr<SWidget> FOutlinerTextWidgetConstructor::CreateWidget(IEditorDataStorageProvider* DataStorage,
 	IEditorDataStorageUiProvider* DataStorageUi, UE::Editor::DataStorage::RowHandle TargetRow,
 	UE::Editor::DataStorage::RowHandle WidgetRow, const UE::Editor::DataStorage::FMetaDataView& Arguments)
@@ -40,61 +45,106 @@ TSharedPtr<SWidget> FOutlinerTextWidgetConstructor::CreateWidget(IEditorDataStor
 			.Text(LOCTEXT("MissingRowReferenceColumn", "Unable to retrieve row reference."));
 	}
 
-	TSharedRef<SWidget> Result = SNullWidget::NullWidget;
+	TSharedPtr<SWidget> Result = SNullWidget::NullWidget;
 
 	const bool* IsEditable = Arguments.FindForColumn<FTypedElementLabelColumn>(IsEditableName).TryGetExact<bool>();
 
 	if (IsEditable && *IsEditable)
 	{
-		FAttributeBinder TargetRowBinder(TargetRow, DataStorage);
-		FAttributeBinder WidgetRowBinder(WidgetRow, DataStorage);
-
-		const bool bShouldUseMiddleEllipsis = GetDefault<UEditorStyleSettings>()->bEnableMiddleEllipsis;
-
-		TSharedPtr<SInlineEditableTextBlock> TextBlock = SNew(SInlineEditableTextBlock)
-			.OnTextCommitted_Lambda(
-				[DataStorage, TargetRow](const FText& NewText, ETextCommit::Type CommitInfo)
-				{
-					// This callback happens on the game thread so it's safe to directly call into the data storage.
-					FString NewLabelText = NewText.ToString();
-					if (FTypedElementLabelHashColumn* LabelHashColumn = DataStorage->GetColumn<FTypedElementLabelHashColumn>(TargetRow))
-					{
-						LabelHashColumn->LabelHash = CityHash64(reinterpret_cast<const char*>(*NewLabelText), NewLabelText.Len() * sizeof(**NewLabelText));
-					}
-					if (FTypedElementLabelColumn* LabelColumn = DataStorage->GetColumn<FTypedElementLabelColumn>(TargetRow))
-					{
-						LabelColumn->Label = MoveTemp(NewLabelText);
-					}
-					DataStorage->AddColumn<FTypedElementSyncBackToWorldTag>(TargetRow);
-				})
-			.OnVerifyTextChanged_Lambda([](const FText& Label, FText& ErrorMessage)
-				{
-					// Note: The use of actor specific functionality should be minimized, but this function acts generic enough that the 
-					// use of actor is just in names.
-					return FActorEditorUtils::ValidateActorName(Label, ErrorMessage);
-				})
-			.Text(TargetRowBinder.BindText(&FTypedElementLabelColumn::Label))
-			.ToolTipText(TargetRowBinder.BindText(&FTypedElementLabelColumn::Label))
-			.OverflowPolicy(bShouldUseMiddleEllipsis ? ETextOverflowPolicy::MiddleEllipsis : TOptional<ETextOverflowPolicy>())
-			.IsSelected(WidgetRowBinder.BindEvent(&FExternalWidgetSelectionColumn::IsSelected));
-
-		TextBlock->AddMetadata(MakeShared<TTypedElementUiEditableCapability<SInlineEditableTextBlock>>(*TextBlock));
-		TextBlock->AddMetadata(MakeShared<TTypedElementUiStyleOverrideCapability<SInlineEditableTextBlock>>(*TextBlock));
-		Result = TextBlock.ToSharedRef();
+		Result = CreateEditableWidget(DataStorage, DataStorageUi, TargetRow, WidgetRow, Arguments);
 	}
 	else
 	{
-		FAttributeBinder TargetRowBinder(TargetRow, DataStorage);
+		Result = CreateNonEditableWidget(DataStorage, DataStorageUi, TargetRow, WidgetRow, Arguments);
+	}
+	return Result;
+}
 
-		TSharedPtr<STextBlock> TextBlock = SNew(STextBlock)
-			.IsEnabled(false)
-			.Text(TargetRowBinder.BindText(&FTypedElementLabelColumn::Label))
-			.ToolTipText(TargetRowBinder.BindText(&FTypedElementLabelColumn::Label));
+void FOutlinerTextWidgetConstructor::OnCommitText(
+	const FText& NewText,
+	ETextCommit::Type,
+	IEditorDataStorageProvider* DataStorage,
+	UE::Editor::DataStorage::RowHandle TargetRow)
+{
+	FString NewLabelText = NewText.ToString();
 
-		Result = TextBlock.ToSharedRef();
+	// This callback happens on the game thread so it's safe to directly call into the data storage.
+	if (FTypedElementLabelHashColumn* LabelHashColumn = DataStorage->GetColumn<FTypedElementLabelHashColumn>(TargetRow))
+	{
+		LabelHashColumn->LabelHash = CityHash64(reinterpret_cast<const char*>(*NewLabelText), NewLabelText.Len() * sizeof(**NewLabelText));
 	}
+	if (FTypedElementLabelColumn* LabelColumn = DataStorage->GetColumn<FTypedElementLabelColumn>(TargetRow))
+	{
+		LabelColumn->Label = MoveTemp(NewLabelText);
+	}
+	DataStorage->AddColumn<FTypedElementSyncBackToWorldTag>(TargetRow);
+}
 
-	return Result;
+bool FOutlinerTextWidgetConstructor::OnVerifyText(const FText& Label, FText& ErrorMessage)
+{
+	// Note: The use of actor specific functionality should be minimized, but this function acts generic enough that the 
+	// use of actor is just in names.
+	return FActorEditorUtils::ValidateActorName(Label, ErrorMessage);
+}
+
+TSharedPtr<SWidget> FOutlinerTextWidgetConstructor::CreateEditableWidget(
+	IEditorDataStorageProvider* DataStorage,
+	IEditorDataStorageUiProvider* DataStorageUi,
+	UE::Editor::DataStorage::RowHandle TargetRow,
+	UE::Editor::DataStorage::RowHandle WidgetRow,
+	const UE::Editor::DataStorage::FMetaDataView& Arguments)
+{
+	return CreateEditableTextBlock(DataStorage, TargetRow, WidgetRow);
+}
+
+TSharedPtr<SWidget> FOutlinerTextWidgetConstructor::CreateNonEditableWidget(
+	IEditorDataStorageProvider* DataStorage,
+	IEditorDataStorageUiProvider* DataStorageUi,
+	UE::Editor::DataStorage::RowHandle TargetRow,
+	UE::Editor::DataStorage::RowHandle WidgetRow,
+	const UE::Editor::DataStorage::FMetaDataView& Arguments)
+{
+	return CreateNonEditableTextBlock(DataStorage, TargetRow, WidgetRow);
+}
+
+TSharedPtr<SInlineEditableTextBlock> FOutlinerTextWidgetConstructor::CreateEditableTextBlock(
+	IEditorDataStorageProvider* DataStorage,
+	UE::Editor::DataStorage::RowHandle TargetRow,
+	UE::Editor::DataStorage::RowHandle WidgetRow) const
+{
+	using namespace UE::Editor::DataStorage;
+	FAttributeBinder TargetRowBinder(TargetRow, DataStorage);
+	FAttributeBinder WidgetRowBinder(WidgetRow, DataStorage);
+
+	const bool bShouldUseMiddleEllipsis = GetDefault<UEditorStyleSettings>()->bEnableMiddleEllipsis;
+
+	TSharedPtr<SInlineEditableTextBlock> TextBlock = SNew(SInlineEditableTextBlock)
+		.OnTextCommitted_Static(&FOutlinerTextWidgetConstructor::OnCommitText, DataStorage, TargetRow)
+		.OnVerifyTextChanged_Static(&FOutlinerTextWidgetConstructor::OnVerifyText)
+		.Text(TargetRowBinder.BindText(&FTypedElementLabelColumn::Label))
+		.ToolTipText(TargetRowBinder.BindText(&FTypedElementLabelColumn::Label))
+		.OverflowPolicy(bShouldUseMiddleEllipsis ? ETextOverflowPolicy::MiddleEllipsis : TOptional<ETextOverflowPolicy>())
+		.IsSelected(WidgetRowBinder.BindEvent(&FExternalWidgetSelectionColumn::IsSelected));
+
+	TextBlock->AddMetadata(MakeShared<TTypedElementUiEditableCapability<SInlineEditableTextBlock>>(*TextBlock));
+	TextBlock->AddMetadata(MakeShared<TTypedElementUiStyleOverrideCapability<SInlineEditableTextBlock>>(*TextBlock));
+	return TextBlock;
+}
+
+TSharedPtr<STextBlock> FOutlinerTextWidgetConstructor::CreateNonEditableTextBlock(
+	IEditorDataStorageProvider* DataStorage,
+	UE::Editor::DataStorage::RowHandle TargetRow,
+	UE::Editor::DataStorage::RowHandle WidgetRow) const
+{
+	using namespace UE::Editor::DataStorage;
+	FAttributeBinder TargetRowBinder(TargetRow, DataStorage);
+
+	TSharedPtr<STextBlock> TextBlock = SNew(STextBlock)
+		.IsEnabled(false)
+		.Text(TargetRowBinder.BindText(&FTypedElementLabelColumn::Label))
+		.ToolTipText(TargetRowBinder.BindText(&FTypedElementLabelColumn::Label));
+
+	return TextBlock;
 }
 
 #undef LOCTEXT_NAMESPACE
diff --git a/Engine/Plugins/Experimental/EditorDataStorageFeatures/Source/TedsOutliner/Public/Widgets/OutlinerTextWidget.h b/Engine/Plugins/Experimental/EditorDataStorageFeatures/Source/TedsOutliner/Public/Widgets/OutlinerTextWidget.h
--- a/Engine/Plugins/Experimental/EditorDataStorageFeatures/Source/TedsOutliner/Public/Widgets/OutlinerTextWidget.h
+++ b/Engine/Plugins/Experimental/EditorDataStorageFeatures/Source/TedsOutliner/Public/Widgets/OutlinerTextWidget.h
@@ -4,12 +4,15 @@
 
 #include "Elements/Interfaces/TypedElementDataStorageFactory.h"
 #include "Elements/Interfaces/TypedElementDataStorageUiInterface.h"
+#include "Types/SlateEnums.h"
 #include "UObject/ObjectMacros.h"
 
 #include "OutlinerTextWidget.generated.h"
 
 struct FSlateBrush;
 struct FTypedElementClassTypeInfoColumn;
+class  SInlineEditableTextBlock;
+class  STextBlock;
 
 UCLASS()
 class UOutlinerTextWidgetFactory : public UEditorDataStorageFactory
@@ -17,7 +20,7 @@ class UOutlinerTextWidgetFactory : public UEditorDataStorageFactory
 	GENERATED_BODY()
 
 public:
-	~UOutlinerTextWidgetFactory() override = default;
+	virtual ~UOutlinerTextWidgetFactory() override = default;
 
 	void RegisterWidgetConstructors(IEditorDataStorageProvider& DataStorage,
 		IEditorDataStorageUiProvider& DataStorageUi) const override;
@@ -30,7 +33,8 @@ struct FOutlinerTextWidgetConstructor : public FSimpleWidgetConstructor
 
 public:
 	TEDSOUTLINER_API FOutlinerTextWidgetConstructor();
-	~FOutlinerTextWidgetConstructor() override = default;
+	TEDSOUTLINER_API FOutlinerTextWidgetConstructor(const UScriptStruct* InTypeInfo);
+	TEDSOUTLINER_API virtual ~FOutlinerTextWidgetConstructor() override = default;
 
 	TEDSOUTLINER_API virtual TSharedPtr<SWidget> CreateWidget(
 		IEditorDataStorageProvider* DataStorage,
@@ -38,4 +42,39 @@ public:
 		UE::Editor::DataStorage::RowHandle TargetRow,
 		UE::Editor::DataStorage::RowHandle WidgetRow,
 		const UE::Editor::DataStorage::FMetaDataView& Arguments) override;
+
+protected:
+	//! Updates the label columns with the given text.
+	static TEDSOUTLINER_API void OnCommitText(const FText& NewText, ETextCommit::Type Type, IEditorDataStorageProvider* DataStorage, UE::Editor::DataStorage::RowHandle TargetRow);
+	static TEDSOUTLINER_API bool OnVerifyText(const FText& Label, FText& ErrorMessage);
+
+	//! Virtual relay function so that deriving Constructors may edit the editable
+	//! widget without having to override the whole CreateWidget procedure.
+	TEDSOUTLINER_API virtual TSharedPtr<SWidget> CreateEditableWidget(
+		IEditorDataStorageProvider* DataStorage,
+		IEditorDataStorageUiProvider* DataStorageUi,
+		UE::Editor::DataStorage::RowHandle TargetRow,
+		UE::Editor::DataStorage::RowHandle WidgetRow,
+		const UE::Editor::DataStorage::FMetaDataView& Arguments);
+
+	//! Virtual relay function so that deriving Constructors may edit the non-editable
+	//! widget without having to override the whole CreateWidget procedure.
+	TEDSOUTLINER_API virtual TSharedPtr<SWidget> CreateNonEditableWidget(
+		IEditorDataStorageProvider* DataStorage,
+		IEditorDataStorageUiProvider* DataStorageUi,
+		UE::Editor::DataStorage::RowHandle TargetRow,
+		UE::Editor::DataStorage::RowHandle WidgetRow,
+		const UE::Editor::DataStorage::FMetaDataView& Arguments);
+
+	//! The actual function that creates the editable widget.
+	TEDSOUTLINER_API TSharedPtr<SInlineEditableTextBlock> CreateEditableTextBlock(
+		IEditorDataStorageProvider* DataStorage,
+		UE::Editor::DataStorage::RowHandle TargetRow,
+		UE::Editor::DataStorage::RowHandle WidgetRow) const;
+
+	//! The actual function that creates the non-editable widget.
+	TEDSOUTLINER_API TSharedPtr<STextBlock> CreateNonEditableTextBlock(
+		IEditorDataStorageProvider* DataStorage,
+		UE::Editor::DataStorage::RowHandle TargetRow,
+		UE::Editor::DataStorage::RowHandle WidgetRow) const;
 };
