[TEDS] Moving the Styling module to the NFL plugin folder

diff --git a/Engine/Plugins/Experimental/EditorDataStorageFeatures/EditorDataStorageFeatures.uplugin b/Engine/Plugins/Experimental/EditorDataStorageFeatures/EditorDataStorageFeatures.uplugin
--- a/Engine/Plugins/Experimental/EditorDataStorageFeatures/EditorDataStorageFeatures.uplugin
+++ b/Engine/Plugins/Experimental/EditorDataStorageFeatures/EditorDataStorageFeatures.uplugin
@@ -51,10 +51,6 @@
 			"Name": "TedsSettings",
 			"Type": "Editor"
 		},
-		{
-			"Name": "TedsStyling",
-			"Type": "Editor"
-		},
 		{
 			"Name": "TedsTableViewer",
 			"Type": "Editor"
diff --git a/Engine/Plugins/Experimental/EditorDataStorageFeatures/Source/TedsStyling/Private/TedsStylingFactory.cpp b/Engine/Plugins/Experimental/EditorDataStorageFeatures/Source/TedsStyling/Private/TedsStylingFactory.cpp
deleted file mode 100644
diff --git a/Engine/Plugins/Experimental/EditorDataStorageFeatures/Source/TedsStyling/Private/TedsStylingFactory.h b/Engine/Plugins/Experimental/EditorDataStorageFeatures/Source/TedsStyling/Private/TedsStylingFactory.h
deleted file mode 100644
diff --git a/Engine/Plugins/Experimental/EditorDataStorageFeatures/Source/TedsStyling/Private/TedsStylingModule.cpp b/Engine/Plugins/Experimental/EditorDataStorageFeatures/Source/TedsStyling/Private/TedsStylingModule.cpp
deleted file mode 100644
diff --git a/Engine/Plugins/Experimental/EditorDataStorageFeatures/Source/TedsStyling/Private/Widgets/SSlateGallery.cpp b/Engine/Plugins/Experimental/EditorDataStorageFeatures/Source/TedsStyling/Private/Widgets/SSlateGallery.cpp
deleted file mode 100644
diff --git a/Engine/Plugins/Experimental/EditorDataStorageFeatures/Source/TedsStyling/Private/Widgets/SSlateGallery.h b/Engine/Plugins/Experimental/EditorDataStorageFeatures/Source/TedsStyling/Private/Widgets/SSlateGallery.h
deleted file mode 100644
diff --git a/Engine/Plugins/Experimental/EditorDataStorageFeatures/Source/TedsStyling/Private/Widgets/SlateBrushWidget.cpp b/Engine/Plugins/Experimental/EditorDataStorageFeatures/Source/TedsStyling/Private/Widgets/SlateBrushWidget.cpp
deleted file mode 100644
diff --git a/Engine/Plugins/Experimental/EditorDataStorageFeatures/Source/TedsStyling/Private/Widgets/SlateBrushWidget.h b/Engine/Plugins/Experimental/EditorDataStorageFeatures/Source/TedsStyling/Private/Widgets/SlateBrushWidget.h
deleted file mode 100644
diff --git a/Engine/Plugins/Experimental/EditorDataStorageFeatures/Source/TedsStyling/Public/Columns/TedsStylingColumns.h b/Engine/Plugins/Experimental/EditorDataStorageFeatures/Source/TedsStyling/Public/Columns/TedsStylingColumns.h
deleted file mode 100644
diff --git a/Engine/Plugins/Experimental/EditorDataStorageFeatures/Source/TedsStyling/TedsStyling.Build.cs b/Engine/Plugins/Experimental/EditorDataStorageFeatures/Source/TedsStyling/TedsStyling.Build.cs
deleted file mode 100644
