[TEDS] Cleanup aliases, move core interface classes into namespace, and fix up one class name that wasn't following standards

diff --git a/Engine/Plugins/Experimental/EditorDataStorageFeatures/Source/TedsAssetData/Private/Widgets/AssetDataItemTypeWidget.h b/Engine/Plugins/Experimental/EditorDataStorageFeatures/Source/TedsAssetData/Private/Widgets/AssetDataItemTypeWidget.h
--- a/Engine/Plugins/Experimental/EditorDataStorageFeatures/Source/TedsAssetData/Private/Widgets/AssetDataItemTypeWidget.h
+++ b/Engine/Plugins/Experimental/EditorDataStorageFeatures/Source/TedsAssetData/Private/Widgets/AssetDataItemTypeWidget.h
@@ -7,9 +7,6 @@
 
 #include "AssetDataItemTypeWidget.generated.h"
 
-class IEditorDataStorageProvider;
-class UScriptStruct;
-
 UCLASS()
 class UAssetDataItemTypeWidgetFactory : public UEditorDataStorageFactory
 {
@@ -18,8 +15,8 @@ class UAssetDataItemTypeWidgetFactory : public UEditorDataStorageFactory
 public:
 	~UAssetDataItemTypeWidgetFactory() override = default;
 
-	TEDSASSETDATA_API void RegisterWidgetConstructors(IEditorDataStorageProvider& DataStorage,
-		IEditorDataStorageUiProvider& DataStorageUi) const override;
+	TEDSASSETDATA_API void RegisterWidgetConstructors(UE::Editor::DataStorage::ICoreProvider& DataStorage,
+		UE::Editor::DataStorage::IUiProvider& DataStorageUi) const override;
 };
 
 // Widget to show Item type (folder/material/mesh etc...)
@@ -33,8 +30,8 @@ public:
 	~FAssetDataItemTypeWidgetConstructor() override = default;
 
 	TEDSASSETDATA_API virtual TSharedPtr<SWidget> CreateWidget(
-		IEditorDataStorageProvider* DataStorage,
-		IEditorDataStorageUiProvider* DataStorageUi,
+		UE::Editor::DataStorage::ICoreProvider* DataStorage,
+		UE::Editor::DataStorage::IUiProvider* DataStorageUi,
 		UE::Editor::DataStorage::RowHandle TargetRow,
 		UE::Editor::DataStorage::RowHandle WidgetRow,
 		const UE::Editor::DataStorage::FMetaDataView& Arguments) override;
diff --git a/Engine/Plugins/Experimental/EditorDataStorageFeatures/Source/TedsAssetData/Private/Widgets/AssetDataLabelWidget.cpp b/Engine/Plugins/Experimental/EditorDataStorageFeatures/Source/TedsAssetData/Private/Widgets/AssetDataLabelWidget.cpp
--- a/Engine/Plugins/Experimental/EditorDataStorageFeatures/Source/TedsAssetData/Private/Widgets/AssetDataLabelWidget.cpp
+++ b/Engine/Plugins/Experimental/EditorDataStorageFeatures/Source/TedsAssetData/Private/Widgets/AssetDataLabelWidget.cpp
@@ -18,8 +18,8 @@
 
 #define LOCTEXT_NAMESPACE "FAssetDataLabelWidgetConstructor"
 
-void UAssetDataLabelWidgetFactory::RegisterWidgetConstructors(IEditorDataStorageProvider& DataStorage,
-                                                              IEditorDataStorageUiProvider& DataStorageUi) const
+void UAssetDataLabelWidgetFactory::RegisterWidgetConstructors(UE::Editor::DataStorage::ICoreProvider& DataStorage,
+	UE::Editor::DataStorage::IUiProvider& DataStorageUi) const
 {
 	using namespace UE::Editor::DataStorage::Queries;
 	DataStorageUi.RegisterWidgetFactory<FAssetDataLabelWidgetConstructor>(
@@ -37,13 +37,15 @@ FAssetDataLabelWidgetConstructor::FAssetDataLabelWidgetConstructor(const UScript
 {
 }
 
-TSharedPtr<SWidget> FAssetDataLabelWidgetConstructor::CreateWidget(IEditorDataStorageProvider* DataStorage,
-	IEditorDataStorageUiProvider* DataStorageUi, UE::Editor::DataStorage::RowHandle TargetRow, UE::Editor::DataStorage::RowHandle WidgetRow,
+TSharedPtr<SWidget> FAssetDataLabelWidgetConstructor::CreateWidget(UE::Editor::DataStorage::ICoreProvider* DataStorage,
+	UE::Editor::DataStorage::IUiProvider* DataStorageUi,
+	UE::Editor::DataStorage::RowHandle TargetRow, UE::Editor::DataStorage::RowHandle WidgetRow,
 	const UE::Editor::DataStorage::FMetaDataView& Arguments)
 {
+	using namespace UE::Editor::DataStorage;
 	if (DataStorage->IsRowAvailable(TargetRow))
 	{
-		UE::Editor::DataStorage::FAttributeBinder Binder(TargetRow, DataStorage);
+		FAttributeBinder Binder(TargetRow, DataStorage);
 
 		bool bIsAsset = DataStorage->HasColumns<FAssetDataColumn_Experimental>(TargetRow);
 
@@ -62,7 +64,7 @@ TSharedPtr<SWidget> FAssetDataLabelWidgetConstructor::CreateWidget(IEditorDataSt
 		};
 
 		TArray<TWeakObjectPtr<const UScriptStruct>> LabelColumns = GetLabelColumns();
-		DataStorageUi->CreateWidgetConstructors(LabelPurposeName, IEditorDataStorageUiProvider::EMatchApproach::ExactMatch, LabelColumns, Arguments, AssignLabelWidgetToColumn);
+		DataStorageUi->CreateWidgetConstructors(LabelPurposeName, IUiProvider::EMatchApproach::ExactMatch, LabelColumns, Arguments, AssignLabelWidgetToColumn);
 
 		TSharedPtr<SWidget> LabelWidget = SNullWidget::NullWidget;
 		if (OutLabelWidgetConstructorPtr)
diff --git a/Engine/Plugins/Experimental/EditorDataStorageFeatures/Source/TedsAssetData/Private/Widgets/AssetDataThumbnailWidget.cpp b/Engine/Plugins/Experimental/EditorDataStorageFeatures/Source/TedsAssetData/Private/Widgets/AssetDataThumbnailWidget.cpp
--- a/Engine/Plugins/Experimental/EditorDataStorageFeatures/Source/TedsAssetData/Private/Widgets/AssetDataThumbnailWidget.cpp
+++ b/Engine/Plugins/Experimental/EditorDataStorageFeatures/Source/TedsAssetData/Private/Widgets/AssetDataThumbnailWidget.cpp
@@ -14,8 +14,8 @@
 
 #define LOCTEXT_NAMESPACE "FAssetThumbnailWidgetConstructor"
 
-void UAssetThumbnailWidgetFactory::RegisterWidgetConstructors(IEditorDataStorageProvider& DataStorage,
-	IEditorDataStorageUiProvider& DataStorageUi) const
+void UAssetThumbnailWidgetFactory::RegisterWidgetConstructors(UE::Editor::DataStorage::ICoreProvider& DataStorage,
+	UE::Editor::DataStorage::IUiProvider& DataStorageUi) const
 {
 	using namespace UE::Editor::DataStorage::Queries;
 	DataStorageUi.RegisterWidgetFactory<FAssetThumbnailWidgetConstructor>(
@@ -33,8 +33,9 @@ FAssetThumbnailWidgetConstructor::FAssetThumbnailWidgetConstructor(const UScript
 {
 }
 
-TSharedPtr<SWidget> FAssetThumbnailWidgetConstructor::CreateWidget(IEditorDataStorageProvider* DataStorage,
-	IEditorDataStorageUiProvider* DataStorageUi, UE::Editor::DataStorage::RowHandle TargetRow, UE::Editor::DataStorage::RowHandle WidgetRow,
+TSharedPtr<SWidget> FAssetThumbnailWidgetConstructor::CreateWidget(UE::Editor::DataStorage::ICoreProvider* DataStorage,
+	UE::Editor::DataStorage::IUiProvider* DataStorageUi,
+	UE::Editor::DataStorage::RowHandle TargetRow, UE::Editor::DataStorage::RowHandle WidgetRow,
 	const UE::Editor::DataStorage::FMetaDataView& Arguments)
 {
 	UE::Editor::DataStorage::FAttributeBinder Binder(TargetRow, DataStorage);
diff --git a/Engine/Plugins/Experimental/EditorDataStorageFeatures/Source/TedsAssetData/Private/Widgets/AssetDataThumbnailWidget.h b/Engine/Plugins/Experimental/EditorDataStorageFeatures/Source/TedsAssetData/Private/Widgets/AssetDataThumbnailWidget.h
--- a/Engine/Plugins/Experimental/EditorDataStorageFeatures/Source/TedsAssetData/Private/Widgets/AssetDataThumbnailWidget.h
+++ b/Engine/Plugins/Experimental/EditorDataStorageFeatures/Source/TedsAssetData/Private/Widgets/AssetDataThumbnailWidget.h
@@ -9,7 +9,6 @@
 
 #include "AssetDataThumbnailWidget.generated.h"
 
-class IEditorDataStorageProvider;
 class UScriptStruct;
 
 UCLASS()
@@ -20,8 +19,8 @@ class UAssetThumbnailWidgetFactory : public UEditorDataStorageFactory
 public:
 	virtual ~UAssetThumbnailWidgetFactory() override = default;
 
-	TEDSASSETDATA_API virtual void RegisterWidgetConstructors(IEditorDataStorageProvider& DataStorage,
-		IEditorDataStorageUiProvider& DataStorageUi) const override;
+	TEDSASSETDATA_API virtual void RegisterWidgetConstructors(UE::Editor::DataStorage::ICoreProvider& DataStorage,
+		UE::Editor::DataStorage::IUiProvider& DataStorageUi) const override;
 };
 
 // Thumbnail widget for assets in TEDS
@@ -36,8 +35,8 @@ public:
 	virtual ~FAssetThumbnailWidgetConstructor() override = default;
 
 	TEDSASSETDATA_API virtual TSharedPtr<SWidget> CreateWidget(
-		IEditorDataStorageProvider* DataStorage,
-		IEditorDataStorageUiProvider* DataStorageUi,
+		UE::Editor::DataStorage::ICoreProvider* DataStorage,
+		UE::Editor::DataStorage::IUiProvider* DataStorageUi,
 		UE::Editor::DataStorage::RowHandle TargetRow,
 		UE::Editor::DataStorage::RowHandle WidgetRow, 
 		const UE::Editor::DataStorage::FMetaDataView& Arguments) override;
diff --git a/Engine/Plugins/Experimental/EditorDataStorageFeatures/Source/TedsAssetData/Private/Widgets/AssetDataVirtualPathWidget.cpp b/Engine/Plugins/Experimental/EditorDataStorageFeatures/Source/TedsAssetData/Private/Widgets/AssetDataVirtualPathWidget.cpp
--- a/Engine/Plugins/Experimental/EditorDataStorageFeatures/Source/TedsAssetData/Private/Widgets/AssetDataVirtualPathWidget.cpp
+++ b/Engine/Plugins/Experimental/EditorDataStorageFeatures/Source/TedsAssetData/Private/Widgets/AssetDataVirtualPathWidget.cpp
@@ -10,8 +10,8 @@
 #include "Widgets/Layout/SBox.h"
 #include "Widgets/Text/STextBlock.h"
 
-void UAssetDataVirtualPathWidgetFactory::RegisterWidgetConstructors(IEditorDataStorageProvider& DataStorage,
-	IEditorDataStorageUiProvider& DataStorageUi) const
+void UAssetDataVirtualPathWidgetFactory::RegisterWidgetConstructors(UE::Editor::DataStorage::ICoreProvider& DataStorage,
+	UE::Editor::DataStorage::IUiProvider& DataStorageUi) const
 {
 	using namespace UE::Editor::DataStorage::Queries;
 	DataStorageUi.RegisterWidgetFactory<FAssetDataVirtualPathWidgetConstructor>(
@@ -24,8 +24,8 @@ FAssetDataVirtualPathWidgetConstructor::FAssetDataVirtualPathWidgetConstructor()
 {
 }
 
-TSharedPtr<SWidget> FAssetDataVirtualPathWidgetConstructor::CreateWidget(IEditorDataStorageProvider* DataStorage,
-	IEditorDataStorageUiProvider* DataStorageUi, UE::Editor::DataStorage::RowHandle TargetRow, UE::Editor::DataStorage::RowHandle WidgetRow,
+TSharedPtr<SWidget> FAssetDataVirtualPathWidgetConstructor::CreateWidget(UE::Editor::DataStorage::ICoreProvider* DataStorage,
+	UE::Editor::DataStorage::IUiProvider* DataStorageUi, UE::Editor::DataStorage::RowHandle TargetRow, UE::Editor::DataStorage::RowHandle WidgetRow,
 	const UE::Editor::DataStorage::FMetaDataView& Arguments)
 {
 	static const FMargin ColumnItemPadding(5, 0, 5, 0);
diff --git a/Engine/Plugins/Experimental/EditorDataStorageFeatures/Source/TedsAssetData/Private/Widgets/AssetDataVirtualPathWidget.h b/Engine/Plugins/Experimental/EditorDataStorageFeatures/Source/TedsAssetData/Private/Widgets/AssetDataVirtualPathWidget.h
--- a/Engine/Plugins/Experimental/EditorDataStorageFeatures/Source/TedsAssetData/Private/Widgets/AssetDataVirtualPathWidget.h
+++ b/Engine/Plugins/Experimental/EditorDataStorageFeatures/Source/TedsAssetData/Private/Widgets/AssetDataVirtualPathWidget.h
@@ -7,9 +7,6 @@
 
 #include "AssetDataVirtualPathWidget.generated.h"
 
-class IEditorDataStorageProvider;
-class UScriptStruct;
-
 UCLASS()
 class UAssetDataVirtualPathWidgetFactory : public UEditorDataStorageFactory
 {
@@ -18,8 +15,8 @@ class UAssetDataVirtualPathWidgetFactory : public UEditorDataStorageFactory
 public:
 	~UAssetDataVirtualPathWidgetFactory() override = default;
 
-	TEDSASSETDATA_API void RegisterWidgetConstructors(IEditorDataStorageProvider& DataStorage,
-		IEditorDataStorageUiProvider& DataStorageUi) const override;
+	TEDSASSETDATA_API void RegisterWidgetConstructors(UE::Editor::DataStorage::ICoreProvider& DataStorage,
+		UE::Editor::DataStorage::IUiProvider& DataStorageUi) const override;
 };
 
 // Widget to show disk size in bytes
@@ -33,8 +30,8 @@ public:
 	~FAssetDataVirtualPathWidgetConstructor() override = default;
 
 	TEDSASSETDATA_API virtual TSharedPtr<SWidget> CreateWidget(
-		IEditorDataStorageProvider* DataStorage,
-		IEditorDataStorageUiProvider* DataStorageUi,
+		UE::Editor::DataStorage::ICoreProvider* DataStorage,
+		UE::Editor::DataStorage::IUiProvider* DataStorageUi,
 		UE::Editor::DataStorage::RowHandle TargetRow,
 		UE::Editor::DataStorage::RowHandle WidgetRow,
 		const UE::Editor::DataStorage::FMetaDataView& Arguments) override;
diff --git a/Engine/Plugins/Experimental/EditorDataStorageFeatures/Source/TedsAssetData/Private/Widgets/AssetLabelWidget.cpp b/Engine/Plugins/Experimental/EditorDataStorageFeatures/Source/TedsAssetData/Private/Widgets/AssetLabelWidget.cpp
--- a/Engine/Plugins/Experimental/EditorDataStorageFeatures/Source/TedsAssetData/Private/Widgets/AssetLabelWidget.cpp
+++ b/Engine/Plugins/Experimental/EditorDataStorageFeatures/Source/TedsAssetData/Private/Widgets/AssetLabelWidget.cpp
@@ -10,8 +10,8 @@
 
 #define LOCTEXT_NAMESPACE "FAssetLabelWidgetConstructor"
 
-void UAssetLabelWidgetFactory::RegisterWidgetConstructors(IEditorDataStorageProvider& DataStorage,
-	IEditorDataStorageUiProvider& DataStorageUi) const
+void UAssetLabelWidgetFactory::RegisterWidgetConstructors(UE::Editor::DataStorage::ICoreProvider& DataStorage,
+	UE::Editor::DataStorage::IUiProvider& DataStorageUi) const
 {
 	using namespace UE::Editor::DataStorage::Queries;
 	DataStorageUi.RegisterWidgetFactory<FAssetLabelWidgetConstructor>(
@@ -29,8 +29,9 @@ FAssetLabelWidgetConstructor::FAssetLabelWidgetConstructor(const UScriptStruct*
 {
 }
 
-TSharedPtr<SWidget> FAssetLabelWidgetConstructor::CreateWidget(IEditorDataStorageProvider* DataStorage,
-	IEditorDataStorageUiProvider* DataStorageUi, UE::Editor::DataStorage::RowHandle TargetRow, UE::Editor::DataStorage::RowHandle WidgetRow,
+TSharedPtr<SWidget> FAssetLabelWidgetConstructor::CreateWidget(UE::Editor::DataStorage::ICoreProvider* DataStorage,
+	UE::Editor::DataStorage::IUiProvider* DataStorageUi,
+	UE::Editor::DataStorage::RowHandle TargetRow, UE::Editor::DataStorage::RowHandle WidgetRow,
 	const UE::Editor::DataStorage::FMetaDataView& Arguments)
 {
 	UE::Editor::DataStorage::FAttributeBinder Binder(TargetRow, DataStorage);
diff --git a/Engine/Plugins/Experimental/EditorDataStorageFeatures/Source/TedsAssetData/Private/Widgets/AssetLabelWidget.h b/Engine/Plugins/Experimental/EditorDataStorageFeatures/Source/TedsAssetData/Private/Widgets/AssetLabelWidget.h
--- a/Engine/Plugins/Experimental/EditorDataStorageFeatures/Source/TedsAssetData/Private/Widgets/AssetLabelWidget.h
+++ b/Engine/Plugins/Experimental/EditorDataStorageFeatures/Source/TedsAssetData/Private/Widgets/AssetLabelWidget.h
@@ -9,7 +9,6 @@
 
 #include "AssetLabelWidget.generated.h"
 
-class IEditorDataStorageProvider;
 class UScriptStruct;
 
 UCLASS()
@@ -20,8 +19,8 @@ class UAssetLabelWidgetFactory : public UEditorDataStorageFactory
 public:
 	virtual ~UAssetLabelWidgetFactory() override = default;
 
-	TEDSASSETDATA_API virtual void RegisterWidgetConstructors(IEditorDataStorageProvider& DataStorage,
-		IEditorDataStorageUiProvider& DataStorageUi) const override;
+	TEDSASSETDATA_API virtual void RegisterWidgetConstructors(UE::Editor::DataStorage::ICoreProvider& DataStorage,
+		UE::Editor::DataStorage::IUiProvider& DataStorageUi) const override;
 };
 
 // Label for assets in TEDS
@@ -36,8 +35,8 @@ public:
 	virtual ~FAssetLabelWidgetConstructor() override = default;
 
 	TEDSASSETDATA_API virtual TSharedPtr<SWidget> CreateWidget(
-		IEditorDataStorageProvider* DataStorage,
-		IEditorDataStorageUiProvider* DataStorageUi,
+		UE::Editor::DataStorage::ICoreProvider* DataStorage,
+		UE::Editor::DataStorage::IUiProvider* DataStorageUi,
 		UE::Editor::DataStorage::RowHandle TargetRow,
 		UE::Editor::DataStorage::RowHandle WidgetRow, 
 		const UE::Editor::DataStorage::FMetaDataView& Arguments) override;
diff --git a/Engine/Plugins/Experimental/EditorDataStorageFeatures/Source/TedsAssetData/Private/Widgets/ContentBrowserListViewNameWidget.cpp b/Engine/Plugins/Experimental/EditorDataStorageFeatures/Source/TedsAssetData/Private/Widgets/ContentBrowserListViewNameWidget.cpp
--- a/Engine/Plugins/Experimental/EditorDataStorageFeatures/Source/TedsAssetData/Private/Widgets/ContentBrowserListViewNameWidget.cpp
+++ b/Engine/Plugins/Experimental/EditorDataStorageFeatures/Source/TedsAssetData/Private/Widgets/ContentBrowserListViewNameWidget.cpp
@@ -18,15 +18,15 @@
 
 FName UContentBrowserListViewNameWidgetFactory::WidgetPurpose(TEXT("ContentBrowser.RowLabel"));
 
-void UContentBrowserListViewNameWidgetFactory::RegisterWidgetPurposes(IEditorDataStorageUiProvider& DataStorageUi) const
+void UContentBrowserListViewNameWidgetFactory::RegisterWidgetPurposes(UE::Editor::DataStorage::IUiProvider& DataStorageUi) const
 {
 	using namespace UE::Editor::DataStorage::Queries;
-	DataStorageUi.RegisterWidgetPurpose(WidgetPurpose, IEditorDataStorageUiProvider::EPurposeType::UniqueByNameAndColumn,
+	DataStorageUi.RegisterWidgetPurpose(WidgetPurpose, IUiProvider::EPurposeType::UniqueByNameAndColumn,
 		LOCTEXT("ContentBrowserLabelWidget_PurposeDescription", "Widget that display a Label + Thumbnail in the Content Browser"));
 }
 
-void UContentBrowserListViewNameWidgetFactory::RegisterWidgetConstructors(IEditorDataStorageProvider& DataStorage,
-	IEditorDataStorageUiProvider& DataStorageUi) const
+void UContentBrowserListViewNameWidgetFactory::RegisterWidgetConstructors(UE::Editor::DataStorage::ICoreProvider& DataStorage,
+	UE::Editor::DataStorage::IUiProvider& DataStorageUi) const
 {
 	using namespace UE::Editor::DataStorage::Queries;
 	DataStorageUi.RegisterWidgetFactory<FContentBrowserListViewNameWidgetConstructor>(
@@ -44,10 +44,12 @@ FContentBrowserListViewNameWidgetConstructor::FContentBrowserListViewNameWidgetC
 {
 }
 
-TSharedPtr<SWidget> FContentBrowserListViewNameWidgetConstructor::CreateWidget(IEditorDataStorageProvider* DataStorage,
-	IEditorDataStorageUiProvider* DataStorageUi, UE::Editor::DataStorage::RowHandle TargetRow, UE::Editor::DataStorage::RowHandle WidgetRow,
+TSharedPtr<SWidget> FContentBrowserListViewNameWidgetConstructor::CreateWidget(UE::Editor::DataStorage::ICoreProvider* DataStorage,
+	UE::Editor::DataStorage::IUiProvider* DataStorageUi, UE::Editor::DataStorage::RowHandle TargetRow, UE::Editor::DataStorage::RowHandle WidgetRow,
 	const UE::Editor::DataStorage::FMetaDataView& Arguments)
 {
+	using namespace UE::Editor::DataStorage;
+
 	if (DataStorage->IsRowAvailable(TargetRow))
 	{
 		static FName LabelThumbnailPurposeName = TEXT("General.Cell");
@@ -69,10 +71,10 @@ TSharedPtr<SWidget> FContentBrowserListViewNameWidgetConstructor::CreateWidget(I
 
 		// TODO: Arguments need to be developed further to customize the Thumbnail Config if needed
 		TArray<TWeakObjectPtr<const UScriptStruct>> ThumbnailColumns = GetThumbnailColumns();
-		DataStorageUi->CreateWidgetConstructors(LabelThumbnailPurposeName, IEditorDataStorageUiProvider::EMatchApproach::ExactMatch, ThumbnailColumns, Arguments, AssignThumbnailWidgetToColumn);
+		DataStorageUi->CreateWidgetConstructors(LabelThumbnailPurposeName, IUiProvider::EMatchApproach::ExactMatch, ThumbnailColumns, Arguments, AssignThumbnailWidgetToColumn);
 
 		TArray<TWeakObjectPtr<const UScriptStruct>> LabelColumns = GetLabelColumns();
-		DataStorageUi->CreateWidgetConstructors(LabelThumbnailPurposeName, IEditorDataStorageUiProvider::EMatchApproach::ExactMatch, LabelColumns, Arguments, AssignLabelWidgetToColumn);
+		DataStorageUi->CreateWidgetConstructors(LabelThumbnailPurposeName, IUiProvider::EMatchApproach::ExactMatch, LabelColumns, Arguments, AssignLabelWidgetToColumn);
 
 		static const FMargin ColumnItemPadding(8, 0);
 		constexpr float ThumbnailNameHorizontalPadding = 8.f;
diff --git a/Engine/Plugins/Experimental/EditorDataStorageFeatures/Source/TedsAssetData/Private/Widgets/ContentBrowserListViewNameWidget.h b/Engine/Plugins/Experimental/EditorDataStorageFeatures/Source/TedsAssetData/Private/Widgets/ContentBrowserListViewNameWidget.h
--- a/Engine/Plugins/Experimental/EditorDataStorageFeatures/Source/TedsAssetData/Private/Widgets/ContentBrowserListViewNameWidget.h
+++ b/Engine/Plugins/Experimental/EditorDataStorageFeatures/Source/TedsAssetData/Private/Widgets/ContentBrowserListViewNameWidget.h
@@ -9,7 +9,10 @@
 
 #include "ContentBrowserListViewNameWidget.generated.h"
 
-class IEditorDataStorageProvider;
+namespace UE::Editor::DataStorage
+{
+	class ICoreProvider;
+} // namespace UE::Editor::DataStorage
 class UScriptStruct;
 
 UCLASS()
@@ -22,10 +25,10 @@ public:
 
 	~UContentBrowserListViewNameWidgetFactory() override = default;
 
-	TEDSASSETDATA_API virtual void RegisterWidgetPurposes(IEditorDataStorageUiProvider& DataStorageUi) const override;
+	TEDSASSETDATA_API virtual void RegisterWidgetPurposes(UE::Editor::DataStorage::IUiProvider& DataStorageUi) const override;
 
-	TEDSASSETDATA_API void RegisterWidgetConstructors(IEditorDataStorageProvider& DataStorage,
-		IEditorDataStorageUiProvider& DataStorageUi) const override;
+	TEDSASSETDATA_API void RegisterWidgetConstructors(UE::Editor::DataStorage::ICoreProvider& DataStorage,
+		UE::Editor::DataStorage::IUiProvider& DataStorageUi) const override;
 };
 
 // Content Browser Label + Thumbnail widget
@@ -40,8 +43,8 @@ public:
 	~FContentBrowserListViewNameWidgetConstructor() override = default;
 
 	TEDSASSETDATA_API virtual TSharedPtr<SWidget> CreateWidget(
-		IEditorDataStorageProvider* DataStorage,
-		IEditorDataStorageUiProvider* DataStorageUi,
+		UE::Editor::DataStorage::ICoreProvider* DataStorage,
+		UE::Editor::DataStorage::IUiProvider* DataStorageUi,
 		UE::Editor::DataStorage::RowHandle TargetRow,
 		UE::Editor::DataStorage::RowHandle WidgetRow, 
 		const UE::Editor::DataStorage::FMetaDataView& Arguments) override;
diff --git a/Engine/Plugins/Experimental/EditorDataStorageFeatures/Source/TedsAssetData/Private/Widgets/DynamicAssetDataColumnBaseWidget.cpp b/Engine/Plugins/Experimental/EditorDataStorageFeatures/Source/TedsAssetData/Private/Widgets/DynamicAssetDataColumnBaseWidget.cpp
--- a/Engine/Plugins/Experimental/EditorDataStorageFeatures/Source/TedsAssetData/Private/Widgets/DynamicAssetDataColumnBaseWidget.cpp
+++ b/Engine/Plugins/Experimental/EditorDataStorageFeatures/Source/TedsAssetData/Private/Widgets/DynamicAssetDataColumnBaseWidget.cpp
@@ -32,8 +32,8 @@ namespace UE::DynamicColumn::Utilities
 	}	
 }
 
-void UDynamicAssetDataColumnBaseWidgetFactory::RegisterWidgetConstructors(IEditorDataStorageProvider& DataStorage,
-	IEditorDataStorageUiProvider& DataStorageUi) const
+void UDynamicAssetDataColumnBaseWidgetFactory::RegisterWidgetConstructors(UE::Editor::DataStorage::ICoreProvider& DataStorage,
+	UE::Editor::DataStorage::IUiProvider& DataStorageUi) const
 {
 	using namespace UE::Editor::DataStorage::Queries;
 	DataStorageUi.RegisterWidgetFactory<FDynamicAssetDataColumnBaseWidgetConstructor>(TEXT("General.Cell"), TColumn<FItemStringAttributeColumn_Experimental>("ParentClass"));
@@ -49,8 +49,9 @@ FDynamicAssetDataColumnBaseWidgetConstructor::FDynamicAssetDataColumnBaseWidgetC
 {
 }
 
-TSharedPtr<SWidget> FDynamicAssetDataColumnBaseWidgetConstructor::CreateWidget(IEditorDataStorageProvider* DataStorage,
-	IEditorDataStorageUiProvider* DataStorageUi, UE::Editor::DataStorage::RowHandle TargetRow, UE::Editor::DataStorage::RowHandle WidgetRow,
+TSharedPtr<SWidget> FDynamicAssetDataColumnBaseWidgetConstructor::CreateWidget(UE::Editor::DataStorage::ICoreProvider* DataStorage,
+	UE::Editor::DataStorage::IUiProvider* DataStorageUi,
+	UE::Editor::DataStorage::RowHandle TargetRow, UE::Editor::DataStorage::RowHandle WidgetRow,
 	const UE::Editor::DataStorage::FMetaDataView& Arguments)
 {
 	using namespace UE::Editor::DataStorage;
diff --git a/Engine/Plugins/Experimental/EditorDataStorageFeatures/Source/TedsAssetData/Private/Widgets/DynamicAssetDataColumnBaseWidget.h b/Engine/Plugins/Experimental/EditorDataStorageFeatures/Source/TedsAssetData/Private/Widgets/DynamicAssetDataColumnBaseWidget.h
--- a/Engine/Plugins/Experimental/EditorDataStorageFeatures/Source/TedsAssetData/Private/Widgets/DynamicAssetDataColumnBaseWidget.h
+++ b/Engine/Plugins/Experimental/EditorDataStorageFeatures/Source/TedsAssetData/Private/Widgets/DynamicAssetDataColumnBaseWidget.h
@@ -7,8 +7,6 @@
 
 #include "DynamicAssetDataColumnBaseWidget.generated.h"
 
-class IEditorDataStorageProvider;
-
 UCLASS()
 class UDynamicAssetDataColumnBaseWidgetFactory : public UEditorDataStorageFactory
 {
@@ -17,8 +15,8 @@ class UDynamicAssetDataColumnBaseWidgetFactory : public UEditorDataStorageFactor
 public:
 	~UDynamicAssetDataColumnBaseWidgetFactory() override = default;
 
-	TEDSASSETDATA_API void RegisterWidgetConstructors(IEditorDataStorageProvider& DataStorage,
-		IEditorDataStorageUiProvider& DataStorageUi) const override;
+	TEDSASSETDATA_API void RegisterWidgetConstructors(UE::Editor::DataStorage::ICoreProvider& DataStorage,
+		UE::Editor::DataStorage::IUiProvider& DataStorageUi) const override;
 };
 
 USTRUCT()
@@ -31,8 +29,8 @@ public:
 	~FDynamicAssetDataColumnBaseWidgetConstructor() override = default;
 
 	TEDSASSETDATA_API virtual TSharedPtr<SWidget> CreateWidget(
-		IEditorDataStorageProvider* DataStorage,
-		IEditorDataStorageUiProvider* DataStorageUi,
+		UE::Editor::DataStorage::ICoreProvider* DataStorage,
+		UE::Editor::DataStorage::IUiProvider* DataStorageUi,
 		UE::Editor::DataStorage::RowHandle TargetRow,
 		UE::Editor::DataStorage::RowHandle WidgetRow,
 		const UE::Editor::DataStorage::FMetaDataView& Arguments) override;
diff --git a/Engine/Plugins/Experimental/EditorDataStorageFeatures/Source/TedsAssetData/Private/Widgets/StaticMeshTrianglesWidget.cpp b/Engine/Plugins/Experimental/EditorDataStorageFeatures/Source/TedsAssetData/Private/Widgets/StaticMeshTrianglesWidget.cpp
--- a/Engine/Plugins/Experimental/EditorDataStorageFeatures/Source/TedsAssetData/Private/Widgets/StaticMeshTrianglesWidget.cpp
+++ b/Engine/Plugins/Experimental/EditorDataStorageFeatures/Source/TedsAssetData/Private/Widgets/StaticMeshTrianglesWidget.cpp
@@ -8,8 +8,8 @@
 #include "UObject/ObjectMacros.h"
 #include "Widgets/Text/STextBlock.h"
 
-void UStaticMeshTrianglesWidgetFactory::RegisterWidgetConstructors(IEditorDataStorageProvider& DataStorage,
-                                                        IEditorDataStorageUiProvider& DataStorageUi) const
+void UStaticMeshTrianglesWidgetFactory::RegisterWidgetConstructors(UE::Editor::DataStorage::ICoreProvider& DataStorage,
+	UE::Editor::DataStorage::IUiProvider& DataStorageUi) const
 {
 	using namespace UE::Editor::DataStorage::Queries;
 	DataStorageUi.RegisterWidgetFactory<FStaticMeshTrianglesWidgetConstructor>(TEXT("General.Cell"), TColumn<FItemStringAttributeColumn_Experimental>("Triangles"));
@@ -20,8 +20,9 @@ FStaticMeshTrianglesWidgetConstructor::FStaticMeshTrianglesWidgetConstructor()
 {
 }
 
-TSharedPtr<SWidget> FStaticMeshTrianglesWidgetConstructor::CreateWidget(IEditorDataStorageProvider* DataStorage,
-	IEditorDataStorageUiProvider* DataStorageUi, UE::Editor::DataStorage::RowHandle TargetRow, UE::Editor::DataStorage::RowHandle WidgetRow,
+TSharedPtr<SWidget> FStaticMeshTrianglesWidgetConstructor::CreateWidget(UE::Editor::DataStorage::ICoreProvider* DataStorage,
+	UE::Editor::DataStorage::IUiProvider* DataStorageUi,
+	UE::Editor::DataStorage::RowHandle TargetRow, UE::Editor::DataStorage::RowHandle WidgetRow,
 	const UE::Editor::DataStorage::FMetaDataView& Arguments)
 {
 	UE::Editor::DataStorage::FAttributeBinder Binder(TargetRow, DataStorage);
diff --git a/Engine/Plugins/Experimental/EditorDataStorageFeatures/Source/TedsAssetData/Public/Widgets/AssetDataLabelWidget.h b/Engine/Plugins/Experimental/EditorDataStorageFeatures/Source/TedsAssetData/Public/Widgets/AssetDataLabelWidget.h
--- a/Engine/Plugins/Experimental/EditorDataStorageFeatures/Source/TedsAssetData/Public/Widgets/AssetDataLabelWidget.h
+++ b/Engine/Plugins/Experimental/EditorDataStorageFeatures/Source/TedsAssetData/Public/Widgets/AssetDataLabelWidget.h
@@ -9,7 +9,11 @@
 
 #include "AssetDataLabelWidget.generated.h"
 
-class IEditorDataStorageProvider;
+namespace UE::Editor::DataStorage
+{
+	class ICoreProvider;
+} // UE::Editor::DataStorage
+
 class UScriptStruct;
 
 UCLASS()
@@ -20,8 +24,8 @@ class UAssetDataLabelWidgetFactory : public UEditorDataStorageFactory
 public:
 	~UAssetDataLabelWidgetFactory() override = default;
 
-	TEDSASSETDATA_API void RegisterWidgetConstructors(IEditorDataStorageProvider& DataStorage,
-		IEditorDataStorageUiProvider& DataStorageUi) const override;
+	TEDSASSETDATA_API void RegisterWidgetConstructors(UE::Editor::DataStorage::ICoreProvider& DataStorage,
+		UE::Editor::DataStorage::IUiProvider& DataStorageUi) const override;
 };
 
 // Label widget for assets in TEDS
@@ -36,8 +40,8 @@ public:
 	~FAssetDataLabelWidgetConstructor() override = default;
 
 	TEDSASSETDATA_API virtual TSharedPtr<SWidget> CreateWidget(
-		IEditorDataStorageProvider* DataStorage,
-		IEditorDataStorageUiProvider* DataStorageUi,
+		UE::Editor::DataStorage::ICoreProvider* DataStorage,
+		UE::Editor::DataStorage::IUiProvider* DataStorageUi,
 		UE::Editor::DataStorage::RowHandle TargetRow,
 		UE::Editor::DataStorage::RowHandle WidgetRow, 
 		const UE::Editor::DataStorage::FMetaDataView& Arguments) override;
diff --git a/Engine/Plugins/Experimental/EditorDataStorageFeatures/Source/TedsContentBrowser/Private/TedsContentBrowserModule.cpp b/Engine/Plugins/Experimental/EditorDataStorageFeatures/Source/TedsContentBrowser/Private/TedsContentBrowserModule.cpp
--- a/Engine/Plugins/Experimental/EditorDataStorageFeatures/Source/TedsContentBrowser/Private/TedsContentBrowserModule.cpp
+++ b/Engine/Plugins/Experimental/EditorDataStorageFeatures/Source/TedsContentBrowser/Private/TedsContentBrowserModule.cpp
@@ -335,7 +335,8 @@ namespace UE::Editor::ContentBrowser
 
 	FTedsContentBrowserViewExtender::FTedsContentBrowserViewExtender()
 	{
-		DataStorage = GetMutableDataStorageFeature<IEditorDataStorageProvider>(StorageFeatureName);
+		using namespace UE::Editor::DataStorage;
+		DataStorage = GetMutableDataStorageFeature<ICoreProvider>(StorageFeatureName);
 		
 		RowQueryStack = MakeShared<QueryStack::FRowViewNode>(FRowHandleArrayView(Rows.GetData(), Rows.Num(),
 			FRowHandleArrayView::EFlags::IsSorted | FRowHandleArrayView::EFlags::IsUnique));
diff --git a/Engine/Plugins/Experimental/EditorDataStorageFeatures/Source/TedsContentBrowser/Private/TedsContentBrowserModule.h b/Engine/Plugins/Experimental/EditorDataStorageFeatures/Source/TedsContentBrowser/Private/TedsContentBrowserModule.h
--- a/Engine/Plugins/Experimental/EditorDataStorageFeatures/Source/TedsContentBrowser/Private/TedsContentBrowserModule.h
+++ b/Engine/Plugins/Experimental/EditorDataStorageFeatures/Source/TedsContentBrowser/Private/TedsContentBrowserModule.h
@@ -8,12 +8,11 @@
 #include "Modules/ModuleInterface.h"
 #include "Templates/SharedPointer.h"
 
-class IEditorDataStorageProvider;
-
 namespace UE::Editor::DataStorage
 {
 	class FQueryStackNode_RowView;
 	class ITableViewer;
+	class ICoreProvider;
 
 	namespace QueryStack
 	{
@@ -87,7 +86,7 @@ namespace UE::Editor::ContentBrowser
 	private:
 
 		// Ptr to the data storage interface
-		IEditorDataStorageProvider* DataStorage;
+		DataStorage::ICoreProvider* DataStorage;
 		
 		// The actual table viewer widget
 		TSharedPtr<DataStorage::ITableViewer> TableViewer;
diff --git a/Engine/Plugins/Experimental/EditorDataStorageFeatures/Source/TedsContentBrowser/Private/Widgets/TedsContentBrowserAssetViewWidget.cpp b/Engine/Plugins/Experimental/EditorDataStorageFeatures/Source/TedsContentBrowser/Private/Widgets/TedsContentBrowserAssetViewWidget.cpp
--- a/Engine/Plugins/Experimental/EditorDataStorageFeatures/Source/TedsContentBrowser/Private/Widgets/TedsContentBrowserAssetViewWidget.cpp
+++ b/Engine/Plugins/Experimental/EditorDataStorageFeatures/Source/TedsContentBrowser/Private/Widgets/TedsContentBrowserAssetViewWidget.cpp
@@ -20,7 +20,7 @@ public:
 		: _InitParams()
 		{}
 		SLATE_ARGUMENT(UE::Editor::ContentBrowser::FTableViewerInitParams, InitParams)
-		SLATE_ARGUMENT(IEditorDataStorageProvider*, Storage)
+		SLATE_ARGUMENT(UE::Editor::DataStorage::ICoreProvider*, Storage)
 	SLATE_END_ARGS()
 
 	void Construct(const FArguments& Arguments)
@@ -77,7 +77,7 @@ public:
 
 		Storage->RunQuery(QueryHandle,
 			Queries::CreateDirectQueryCallbackBinding([&NewRows]
-				(const IEditorDataStorageProvider::IDirectQueryContext& Context, const RowHandle*)
+				(const IDirectQueryContext& Context, const RowHandle*)
 				{
 					NewRows.Append(Context.GetRowHandles());
 				}));
@@ -105,7 +105,7 @@ private:
 	TArray<UE::Editor::DataStorage::RowHandle> Rows;
 	TSet<UE::Editor::DataStorage::RowHandle> Rows_Set;
 	TSharedPtr<UE::Editor::DataStorage::QueryStack::FRowViewNode> RowView;
-	IEditorDataStorageProvider* Storage = nullptr;
+	UE::Editor::DataStorage::ICoreProvider* Storage = nullptr;
 	
 };
 
@@ -114,14 +114,14 @@ namespace TedsContentBrowser::Private
 	static const FName WidgetPurpose(TEXT("ContentBrowser.AssetView"));
 }
 
-void UContentBrowserAssetViewWidgetFactory::RegisterWidgetPurposes(IEditorDataStorageUiProvider& DataStorageUi) const
+void UContentBrowserAssetViewWidgetFactory::RegisterWidgetPurposes(UE::Editor::DataStorage::IUiProvider& DataStorageUi) const
 {
-	DataStorageUi.RegisterWidgetPurpose(TedsContentBrowser::Private::WidgetPurpose, IEditorDataStorageUiProvider::EPurposeType::UniqueByName,
+	DataStorageUi.RegisterWidgetPurpose(TedsContentBrowser::Private::WidgetPurpose, UE::Editor::DataStorage::IUiProvider::EPurposeType::UniqueByName,
 		LOCTEXT("ContentBrowserAssetView_PurposeDescription", "Widget that displays a table viewer in the content browser"));
 }
 
-void UContentBrowserAssetViewWidgetFactory::RegisterWidgetConstructors(IEditorDataStorageProvider& DataStorage,
-	IEditorDataStorageUiProvider& DataStorageUi) const
+void UContentBrowserAssetViewWidgetFactory::RegisterWidgetConstructors(UE::Editor::DataStorage::ICoreProvider& DataStorage,
+	UE::Editor::DataStorage::IUiProvider& DataStorageUi) const
 {
 	using namespace UE::Editor::DataStorage::Queries;
 	DataStorageUi.RegisterWidgetFactory<FContentBrowserAssetViewWidgetConstructor>(TedsContentBrowser::Private::WidgetPurpose);
@@ -137,8 +137,9 @@ FContentBrowserAssetViewWidgetConstructor::FContentBrowserAssetViewWidgetConstru
 {
 }
 
-TSharedPtr<SWidget> FContentBrowserAssetViewWidgetConstructor::CreateWidget(IEditorDataStorageProvider* DataStorage,
-	IEditorDataStorageUiProvider* DataStorageUi, UE::Editor::DataStorage::RowHandle TargetRow, UE::Editor::DataStorage::RowHandle WidgetRow,
+TSharedPtr<SWidget> FContentBrowserAssetViewWidgetConstructor::CreateWidget(UE::Editor::DataStorage::ICoreProvider* DataStorage,
+	UE::Editor::DataStorage::IUiProvider* DataStorageUi,
+	UE::Editor::DataStorage::RowHandle TargetRow, UE::Editor::DataStorage::RowHandle WidgetRow,
 	const UE::Editor::DataStorage::FMetaDataView& Arguments)
 {
 	using namespace UE::Editor::DataStorage::Queries;
diff --git a/Engine/Plugins/Experimental/EditorDataStorageFeatures/Source/TedsContentBrowser/Private/Widgets/TedsContentBrowserAssetViewWidget.h b/Engine/Plugins/Experimental/EditorDataStorageFeatures/Source/TedsContentBrowser/Private/Widgets/TedsContentBrowserAssetViewWidget.h
--- a/Engine/Plugins/Experimental/EditorDataStorageFeatures/Source/TedsContentBrowser/Private/Widgets/TedsContentBrowserAssetViewWidget.h
+++ b/Engine/Plugins/Experimental/EditorDataStorageFeatures/Source/TedsContentBrowser/Private/Widgets/TedsContentBrowserAssetViewWidget.h
@@ -9,7 +9,6 @@
 
 #include "TedsContentBrowserAssetViewWidget.generated.h"
 
-class IEditorDataStorageProvider;
 class UScriptStruct;
 
 UCLASS()
@@ -20,10 +19,10 @@ class UContentBrowserAssetViewWidgetFactory : public UEditorDataStorageFactory
 public:
 	~UContentBrowserAssetViewWidgetFactory() override = default;
 
-	TEDSCONTENTBROWSER_API virtual void RegisterWidgetPurposes(IEditorDataStorageUiProvider& DataStorageUi) const override;
+	TEDSCONTENTBROWSER_API virtual void RegisterWidgetPurposes(UE::Editor::DataStorage::IUiProvider& DataStorageUi) const override;
 
-	TEDSCONTENTBROWSER_API void RegisterWidgetConstructors(IEditorDataStorageProvider& DataStorage,
-		IEditorDataStorageUiProvider& DataStorageUi) const override;
+	TEDSCONTENTBROWSER_API void RegisterWidgetConstructors(UE::Editor::DataStorage::ICoreProvider& DataStorage,
+		UE::Editor::DataStorage::IUiProvider& DataStorageUi) const override;
 };
 
 // Default asset view widget shown by a content source - the query and other init params are populated by the content source itself
@@ -38,8 +37,8 @@ public:
 	~FContentBrowserAssetViewWidgetConstructor() override = default;
 
 	TEDSCONTENTBROWSER_API virtual TSharedPtr<SWidget> CreateWidget(
-		IEditorDataStorageProvider* DataStorage,
-		IEditorDataStorageUiProvider* DataStorageUi,
+		UE::Editor::DataStorage::ICoreProvider* DataStorage,
+		UE::Editor::DataStorage::IUiProvider* DataStorageUi,
 		UE::Editor::DataStorage::RowHandle TargetRow,
 		UE::Editor::DataStorage::RowHandle WidgetRow, 
 		const UE::Editor::DataStorage::FMetaDataView& Arguments) override;
diff --git a/Engine/Plugins/Experimental/EditorDataStorageFeatures/Source/TedsDebugger/Private/QueryEditor/TedsQueryEditorModel.cpp b/Engine/Plugins/Experimental/EditorDataStorageFeatures/Source/TedsDebugger/Private/QueryEditor/TedsQueryEditorModel.cpp
--- a/Engine/Plugins/Experimental/EditorDataStorageFeatures/Source/TedsDebugger/Private/QueryEditor/TedsQueryEditorModel.cpp
+++ b/Engine/Plugins/Experimental/EditorDataStorageFeatures/Source/TedsDebugger/Private/QueryEditor/TedsQueryEditorModel.cpp
@@ -24,7 +24,7 @@ namespace UE::Editor::DataStorage::Debug::QueryEditor
 		Id = NAME_None;
 	}
 
-	FTedsQueryEditorModel::FTedsQueryEditorModel(IEditorDataStorageProvider& InDataStorageProvider)
+	FTedsQueryEditorModel::FTedsQueryEditorModel(ICoreProvider& InDataStorageProvider)
 		: EditorDataStorageProvider(InDataStorageProvider)
 	{
 	}
@@ -35,12 +35,12 @@ namespace UE::Editor::DataStorage::Debug::QueryEditor
 		ModelChangedDelegate.Broadcast();
 	}
 
-	IEditorDataStorageProvider& FTedsQueryEditorModel::GetTedsInterface()
+	ICoreProvider& FTedsQueryEditorModel::GetTedsInterface()
 	{
 		return EditorDataStorageProvider;
 	}
 	
-	const IEditorDataStorageProvider& FTedsQueryEditorModel::GetTedsInterface() const
+	const ICoreProvider& FTedsQueryEditorModel::GetTedsInterface() const
 	{
 		return EditorDataStorageProvider;
 	}
@@ -85,21 +85,21 @@ namespace UE::Editor::DataStorage::Debug::QueryEditor
 				}
 				else
 				{
-					IEditorDataStorageProvider::FQueryDescription::EOperatorType OperatorType = [](const FConditionEntryInternal& Entry)
+					FQueryDescription::EOperatorType OperatorType = [](const FConditionEntryInternal& Entry)
 					{
 						switch(Entry.OperatorType)
 						{
 						case EOperatorType::All:
-							return IEditorDataStorageProvider::FQueryDescription::EOperatorType::SimpleAll;
+							return FQueryDescription::EOperatorType::SimpleAll;
 						case EOperatorType::Any:
-							return IEditorDataStorageProvider::FQueryDescription::EOperatorType::SimpleAny;
+							return FQueryDescription::EOperatorType::SimpleAny;
 						case EOperatorType::None:
-							return IEditorDataStorageProvider::FQueryDescription::EOperatorType::SimpleNone;
+							return FQueryDescription::EOperatorType::SimpleNone;
 						case EOperatorType::Invalid:
 						case EOperatorType::Unset:
 						default:
 							check(false);
-							return IEditorDataStorageProvider::FQueryDescription::EOperatorType::Max;
+							return FQueryDescription::EOperatorType::Max;
 						}
 					}(Entry);
 				
@@ -109,7 +109,7 @@ namespace UE::Editor::DataStorage::Debug::QueryEditor
 			}
 		}
 
-		Description.Action = IEditorDataStorageProvider::FQueryDescription::EActionType::Select;
+		Description.Action = FQueryDescription::EActionType::Select;
 
 		return Description;
 	}
@@ -124,14 +124,14 @@ namespace UE::Editor::DataStorage::Debug::QueryEditor
 		for (int32 Index = 0; Index < SelectionTypeCount; ++Index)
 		{
 			Description.ConditionOperators.AddZeroed_GetRef().Type = Description.SelectionTypes[Index];
-			Description.ConditionTypes.Add(IEditorDataStorageProvider::FQueryDescription::EOperatorType::SimpleAll);
+			Description.ConditionTypes.Add(FQueryDescription::EOperatorType::SimpleAll);
 		}
 
 		Description.SelectionTypes.Empty();
 		Description.SelectionMetaData.Empty();
 		Description.SelectionAccessTypes.Empty();
 
-		Description.Action = IEditorDataStorageProvider::FQueryDescription::EActionType::Count;
+		Description.Action = FQueryDescription::EActionType::Count;
 
 		return Description;
 	}
diff --git a/Engine/Plugins/Experimental/EditorDataStorageFeatures/Source/TedsDebugger/Private/QueryEditor/TedsQueryEditorModel.h b/Engine/Plugins/Experimental/EditorDataStorageFeatures/Source/TedsDebugger/Private/QueryEditor/TedsQueryEditorModel.h
--- a/Engine/Plugins/Experimental/EditorDataStorageFeatures/Source/TedsDebugger/Private/QueryEditor/TedsQueryEditorModel.h
+++ b/Engine/Plugins/Experimental/EditorDataStorageFeatures/Source/TedsDebugger/Private/QueryEditor/TedsQueryEditorModel.h
@@ -4,7 +4,10 @@
 #include "Elements/Framework/TypedElementMetaData.h"
 #include "UObject/Class.h"
 
-class IEditorDataStorageProvider;
+namespace UE::Editor::DataStorage
+{
+	class ICoreProvider;
+}
 
 namespace UE::Editor::DataStorage::Debug::QueryEditor
 {
@@ -48,11 +51,11 @@ namespace UE::Editor::DataStorage::Debug::QueryEditor
 	class FTedsQueryEditorModel
 	{
 	public:
-		explicit FTedsQueryEditorModel(IEditorDataStorageProvider& InDataStorageProvider);
+		explicit FTedsQueryEditorModel(ICoreProvider& InDataStorageProvider);
 
 		void Reset();
-		IEditorDataStorageProvider& GetTedsInterface();
-		const IEditorDataStorageProvider& GetTedsInterface() const;
+		ICoreProvider& GetTedsInterface();
+		const ICoreProvider& GetTedsInterface() const;
 		
 		FQueryDescription GenerateQueryDescription();
 		// Special function to generate a description that puts the Select elements as All Conditions
@@ -107,7 +110,7 @@ namespace UE::Editor::DataStorage::Debug::QueryEditor
 		FConditionEntryInternal* FindEntryByHandle(FConditionEntryHandle Handle);
 		const FConditionEntryInternal* FindEntryByHandle(FConditionEntryHandle Handle) const;
 
-		IEditorDataStorageProvider& EditorDataStorageProvider;
+		ICoreProvider& EditorDataStorageProvider;
 
 		TArray<FConditionEntryInternal> Conditions;
 		FTedsQueryEditorModel_ModelChanged ModelChangedDelegate;
diff --git a/Engine/Plugins/Experimental/EditorDataStorageFeatures/Source/TedsDebugger/Private/STedsDebugger.cpp b/Engine/Plugins/Experimental/EditorDataStorageFeatures/Source/TedsDebugger/Private/STedsDebugger.cpp
--- a/Engine/Plugins/Experimental/EditorDataStorageFeatures/Source/TedsDebugger/Private/STedsDebugger.cpp
+++ b/Engine/Plugins/Experimental/EditorDataStorageFeatures/Source/TedsDebugger/Private/STedsDebugger.cpp
@@ -29,7 +29,7 @@ STedsDebugger::~STedsDebugger()
 {
 	if (AreEditorDataStorageFeaturesEnabled())
 	{
-		GetMutableDataStorageFeature<IEditorDataStorageProvider>(StorageFeatureName)->UnregisterQuery(TableViewerQuery);
+		GetMutableDataStorageFeature<ICoreProvider>(StorageFeatureName)->UnregisterQuery(TableViewerQuery);
 	}
 }
 
@@ -129,7 +129,7 @@ TSharedRef<SDockTab> STedsDebugger::SpawnQueryEditorTab(const FSpawnTabArgs& Arg
 	{
 		if (AreEditorDataStorageFeaturesEnabled())
 		{
-			IEditorDataStorageProvider* DataStorageInterface = GetMutableDataStorageFeature<IEditorDataStorageProvider>(StorageFeatureName);
+			ICoreProvider* DataStorageInterface = GetMutableDataStorageFeature<ICoreProvider>(StorageFeatureName);
 			QueryEditorModel = MakeUnique<QueryEditor::FTedsQueryEditorModel>(*DataStorageInterface);
 		}
 	}
diff --git a/Engine/Plugins/Experimental/EditorDataStorageFeatures/Source/TedsDebugger/Private/Widgets/QueryEditor/TedsQueryEditorResultsView.cpp b/Engine/Plugins/Experimental/EditorDataStorageFeatures/Source/TedsDebugger/Private/Widgets/QueryEditor/TedsQueryEditorResultsView.cpp
--- a/Engine/Plugins/Experimental/EditorDataStorageFeatures/Source/TedsDebugger/Private/Widgets/QueryEditor/TedsQueryEditorResultsView.cpp
+++ b/Engine/Plugins/Experimental/EditorDataStorageFeatures/Source/TedsDebugger/Private/Widgets/QueryEditor/TedsQueryEditorResultsView.cpp
@@ -79,7 +79,7 @@ namespace UE::Editor::DataStorage::Debug::QueryEditor
 				SNew(STextBlock)
 				.Text_Lambda([this]()
 				{
-					auto& TedsInterface = Model->GetTedsInterface();
+					ICoreProvider& TedsInterface = Model->GetTedsInterface();
 					FQueryResult QueryResult = TedsInterface.RunQuery(CountQueryHandle);
 					if (QueryResult.Completed == FQueryResult::ECompletion::Fully)
 					{
@@ -104,7 +104,7 @@ namespace UE::Editor::DataStorage::Debug::QueryEditor
 	{
 		using namespace UE::Editor::DataStorage::Queries;
 
-		IEditorDataStorageProvider& TedsInterface = Model->GetTedsInterface();
+		ICoreProvider& TedsInterface = Model->GetTedsInterface();
 
 		if(bModelDirty)
 		{
@@ -155,7 +155,7 @@ namespace UE::Editor::DataStorage::Debug::QueryEditor
 			NewTableViewerRows_Set.Reserve(TableViewerRows_Set.Num());
 
 			FQueryResult QueryResult = Model->GetTedsInterface().RunQuery(TableViewerQueryHandle,
-				CreateDirectQueryCallbackBinding([&NewTableViewerRows_Set](const IEditorDataStorageProvider::IDirectQueryContext& Context, const RowHandle*)
+				CreateDirectQueryCallbackBinding([&NewTableViewerRows_Set](const IDirectQueryContext& Context, const RowHandle*)
 				{
 					NewTableViewerRows_Set.Append(Context.GetRowHandles());
 				}));
@@ -189,7 +189,7 @@ namespace UE::Editor::DataStorage::Debug::QueryEditor
 			return false;
 		};
 	
-		IEditorDataStorageUiProvider* StorageUi = GetMutableDataStorageFeature<IEditorDataStorageUiProvider>(UiFeatureName);
+		IUiProvider* StorageUi = GetMutableDataStorageFeature<IUiProvider>(UiFeatureName);
 		checkf(StorageUi, TEXT("SResultsView created before data storage interfaces were initialized."))
 
 		StorageUi->CreateWidgetConstructors(TEXT("General.Cell.RowHandle"), FMetaDataView(), AssignWidgetToColumn);
diff --git a/Engine/Plugins/Experimental/EditorDataStorageFeatures/Source/TedsDebugger/Private/Widgets/RowReferenceWidget.cpp b/Engine/Plugins/Experimental/EditorDataStorageFeatures/Source/TedsDebugger/Private/Widgets/RowReferenceWidget.cpp
--- a/Engine/Plugins/Experimental/EditorDataStorageFeatures/Source/TedsDebugger/Private/Widgets/RowReferenceWidget.cpp
+++ b/Engine/Plugins/Experimental/EditorDataStorageFeatures/Source/TedsDebugger/Private/Widgets/RowReferenceWidget.cpp
@@ -20,7 +20,7 @@
 
 namespace UE::Editor::DataStorage::Debug::Private
 {
-	void OnNavigateHyperlink(IEditorDataStorageProvider* DataStorage,  RowHandle TargetRowHandle, RowHandle UiRowHandle)
+	void OnNavigateHyperlink(ICoreProvider* DataStorage,  RowHandle TargetRowHandle, RowHandle UiRowHandle)
 	{
 		const FTedsOutlinerColumn* TedsOutlinerColumn = DataStorage->GetColumn<FTedsOutlinerColumn>(UiRowHandle);
 
@@ -56,7 +56,8 @@ URowReferenceWidgetFactory::~URowReferenceWidgetFactory()
 {
 }
 
-void URowReferenceWidgetFactory::RegisterWidgetConstructors(IEditorDataStorageProvider& DataStorage, IEditorDataStorageUiProvider& DataStorageUi) const
+void URowReferenceWidgetFactory::RegisterWidgetConstructors(UE::Editor::DataStorage::ICoreProvider& DataStorage,
+	UE::Editor::DataStorage::IUiProvider& DataStorageUi) const
 {
 	using namespace UE::Editor::DataStorage::Queries;
 
@@ -70,7 +71,8 @@ FRowReferenceWidgetConstructor::FRowReferenceWidgetConstructor()
 {
 }
 
-TSharedPtr<SWidget> FRowReferenceWidgetConstructor::CreateWidget(IEditorDataStorageProvider* DataStorage, IEditorDataStorageUiProvider* DataStorageUi,
+TSharedPtr<SWidget> FRowReferenceWidgetConstructor::CreateWidget(UE::Editor::DataStorage::ICoreProvider* DataStorage,
+	UE::Editor::DataStorage::IUiProvider* DataStorageUi,
 	UE::Editor::DataStorage::RowHandle TargetRow, UE::Editor::DataStorage::RowHandle WidgetRow,
 	const UE::Editor::DataStorage::FMetaDataView& Arguments)
 {
diff --git a/Engine/Plugins/Experimental/EditorDataStorageFeatures/Source/TedsDebugger/Private/Widgets/RowReferenceWidget.h b/Engine/Plugins/Experimental/EditorDataStorageFeatures/Source/TedsDebugger/Private/Widgets/RowReferenceWidget.h
--- a/Engine/Plugins/Experimental/EditorDataStorageFeatures/Source/TedsDebugger/Private/Widgets/RowReferenceWidget.h
+++ b/Engine/Plugins/Experimental/EditorDataStorageFeatures/Source/TedsDebugger/Private/Widgets/RowReferenceWidget.h
@@ -7,7 +7,6 @@
 
 #include "RowReferenceWidget.generated.h"
 
-class IEditorDataStorageProvider;
 class SWidget;
 
 /*
@@ -21,8 +20,8 @@ class URowReferenceWidgetFactory : public UEditorDataStorageFactory
 public:
 	~URowReferenceWidgetFactory() override;
 
-	void RegisterWidgetConstructors(IEditorDataStorageProvider& DataStorage,
-		IEditorDataStorageUiProvider& DataStorageUi) const override;
+	void RegisterWidgetConstructors(UE::Editor::DataStorage::ICoreProvider& DataStorage,
+		UE::Editor::DataStorage::IUiProvider& DataStorageUi) const override;
 };
 
 USTRUCT()
@@ -35,8 +34,8 @@ public:
 	~FRowReferenceWidgetConstructor() override = default;
 
 	virtual TSharedPtr<SWidget> CreateWidget(
-		IEditorDataStorageProvider* DataStorage,
-		IEditorDataStorageUiProvider* DataStorageUi,
+		UE::Editor::DataStorage::ICoreProvider* DataStorage,
+		UE::Editor::DataStorage::IUiProvider* DataStorageUi,
 		UE::Editor::DataStorage::RowHandle TargetRow,
 		UE::Editor::DataStorage::RowHandle WidgetRow,
 		const UE::Editor::DataStorage::FMetaDataView& Arguments) override;
diff --git a/Engine/Plugins/Experimental/EditorDataStorageFeatures/Source/TedsDebugger/Private/Widgets/SlateVisualizationWidget.cpp b/Engine/Plugins/Experimental/EditorDataStorageFeatures/Source/TedsDebugger/Private/Widgets/SlateVisualizationWidget.cpp
--- a/Engine/Plugins/Experimental/EditorDataStorageFeatures/Source/TedsDebugger/Private/Widgets/SlateVisualizationWidget.cpp
+++ b/Engine/Plugins/Experimental/EditorDataStorageFeatures/Source/TedsDebugger/Private/Widgets/SlateVisualizationWidget.cpp
@@ -9,7 +9,7 @@
 #include "Widgets/SBoxPanel.h"
 
 
-void USlateVisualizationWidgetFactory::RegisterWidgetConstructors(IEditorDataStorageProvider& DataStorage, IEditorDataStorageUiProvider& DataStorageUi) const
+void USlateVisualizationWidgetFactory::RegisterWidgetConstructors(UE::Editor::DataStorage::ICoreProvider& DataStorage, UE::Editor::DataStorage::IUiProvider& DataStorageUi) const
 {
 	using namespace UE::Editor::DataStorage::Queries;
 
@@ -28,7 +28,9 @@ TSharedPtr<SWidget> FSlateVisualizationWidgetConstructor::CreateWidget(const UE:
 	return SNew(SHorizontalBox);
 }
 
-bool FSlateVisualizationWidgetConstructor::FinalizeWidget(IEditorDataStorageProvider* DataStorage, IEditorDataStorageUiProvider* DataStorageUi, UE::Editor::DataStorage::RowHandle Row, const TSharedPtr<SWidget>& Widget)
+bool FSlateVisualizationWidgetConstructor::FinalizeWidget(UE::Editor::DataStorage::ICoreProvider* DataStorage,
+	UE::Editor::DataStorage::IUiProvider* DataStorageUi,
+	UE::Editor::DataStorage::RowHandle Row, const TSharedPtr<SWidget>& Widget)
 {
 	checkf(Widget, TEXT("Referenced widget is not valid. A constructed widget may not have been cleaned up. This can "
 	"also happen if this processor is running in the same phase as the processors responsible for cleaning up old "
diff --git a/Engine/Plugins/Experimental/EditorDataStorageFeatures/Source/TedsDebugger/Private/Widgets/SlateVisualizationWidget.h b/Engine/Plugins/Experimental/EditorDataStorageFeatures/Source/TedsDebugger/Private/Widgets/SlateVisualizationWidget.h
--- a/Engine/Plugins/Experimental/EditorDataStorageFeatures/Source/TedsDebugger/Private/Widgets/SlateVisualizationWidget.h
+++ b/Engine/Plugins/Experimental/EditorDataStorageFeatures/Source/TedsDebugger/Private/Widgets/SlateVisualizationWidget.h
@@ -7,7 +7,6 @@
 
 #include "SlateVisualizationWidget.generated.h"
 
-class IEditorDataStorageProvider;
 class SWidget;
 
 /*
@@ -21,8 +20,8 @@ class USlateVisualizationWidgetFactory : public UEditorDataStorageFactory
 public:
 	~USlateVisualizationWidgetFactory() override = default;
 
-	void RegisterWidgetConstructors(IEditorDataStorageProvider& DataStorage,
-		IEditorDataStorageUiProvider& DataStorageUi) const override;
+	void RegisterWidgetConstructors(UE::Editor::DataStorage::ICoreProvider& DataStorage,
+		UE::Editor::DataStorage::IUiProvider& DataStorageUi) const override;
 };
 
 USTRUCT()
@@ -36,6 +35,7 @@ public:
 
 protected:
 	TSharedPtr<SWidget> CreateWidget(const UE::Editor::DataStorage::FMetaDataView& Arguments) override;
-	bool FinalizeWidget(IEditorDataStorageProvider* DataStorage, IEditorDataStorageUiProvider* DataStorageUi,
+	bool FinalizeWidget(UE::Editor::DataStorage::ICoreProvider* DataStorage,
+		UE::Editor::DataStorage::IUiProvider* DataStorageUi,
 		UE::Editor::DataStorage::RowHandle Row, const TSharedPtr<SWidget>& Widget) override;
 };
\ No newline at end of file
diff --git a/Engine/Plugins/Experimental/EditorDataStorageFeatures/Source/TedsOutliner/Private/Compatibility/SceneOutlinerRowHandleColumn.cpp b/Engine/Plugins/Experimental/EditorDataStorageFeatures/Source/TedsOutliner/Private/Compatibility/SceneOutlinerRowHandleColumn.cpp
--- a/Engine/Plugins/Experimental/EditorDataStorageFeatures/Source/TedsOutliner/Private/Compatibility/SceneOutlinerRowHandleColumn.cpp
+++ b/Engine/Plugins/Experimental/EditorDataStorageFeatures/Source/TedsOutliner/Private/Compatibility/SceneOutlinerRowHandleColumn.cpp
@@ -34,7 +34,7 @@ void FSceneOutlinerRowHandleColumn::CreateWidgetConstructor()
 		return false;
 	};
 	
-	IEditorDataStorageUiProvider* StorageUi = GetMutableDataStorageFeature<IEditorDataStorageUiProvider>(UiFeatureName);
+	IUiProvider* StorageUi = GetMutableDataStorageFeature<IUiProvider>(UiFeatureName);
 	checkf(StorageUi, TEXT("FSceneOutlinerRowHandleColumn created before data storage interfaces were initialized."))
 	
 	StorageUi->CreateWidgetConstructors(TEXT("General.Cell.RowHandle"), FMetaDataView(), AssignWidgetToColumn);
@@ -61,6 +61,7 @@ SHeaderRow::FColumn::FArguments FSceneOutlinerRowHandleColumn::ConstructHeaderRo
 
 const TSharedRef<SWidget> FSceneOutlinerRowHandleColumn::ConstructRowWidget(FSceneOutlinerTreeItemRef TreeItem, const STableRow<FSceneOutlinerTreeItemPtr>& Row)
 {
+	using namespace UE::Editor::DataStorage;
 	if (!TableViewerColumn)
 	{
 		return SNew(STextBlock)
@@ -71,7 +72,7 @@ const TSharedRef<SWidget> FSceneOutlinerRowHandleColumn::ConstructRowWidget(FSce
 	auto SceneOutliner = WeakSceneOutliner.Pin();
 	check(SceneOutliner.IsValid());
 	
-	UE::Editor::DataStorage::RowHandle RowHandle = UE::Editor::DataStorage::InvalidRowHandle;;
+	RowHandle RowHandle = InvalidRowHandle;
 
 	// Check any custom delegates to get the row handle of the item
 	if (GetRowHandle.IsBound())
@@ -84,7 +85,7 @@ const TSharedRef<SWidget> FSceneOutlinerRowHandleColumn::ConstructRowWidget(FSce
 		RowHandle = OutlinerTreeItem->GetRowHandle();
 	}
 
-	IEditorDataStorageProvider* Storage = TableViewerColumn->GetStorage();
+	ICoreProvider* Storage = TableViewerColumn->GetStorage();
 
 	if (ensureMsgf(Storage, TEXT("Cannot create a widget to display row handles before TEDS is initialized")))
 	{
diff --git a/Engine/Plugins/Experimental/EditorDataStorageFeatures/Source/TedsOutliner/Private/Compatibility/SceneOutlinerTedsBridge.cpp b/Engine/Plugins/Experimental/EditorDataStorageFeatures/Source/TedsOutliner/Private/Compatibility/SceneOutlinerTedsBridge.cpp
--- a/Engine/Plugins/Experimental/EditorDataStorageFeatures/Source/TedsOutliner/Private/Compatibility/SceneOutlinerTedsBridge.cpp
+++ b/Engine/Plugins/Experimental/EditorDataStorageFeatures/Source/TedsOutliner/Private/Compatibility/SceneOutlinerTedsBridge.cpp
@@ -42,7 +42,7 @@ FAutoConsoleCommand BindColumnsToSceneOutlinerConsoleCommand(
 			
 		    const FName WidgetPurposes[] = {TEXT("SceneOutliner.Cell"), TEXT("General.Cell")};
 
-			if (IEditorDataStorageProvider* DataStorage = GetMutableDataStorageFeature<IEditorDataStorageProvider>(StorageFeatureName))
+			if (ICoreProvider* DataStorage = GetMutableDataStorageFeature<ICoreProvider>(StorageFeatureName))
 			{
 				static UE::Editor::DataStorage::QueryHandle Queries[] =
 				{
@@ -120,9 +120,9 @@ public:
 	~FSceneOutlinerTedsBridge();
 
 	void Initialize(
-		IEditorDataStorageProvider& InStorage,
-		IEditorDataStorageUiProvider& InStorageUi,
-		IEditorDataStorageCompatibilityProvider& InStorageCompatibility,
+		UE::Editor::DataStorage::ICoreProvider& InStorage,
+		UE::Editor::DataStorage::IUiProvider& InStorageUi,
+		UE::Editor::DataStorage::ICompatibilityProvider& InStorageCompatibility,
 		const TSharedPtr<ISceneOutliner>& InOutliner);
 
 	void AssignQuery(UE::Editor::DataStorage::QueryHandle Query, const TConstArrayView<FName> CellWidgetPurposes);
@@ -134,9 +134,9 @@ private:
 
 	TArray<FName> AddedColumns;
 	TWeakPtr<ISceneOutliner> Outliner;
-	IEditorDataStorageProvider* Storage{ nullptr };
-	IEditorDataStorageUiProvider* StorageUi{ nullptr };
-	IEditorDataStorageCompatibilityProvider* StorageCompatibility{ nullptr };
+	UE::Editor::DataStorage::ICoreProvider* Storage{ nullptr };
+	UE::Editor::DataStorage::IUiProvider* StorageUi{ nullptr };
+	UE::Editor::DataStorage::ICompatibilityProvider* StorageCompatibility{ nullptr };
 	FTreeItemIDDealiaser Dealiaser;
 	TArray<FName> CellWidgetPurposes;
 };
@@ -146,9 +146,9 @@ class FOutlinerColumn : public ISceneOutlinerColumn
 public:
 	FOutlinerColumn(
 		UE::Editor::DataStorage::QueryHandle InQuery,
-		IEditorDataStorageProvider& InStorage, 
-		IEditorDataStorageUiProvider& InStorageUi, 
-		IEditorDataStorageCompatibilityProvider& InStorageCompatibility,
+		UE::Editor::DataStorage::ICoreProvider& InStorage, 
+		UE::Editor::DataStorage::IUiProvider& InStorageUi, 
+		UE::Editor::DataStorage::ICompatibilityProvider& InStorageCompatibility,
 		FName InNameId, 
 		TArray<TWeakObjectPtr<const UScriptStruct>> InColumnTypes,
 		TSharedPtr<FTypedElementWidgetConstructor> InHeaderWidgetConstructor,
@@ -295,7 +295,7 @@ public:
 		if(Storage.IsRowAssigned(TargetRowHandle))
 		{
 			RowWidget = TableViewerColumnImpl->ConstructRowWidget(TargetRowHandle,
-				[&](IEditorDataStorageProvider& DataStorage, const RowHandle& WidgetRow)
+				[&](ICoreProvider& DataStorage, const RowHandle& WidgetRow)
 				{
 					DataStorage.AddColumn(WidgetRow, FTedsOutlinerColumn{ .Outliner = OwningOutliner });
 				});
@@ -322,9 +322,9 @@ public:
 	// The table viewer implementation that we internally use to create our widgets
 	TUniquePtr<UE::Editor::DataStorage::FTedsTableViewerColumn> TableViewerColumnImpl;
 	
-	IEditorDataStorageProvider& Storage;
-	IEditorDataStorageUiProvider& StorageUi;
-	IEditorDataStorageCompatibilityProvider& StorageCompatibility;
+	UE::Editor::DataStorage::ICoreProvider& Storage;
+	UE::Editor::DataStorage::IUiProvider& StorageUi;
+	UE::Editor::DataStorage::ICompatibilityProvider& StorageCompatibility;
 	UE::Editor::DataStorage::QueryHandle QueryHandle;
 	UE::Editor::DataStorage::FMetaData MetaData;
 	FName NameId;
@@ -340,9 +340,10 @@ public:
 // USceneOutlinerTedsBridgeFactory
 // 
 
-void USceneOutlinerTedsBridgeFactory::RegisterWidgetPurposes(IEditorDataStorageUiProvider& DataStorageUi) const
+void USceneOutlinerTedsBridgeFactory::RegisterWidgetPurposes(UE::Editor::DataStorage::IUiProvider& DataStorageUi) const
 {
-	using PurposeType = IEditorDataStorageUiProvider::EPurposeType;
+	using namespace UE::Editor::DataStorage;
+	using PurposeType = IUiProvider::EPurposeType;
 
 	DataStorageUi.RegisterWidgetPurpose(FSceneOutlinerTedsQueryBinder::HeaderWidgetPurpose, PurposeType::UniqueByNameAndColumn,
 		LOCTEXT("HeaderWidgetPurpose", "Widgets for headers in any Scene Outliner for specific columns or column combinations."));
@@ -379,9 +380,9 @@ const FName FSceneOutlinerTedsQueryBinder::DefaultItemLabelCellWidgetPurpose(TEX
 FSceneOutlinerTedsQueryBinder::FSceneOutlinerTedsQueryBinder()
 {
 	using namespace UE::Editor::DataStorage;
-	Storage = GetMutableDataStorageFeature<IEditorDataStorageProvider>(StorageFeatureName);
-	StorageUi = GetMutableDataStorageFeature<IEditorDataStorageUiProvider>(UiFeatureName);
-	StorageCompatibility = GetMutableDataStorageFeature<IEditorDataStorageCompatibilityProvider>(CompatibilityFeatureName);
+	Storage = GetMutableDataStorageFeature<ICoreProvider>(StorageFeatureName);
+	StorageUi = GetMutableDataStorageFeature<IUiProvider>(UiFeatureName);
+	StorageCompatibility = GetMutableDataStorageFeature<ICompatibilityProvider>(CompatibilityFeatureName);
 
 	SetupDefaultColumnMapping();
 }
@@ -482,9 +483,9 @@ FSceneOutlinerTedsBridge::~FSceneOutlinerTedsBridge()
 }
 
 void FSceneOutlinerTedsBridge::Initialize(
-	IEditorDataStorageProvider& InStorage,
-	IEditorDataStorageUiProvider& InStorageUi,
-	IEditorDataStorageCompatibilityProvider& InStorageCompatibility,
+	UE::Editor::DataStorage::ICoreProvider& InStorage,
+	UE::Editor::DataStorage::IUiProvider& InStorageUi,
+	UE::Editor::DataStorage::ICompatibilityProvider& InStorageCompatibility,
 	const TSharedPtr<ISceneOutliner>& InOutliner)
 {
 	Storage = &InStorage;
@@ -506,7 +507,8 @@ FTreeItemIDDealiaser FSceneOutlinerTedsBridge::GetDealiaser()
 
 void FSceneOutlinerTedsBridge::AssignQuery(UE::Editor::DataStorage::QueryHandle Query, const TConstArrayView<FName> InCellWidgetPurposes)
 {
-	using MatchApproach = IEditorDataStorageUiProvider::EMatchApproach;
+	using namespace UE::Editor::DataStorage;
+	using MatchApproach = IUiProvider::EMatchApproach;
 	constexpr uint8 DefaultPriorityIndex = 100;
 	FSceneOutlinerTedsQueryBinder& Binder = FSceneOutlinerTedsQueryBinder::GetInstance();
 	CellWidgetPurposes = InCellWidgetPurposes;
@@ -514,12 +516,12 @@ void FSceneOutlinerTedsBridge::AssignQuery(UE::Editor::DataStorage::QueryHandle
 
 	if (TSharedPtr<ISceneOutliner> OutlinerPinned = Outliner.Pin())
 	{
-		const IEditorDataStorageProvider::FQueryDescription& Description = Storage->GetQueryDescription(Query);
+		const FQueryDescription& Description = Storage->GetQueryDescription(Query);
 		UE::Editor::DataStorage::FQueryMetaDataView MetaDataView(Description);
 
 		ClearColumns(*OutlinerPinned);
 
-		if (Description.Action == IEditorDataStorageProvider::FQueryDescription::EActionType::Select)
+		if (Description.Action == FQueryDescription::EActionType::Select)
 		{
 			int32 SelectionCount = Description.SelectionTypes.Num();
 			AddedColumns.Reset(SelectionCount);
diff --git a/Engine/Plugins/Experimental/EditorDataStorageFeatures/Source/TedsOutliner/Private/TedsOutlinerHierarchy.cpp b/Engine/Plugins/Experimental/EditorDataStorageFeatures/Source/TedsOutliner/Private/TedsOutlinerHierarchy.cpp
--- a/Engine/Plugins/Experimental/EditorDataStorageFeatures/Source/TedsOutliner/Private/TedsOutlinerHierarchy.cpp
+++ b/Engine/Plugins/Experimental/EditorDataStorageFeatures/Source/TedsOutliner/Private/TedsOutlinerHierarchy.cpp
@@ -48,8 +48,9 @@ void FTedsOutlinerHierarchy::CreateChildren(const FSceneOutlinerTreeItemPtr& Ite
 FSceneOutlinerTreeItemPtr FTedsOutlinerHierarchy::FindOrCreateParentItem(const ISceneOutlinerTreeItem& Item,
 	const TMap<FSceneOutlinerTreeItemID, FSceneOutlinerTreeItemPtr>& Items, bool bCreate)
 {
+	using namespace UE::Editor::DataStorage;
 	const FTedsOutlinerTreeItem* TedsTreeItem = Item.CastTo<FTedsOutlinerTreeItem>();
-	const IEditorDataStorageProvider* Storage = TedsOutlinerImpl->GetStorage();
+	const ICoreProvider* Storage = TedsOutlinerImpl->GetStorage();
 	
 	// If this item is not a TEDS item, we are not handling it
 	if(!TedsTreeItem)
@@ -57,7 +58,7 @@ FSceneOutlinerTreeItemPtr FTedsOutlinerHierarchy::FindOrCreateParentItem(const I
 		return nullptr;
 	}
 	
-	const UE::Editor::DataStorage::RowHandle ParentRowHandle = TedsOutlinerImpl->GetParentRow(TedsTreeItem->GetRowHandle());
+	const RowHandle ParentRowHandle = TedsOutlinerImpl->GetParentRow(TedsTreeItem->GetRowHandle());
 
 	if(!Storage->IsRowAvailable(ParentRowHandle))
 	{
diff --git a/Engine/Plugins/Experimental/EditorDataStorageFeatures/Source/TedsOutliner/Private/TedsOutlinerImpl.cpp b/Engine/Plugins/Experimental/EditorDataStorageFeatures/Source/TedsOutliner/Private/TedsOutlinerImpl.cpp
--- a/Engine/Plugins/Experimental/EditorDataStorageFeatures/Source/TedsOutliner/Private/TedsOutlinerImpl.cpp
+++ b/Engine/Plugins/Experimental/EditorDataStorageFeatures/Source/TedsOutliner/Private/TedsOutlinerImpl.cpp
@@ -91,9 +91,9 @@ FTedsOutlinerImpl::FTedsOutlinerImpl(const FTedsOutlinerParams& InParams, IScene
 {
 	// Initialize the TEDS constructs
 	using namespace UE::Editor::DataStorage;
-	Storage = GetMutableDataStorageFeature<IEditorDataStorageProvider>(StorageFeatureName);
-	StorageUi = GetMutableDataStorageFeature<IEditorDataStorageUiProvider>(UiFeatureName);
-	StorageCompatibility = GetMutableDataStorageFeature<IEditorDataStorageCompatibilityProvider>(CompatibilityFeatureName);
+	Storage = GetMutableDataStorageFeature<ICoreProvider>(StorageFeatureName);
+	StorageUi = GetMutableDataStorageFeature<IUiProvider>(UiFeatureName);
+	StorageCompatibility = GetMutableDataStorageFeature<ICompatibilityProvider>(CompatibilityFeatureName);
 }
 
 void FTedsOutlinerImpl::CreateFilterQueries()
@@ -192,6 +192,7 @@ TSharedRef<SWidget> FTedsOutlinerImpl::CreateLabelWidgetForItem(DataStorage::Row
 	// Find the best matching label widget for a given list of columns
 	auto CreateWidgetConstructor = [this](TArray<TWeakObjectPtr<const UScriptStruct>>& ColumnTypes) -> TSharedPtr<FTypedElementWidgetConstructor>
 	{
+		using namespace UE::Editor::DataStorage;
 		TSharedPtr<FTypedElementWidgetConstructor> OutWidgetConstructorPtr;
 
 		bool bFoundWidget = false;
@@ -200,7 +201,7 @@ TSharedRef<SWidget> FTedsOutlinerImpl::CreateLabelWidgetForItem(DataStorage::Row
 
 		for(const FName& Purpose : Purposes)
 		{
-			StorageUi->CreateWidgetConstructors(Purpose, IEditorDataStorageUiProvider::EMatchApproach::LongestMatch, ColumnTypes, {},
+			StorageUi->CreateWidgetConstructors(Purpose, IUiProvider::EMatchApproach::LongestMatch, ColumnTypes, {},
 					[&OutWidgetConstructorPtr, ColumnTypes, &bFoundWidget](
 					TUniquePtr<FTypedElementWidgetConstructor> CreatedConstructor, 
 					TConstArrayView<TWeakObjectPtr<const UScriptStruct>> MatchedColumnTypes)
@@ -743,17 +744,17 @@ void FTedsOutlinerImpl::Tick()
 	}
 }
 
-IEditorDataStorageProvider* FTedsOutlinerImpl::GetStorage() const
+DataStorage::ICoreProvider* FTedsOutlinerImpl::GetStorage() const
 {
 	return Storage;
 }
 
-IEditorDataStorageUiProvider* FTedsOutlinerImpl::GetStorageUI() const
+DataStorage::IUiProvider* FTedsOutlinerImpl::GetStorageUI() const
 {
 	return StorageUi;
 }
 
-IEditorDataStorageCompatibilityProvider* FTedsOutlinerImpl::GetStorageCompatibility() const
+DataStorage::ICompatibilityProvider* FTedsOutlinerImpl::GetStorageCompatibility() const
 {
 	return StorageCompatibility;
 }
diff --git a/Engine/Plugins/Experimental/EditorDataStorageFeatures/Source/TedsOutliner/Private/TedsOutlinerItem.cpp b/Engine/Plugins/Experimental/EditorDataStorageFeatures/Source/TedsOutliner/Private/TedsOutlinerItem.cpp
--- a/Engine/Plugins/Experimental/EditorDataStorageFeatures/Source/TedsOutliner/Private/TedsOutlinerItem.cpp
+++ b/Engine/Plugins/Experimental/EditorDataStorageFeatures/Source/TedsOutliner/Private/TedsOutlinerItem.cpp
@@ -70,7 +70,7 @@ TSharedRef<SWidget> FTedsOutlinerTreeItem::GenerateLabelWidget(ISceneOutliner& O
 void FTedsOutlinerTreeItem::GenerateContextMenu(UToolMenu* Menu, SSceneOutliner& Outliner)
 {
 	using namespace UE::Editor::DataStorage;
-	IEditorDataStorageProvider* DataStorageInterface = GetMutableDataStorageFeature<IEditorDataStorageProvider>(StorageFeatureName);
+	ICoreProvider* DataStorageInterface = GetMutableDataStorageFeature<ICoreProvider>(StorageFeatureName);
 
 	if (FTedsOutlinerContextMenuColumn* Column = DataStorageInterface->GetColumn<FTedsOutlinerContextMenuColumn>(RowHandle))
 	{
diff --git a/Engine/Plugins/Experimental/EditorDataStorageFeatures/Source/TedsOutliner/Private/TedsOutlinerMode.cpp b/Engine/Plugins/Experimental/EditorDataStorageFeatures/Source/TedsOutliner/Private/TedsOutlinerMode.cpp
--- a/Engine/Plugins/Experimental/EditorDataStorageFeatures/Source/TedsOutliner/Private/TedsOutlinerMode.cpp
+++ b/Engine/Plugins/Experimental/EditorDataStorageFeatures/Source/TedsOutliner/Private/TedsOutlinerMode.cpp
@@ -56,7 +56,7 @@ void FTedsOutlinerMode::Rebuild()
 void FTedsOutlinerMode::OnSelectionChanged()
 {
 	TOptional<FName> SelectionSetName = TedsOutlinerImpl->GetSelectionSetName();
-	IEditorDataStorageProvider* Storage = TedsOutlinerImpl->GetStorage();
+	DataStorage::ICoreProvider* Storage = TedsOutlinerImpl->GetStorage();
 	
 	// The selection in TEDS was changed, update the outliner to respond
 	SceneOutliner->SetSelection([SelectionSetName, Storage](ISceneOutlinerTreeItem& InItem) -> bool
@@ -142,7 +142,7 @@ bool FTedsOutlinerMode::ParseDragDrop(FSceneOutlinerDragDropPayload& OutPayload,
 FSceneOutlinerDragValidationInfo FTedsOutlinerMode::ValidateDrop(const ISceneOutlinerTreeItem& DropTarget, const FSceneOutlinerDragDropPayload& Payload) const
 {
 	const TOptional<FTedsOutlinerHierarchyData>& HierarchyData = TedsOutlinerImpl->GetHierarchyData();
-	IEditorDataStorageProvider* Storage = TedsOutlinerImpl->GetStorage();
+	DataStorage::ICoreProvider* Storage = TedsOutlinerImpl->GetStorage();
 
 	// We don't want drag/drop if this TEDS Outliner isn't showing any hierarchy data
 	if(!HierarchyData.IsSet())
@@ -223,7 +223,7 @@ FSceneOutlinerDragValidationInfo FTedsOutlinerMode::ValidateDrop(const ISceneOut
 void FTedsOutlinerMode::OnDrop(ISceneOutlinerTreeItem& DropTarget, const FSceneOutlinerDragDropPayload& Payload, const FSceneOutlinerDragValidationInfo& ValidationInfo) const
 {
 	const TOptional<FTedsOutlinerHierarchyData>& HierarchyData = TedsOutlinerImpl->GetHierarchyData();
-	IEditorDataStorageProvider* Storage = TedsOutlinerImpl->GetStorage();
+	DataStorage::ICoreProvider* Storage = TedsOutlinerImpl->GetStorage();
 
 	if(!Private::TedsOutlinerDragDropEnabledCvar->GetBool() || !HierarchyData.IsSet())
 	{
diff --git a/Engine/Plugins/Experimental/EditorDataStorageFeatures/Source/TedsOutliner/Private/Widgets/OutlinerIconWidget.cpp b/Engine/Plugins/Experimental/EditorDataStorageFeatures/Source/TedsOutliner/Private/Widgets/OutlinerIconWidget.cpp
--- a/Engine/Plugins/Experimental/EditorDataStorageFeatures/Source/TedsOutliner/Private/Widgets/OutlinerIconWidget.cpp
+++ b/Engine/Plugins/Experimental/EditorDataStorageFeatures/Source/TedsOutliner/Private/Widgets/OutlinerIconWidget.cpp
@@ -92,8 +92,8 @@ namespace UE::OutlinerIconWidget::Local
 	}
 }
 
-void UOutlinerIconWidgetFactory::RegisterWidgetConstructors(IEditorDataStorageProvider& DataStorage,
-	IEditorDataStorageUiProvider& DataStorageUi) const
+void UOutlinerIconWidgetFactory::RegisterWidgetConstructors(UE::Editor::DataStorage::ICoreProvider& DataStorage,
+	UE::Editor::DataStorage::IUiProvider& DataStorageUi) const
 {
 	using namespace UE::Editor::DataStorage::Queries;
 
@@ -106,8 +106,8 @@ FOutlinerIconWidgetConstructor::FOutlinerIconWidgetConstructor()
 {
 }
 
-TSharedPtr<SWidget> FOutlinerIconWidgetConstructor::CreateWidget(IEditorDataStorageProvider* DataStorage,
-	IEditorDataStorageUiProvider* DataStorageUi, UE::Editor::DataStorage::RowHandle TargetRow,
+TSharedPtr<SWidget> FOutlinerIconWidgetConstructor::CreateWidget(UE::Editor::DataStorage::ICoreProvider* DataStorage,
+	UE::Editor::DataStorage::IUiProvider* DataStorageUi, UE::Editor::DataStorage::RowHandle TargetRow,
 	UE::Editor::DataStorage::RowHandle WidgetRow, const UE::Editor::DataStorage::FMetaDataView& Arguments)
 {
 	using namespace UE::Editor::DataStorage;
diff --git a/Engine/Plugins/Experimental/EditorDataStorageFeatures/Source/TedsOutliner/Private/Widgets/OutlinerLabelWidget.cpp b/Engine/Plugins/Experimental/EditorDataStorageFeatures/Source/TedsOutliner/Private/Widgets/OutlinerLabelWidget.cpp
--- a/Engine/Plugins/Experimental/EditorDataStorageFeatures/Source/TedsOutliner/Private/Widgets/OutlinerLabelWidget.cpp
+++ b/Engine/Plugins/Experimental/EditorDataStorageFeatures/Source/TedsOutliner/Private/Widgets/OutlinerLabelWidget.cpp
@@ -14,26 +14,27 @@
 static const FName IconWidgetPurpose = TEXT("SceneOutliner.RowLabel.Icon");
 static const FName TextWidgetPurpose = TEXT("SceneOutliner.RowLabel.Text");
 
-void UOutlinerLabelWidgetFactory::RegisterWidgetConstructors(IEditorDataStorageProvider& DataStorage,
-	IEditorDataStorageUiProvider& DataStorageUi) const
+void UOutlinerLabelWidgetFactory::RegisterWidgetConstructors(UE::Editor::DataStorage::ICoreProvider& DataStorage,
+	UE::Editor::DataStorage::IUiProvider& DataStorageUi) const
 {
 	using namespace UE::Editor::DataStorage::Queries;
 	DataStorageUi.RegisterWidgetFactory<FOutlinerLabelWidgetConstructor>(TEXT("SceneOutliner.RowLabel"),
 		TColumn<FTypedElementLabelColumn>() && (TColumn<FTypedElementClassTypeInfoColumn>() || TColumn<FTypedElementScriptStructTypeInfoColumn>()));
 }
 
-void UOutlinerLabelWidgetFactory::RegisterWidgetPurposes(IEditorDataStorageUiProvider& DataStorageUi) const
+void UOutlinerLabelWidgetFactory::RegisterWidgetPurposes(UE::Editor::DataStorage::IUiProvider& DataStorageUi) const
 {
-	DataStorageUi.RegisterWidgetPurpose(IconWidgetPurpose, IEditorDataStorageUiProvider::EPurposeType::UniqueByNameAndColumn,
+	using namespace UE::Editor::DataStorage;
+	DataStorageUi.RegisterWidgetPurpose(IconWidgetPurpose, IUiProvider::EPurposeType::UniqueByNameAndColumn,
 		LOCTEXT("IconItemCellWidgetPurpose", "The icon widget to use in cells for the Scene Outliner specific to the Item label column.")
 	);
 
-	DataStorageUi.RegisterWidgetPurpose(TextWidgetPurpose, IEditorDataStorageUiProvider::EPurposeType::UniqueByNameAndColumn,
+	DataStorageUi.RegisterWidgetPurpose(TextWidgetPurpose, IUiProvider::EPurposeType::UniqueByNameAndColumn,
 		LOCTEXT("TextItemCellWidgetPurpose", "The text widget to use in cells for the Scene Outliner specific to the Item label column.")
 	);
 }
 
-static void GetAllColumns(TArray<TWeakObjectPtr<const UScriptStruct>>& OutColumns, const IEditorDataStorageProvider& DataStorage, UE::Editor::DataStorage::RowHandle Row)
+static void GetAllColumns(TArray<TWeakObjectPtr<const UScriptStruct>>& OutColumns, const UE::Editor::DataStorage::ICoreProvider& DataStorage, UE::Editor::DataStorage::RowHandle Row)
 {
 	OutColumns.Empty();
 	DataStorage.ListColumns(Row, [&OutColumns](const UScriptStruct& ColumnType)
@@ -48,10 +49,12 @@ FOutlinerLabelWidgetConstructor::FOutlinerLabelWidgetConstructor()
 {
 }
 
-TSharedPtr<SWidget> FOutlinerLabelWidgetConstructor::CreateWidget(IEditorDataStorageProvider* DataStorage,
-	IEditorDataStorageUiProvider* DataStorageUi, UE::Editor::DataStorage::RowHandle TargetRow,
+TSharedPtr<SWidget> FOutlinerLabelWidgetConstructor::CreateWidget(UE::Editor::DataStorage::ICoreProvider* DataStorage,
+	UE::Editor::DataStorage::IUiProvider* DataStorageUi,
+	UE::Editor::DataStorage::RowHandle TargetRow,
 	UE::Editor::DataStorage::RowHandle WidgetRow, const UE::Editor::DataStorage::FMetaDataView& Arguments)
 {
+	using namespace UE::Editor::DataStorage;
 	if (DataStorage->IsRowAvailable(TargetRow))
 	{
 		TSharedPtr<FTypedElementWidgetConstructor> OutIconWidgetConstructorPtr;
@@ -71,9 +74,9 @@ TSharedPtr<SWidget> FOutlinerLabelWidgetConstructor::CreateWidget(IEditorDataSto
 
 		TArray<TWeakObjectPtr<const UScriptStruct>> Columns;
 		GetAllColumns(Columns, *DataStorage, TargetRow);
-		DataStorageUi->CreateWidgetConstructors(IconWidgetPurpose, IEditorDataStorageUiProvider::EMatchApproach::LongestMatch, Columns, Arguments, AssignIconWidgetToColumn);
+		DataStorageUi->CreateWidgetConstructors(IconWidgetPurpose, IUiProvider::EMatchApproach::LongestMatch, Columns, Arguments, AssignIconWidgetToColumn);
 		GetAllColumns(Columns, *DataStorage, TargetRow);
-		DataStorageUi->CreateWidgetConstructors(TextWidgetPurpose, IEditorDataStorageUiProvider::EMatchApproach::LongestMatch, Columns, Arguments, AssignTextWidgetToColumn);
+		DataStorageUi->CreateWidgetConstructors(TextWidgetPurpose, IUiProvider::EMatchApproach::LongestMatch, Columns, Arguments, AssignTextWidgetToColumn);
 
 		static const FMargin ColumnItemPadding(8, 0);
 		constexpr float IconNameHorizontalPadding = 8.f;
diff --git a/Engine/Plugins/Experimental/EditorDataStorageFeatures/Source/TedsOutliner/Private/Widgets/OutlinerLabelWidget.h b/Engine/Plugins/Experimental/EditorDataStorageFeatures/Source/TedsOutliner/Private/Widgets/OutlinerLabelWidget.h
--- a/Engine/Plugins/Experimental/EditorDataStorageFeatures/Source/TedsOutliner/Private/Widgets/OutlinerLabelWidget.h
+++ b/Engine/Plugins/Experimental/EditorDataStorageFeatures/Source/TedsOutliner/Private/Widgets/OutlinerLabelWidget.h
@@ -8,8 +8,6 @@
 
 #include "OutlinerLabelWidget.generated.h"
 
-class IEditorDataStorageProvider;
-
 UCLASS()
 class UOutlinerLabelWidgetFactory : public UEditorDataStorageFactory
 {
@@ -18,10 +16,10 @@ class UOutlinerLabelWidgetFactory : public UEditorDataStorageFactory
 public:
 	~UOutlinerLabelWidgetFactory() override = default;
 
-	TEDSOUTLINER_API void RegisterWidgetConstructors(IEditorDataStorageProvider& DataStorage,
-		IEditorDataStorageUiProvider& DataStorageUi) const override;
+	TEDSOUTLINER_API void RegisterWidgetConstructors(UE::Editor::DataStorage::ICoreProvider& DataStorage,
+		UE::Editor::DataStorage::IUiProvider& DataStorageUi) const override;
 
-	TEDSOUTLINER_API void RegisterWidgetPurposes(IEditorDataStorageUiProvider& DataStorageUi) const override;
+	TEDSOUTLINER_API void RegisterWidgetPurposes(UE::Editor::DataStorage::IUiProvider& DataStorageUi) const override;
 };
 
 // Label widget for the Scene Outliner that shows an icon (with optional override information) + a text label
@@ -35,8 +33,8 @@ public:
 	~FOutlinerLabelWidgetConstructor() override = default;
 
 	TEDSOUTLINER_API virtual TSharedPtr<SWidget> CreateWidget(
-		IEditorDataStorageProvider* DataStorage,
-		IEditorDataStorageUiProvider* DataStorageUi,
+		UE::Editor::DataStorage::ICoreProvider* DataStorage,
+		UE::Editor::DataStorage::IUiProvider* DataStorageUi,
 		UE::Editor::DataStorage::RowHandle TargetRow,
 		UE::Editor::DataStorage::RowHandle WidgetRow, 
 		const UE::Editor::DataStorage::FMetaDataView& Arguments) override;
diff --git a/Engine/Plugins/Experimental/EditorDataStorageFeatures/Source/TedsOutliner/Private/Widgets/OutlinerTextWidget.cpp b/Engine/Plugins/Experimental/EditorDataStorageFeatures/Source/TedsOutliner/Private/Widgets/OutlinerTextWidget.cpp
--- a/Engine/Plugins/Experimental/EditorDataStorageFeatures/Source/TedsOutliner/Private/Widgets/OutlinerTextWidget.cpp
+++ b/Engine/Plugins/Experimental/EditorDataStorageFeatures/Source/TedsOutliner/Private/Widgets/OutlinerTextWidget.cpp
@@ -15,8 +15,8 @@
 
 #define LOCTEXT_NAMESPACE "OutlinerTextWidget"
 
-void UOutlinerTextWidgetFactory::RegisterWidgetConstructors(IEditorDataStorageProvider& DataStorage,
-	IEditorDataStorageUiProvider& DataStorageUi) const
+void UOutlinerTextWidgetFactory::RegisterWidgetConstructors(UE::Editor::DataStorage::ICoreProvider& DataStorage,
+	UE::Editor::DataStorage::IUiProvider& DataStorageUi) const
 {
 	using namespace UE::Editor::DataStorage::Queries;
 
@@ -34,8 +34,9 @@ FOutlinerTextWidgetConstructor::FOutlinerTextWidgetConstructor(const UScriptStru
 {
 }
 
-TSharedPtr<SWidget> FOutlinerTextWidgetConstructor::CreateWidget(IEditorDataStorageProvider* DataStorage,
-	IEditorDataStorageUiProvider* DataStorageUi, UE::Editor::DataStorage::RowHandle TargetRow,
+TSharedPtr<SWidget> FOutlinerTextWidgetConstructor::CreateWidget(UE::Editor::DataStorage::ICoreProvider* DataStorage,
+	UE::Editor::DataStorage::IUiProvider* DataStorageUi,
+	UE::Editor::DataStorage::RowHandle TargetRow,
 	UE::Editor::DataStorage::RowHandle WidgetRow, const UE::Editor::DataStorage::FMetaDataView& Arguments)
 {
 	using namespace UE::Editor::DataStorage;
@@ -63,7 +64,7 @@ TSharedPtr<SWidget> FOutlinerTextWidgetConstructor::CreateWidget(IEditorDataStor
 void FOutlinerTextWidgetConstructor::OnCommitText(
 	const FText& NewText,
 	ETextCommit::Type,
-	IEditorDataStorageProvider* DataStorage,
+	UE::Editor::DataStorage::ICoreProvider* DataStorage,
 	UE::Editor::DataStorage::RowHandle TargetRow)
 {
 	FString NewLabelText = NewText.ToString();
@@ -88,8 +89,8 @@ bool FOutlinerTextWidgetConstructor::OnVerifyText(const FText& Label, FText& Err
 }
 
 TSharedPtr<SWidget> FOutlinerTextWidgetConstructor::CreateEditableWidget(
-	IEditorDataStorageProvider* DataStorage,
-	IEditorDataStorageUiProvider* DataStorageUi,
+	UE::Editor::DataStorage::ICoreProvider* DataStorage,
+	UE::Editor::DataStorage::IUiProvider* DataStorageUi,
 	UE::Editor::DataStorage::RowHandle TargetRow,
 	UE::Editor::DataStorage::RowHandle WidgetRow,
 	const UE::Editor::DataStorage::FMetaDataView& Arguments)
@@ -98,8 +99,8 @@ TSharedPtr<SWidget> FOutlinerTextWidgetConstructor::CreateEditableWidget(
 }
 
 TSharedPtr<SWidget> FOutlinerTextWidgetConstructor::CreateNonEditableWidget(
-	IEditorDataStorageProvider* DataStorage,
-	IEditorDataStorageUiProvider* DataStorageUi,
+	UE::Editor::DataStorage::ICoreProvider* DataStorage,
+	UE::Editor::DataStorage::IUiProvider* DataStorageUi,
 	UE::Editor::DataStorage::RowHandle TargetRow,
 	UE::Editor::DataStorage::RowHandle WidgetRow,
 	const UE::Editor::DataStorage::FMetaDataView& Arguments)
@@ -108,7 +109,7 @@ TSharedPtr<SWidget> FOutlinerTextWidgetConstructor::CreateNonEditableWidget(
 }
 
 TSharedPtr<SInlineEditableTextBlock> FOutlinerTextWidgetConstructor::CreateEditableTextBlock(
-	IEditorDataStorageProvider* DataStorage,
+	UE::Editor::DataStorage::ICoreProvider* DataStorage,
 	UE::Editor::DataStorage::RowHandle TargetRow,
 	UE::Editor::DataStorage::RowHandle WidgetRow) const
 {
@@ -132,7 +133,7 @@ TSharedPtr<SInlineEditableTextBlock> FOutlinerTextWidgetConstructor::CreateEdita
 }
 
 TSharedPtr<STextBlock> FOutlinerTextWidgetConstructor::CreateNonEditableTextBlock(
-	IEditorDataStorageProvider* DataStorage,
+	UE::Editor::DataStorage::ICoreProvider* DataStorage,
 	UE::Editor::DataStorage::RowHandle TargetRow,
 	UE::Editor::DataStorage::RowHandle WidgetRow) const
 {
diff --git a/Engine/Plugins/Experimental/EditorDataStorageFeatures/Source/TedsOutliner/Private/Widgets/OutlinerVisibilityWidget.cpp b/Engine/Plugins/Experimental/EditorDataStorageFeatures/Source/TedsOutliner/Private/Widgets/OutlinerVisibilityWidget.cpp
--- a/Engine/Plugins/Experimental/EditorDataStorageFeatures/Source/TedsOutliner/Private/Widgets/OutlinerVisibilityWidget.cpp
+++ b/Engine/Plugins/Experimental/EditorDataStorageFeatures/Source/TedsOutliner/Private/Widgets/OutlinerVisibilityWidget.cpp
@@ -33,21 +33,24 @@ FOutlinerVisibilityWidgetConstructor::FOutlinerVisibilityWidgetConstructor()
 {
 }
 
-void UOutlinerVisibilityWidgetFactory::RegisterWidgetConstructors(IEditorDataStorageProvider& DataStorage, IEditorDataStorageUiProvider& DataStorageUi) const
+void UOutlinerVisibilityWidgetFactory::RegisterWidgetConstructors(UE::Editor::DataStorage::ICoreProvider& DataStorage,
+	UE::Editor::DataStorage::IUiProvider& DataStorageUi) const
 {
 	using namespace UE::Editor::DataStorage::Queries;
 	DataStorageUi.RegisterWidgetFactory<FOutlinerVisibilityWidgetConstructor>(VisibilityWidgetPurpose, TColumn<FVisibleInEditorColumn>());
 	
 }
 
-TSharedPtr<SWidget> FOutlinerVisibilityWidgetConstructor::CreateWidget(IEditorDataStorageProvider* DataStorage, IEditorDataStorageUiProvider* DataStorageUi, UE::Editor::DataStorage::RowHandle TargetRow, UE::Editor::DataStorage::RowHandle WidgetRow, const UE::Editor::DataStorage::FMetaDataView& Arguments)
+TSharedPtr<SWidget> FOutlinerVisibilityWidgetConstructor::CreateWidget(UE::Editor::DataStorage::ICoreProvider* DataStorage,
+	UE::Editor::DataStorage::IUiProvider* DataStorageUi,
+	UE::Editor::DataStorage::RowHandle TargetRow, UE::Editor::DataStorage::RowHandle WidgetRow, const UE::Editor::DataStorage::FMetaDataView& Arguments)
 {
 	return SNew(SHorizontalBox)
 		+ SHorizontalBox::Slot()
 		.AutoWidth()
 		.VAlign(VAlign_Center)
 		[
-			SNew(STEDSVisibilityWidget, TargetRow, WidgetRow)
+			SNew(STedsVisibilityWidget, TargetRow, WidgetRow)
 				.ToolTipText(LOCTEXT("SceneOutlinerVisibilityToggleTooltip", "Toggles the visibility of this object in the level editor."))
 		];
 }
@@ -61,13 +64,16 @@ FOutlinerVisibilityHeaderConstructor::FOutlinerVisibilityHeaderConstructor()
 {
 }
 
-void UOutlinerVisibilityHeaderFactory::RegisterWidgetConstructors(IEditorDataStorageProvider& DataStorage, IEditorDataStorageUiProvider& DataStorageUi) const
+void UOutlinerVisibilityHeaderFactory::RegisterWidgetConstructors(UE::Editor::DataStorage::ICoreProvider& DataStorage,
+	UE::Editor::DataStorage::IUiProvider& DataStorageUi) const
 {
 	using namespace UE::Editor::DataStorage::Queries;
 	DataStorageUi.RegisterWidgetFactory<FOutlinerVisibilityHeaderConstructor>(VisibilityHeaderPurpose, TColumn<FVisibleInEditorColumn>());
 }
 
-TSharedPtr<SWidget> FOutlinerVisibilityHeaderConstructor::CreateWidget(IEditorDataStorageProvider* DataStorage, IEditorDataStorageUiProvider* DataStorageUi, UE::Editor::DataStorage::RowHandle TargetRow, UE::Editor::DataStorage::RowHandle WidgetRow, const UE::Editor::DataStorage::FMetaDataView& Arguments)
+TSharedPtr<SWidget> FOutlinerVisibilityHeaderConstructor::CreateWidget(UE::Editor::DataStorage::ICoreProvider* DataStorage,
+	UE::Editor::DataStorage::IUiProvider* DataStorageUi,
+	UE::Editor::DataStorage::RowHandle TargetRow, UE::Editor::DataStorage::RowHandle WidgetRow, const UE::Editor::DataStorage::FMetaDataView& Arguments)
 {
 	DataStorage->AddColumn(WidgetRow, FHeaderWidgetSizeColumn
 	{
@@ -116,16 +122,16 @@ public:
 	}
 };
 
-void STEDSVisibilityWidget::Construct(const FArguments& InArgs, const RowHandle& InTargetRow, const RowHandle& InWidgetRow)
+void STedsVisibilityWidget::Construct(const FArguments& InArgs, const RowHandle& InTargetRow, const RowHandle& InWidgetRow)
 {
 	TargetRow = InTargetRow;
 	WidgetRow = InWidgetRow;
 
 	SImage::Construct(
 		SImage::FArguments()
-		.IsEnabled(this, &STEDSVisibilityWidget::IsEnabled)
-		.ColorAndOpacity(this, &STEDSVisibilityWidget::GetForegroundColor)
-		.Image(this, &STEDSVisibilityWidget::GetBrush)
+		.IsEnabled(this, &STedsVisibilityWidget::IsEnabled)
+		.ColorAndOpacity(this, &STedsVisibilityWidget::GetForegroundColor)
+		.Image(this, &STedsVisibilityWidget::GetBrush)
 	);
 
 
@@ -142,7 +148,7 @@ void STEDSVisibilityWidget::Construct(const FArguments& InArgs, const RowHandle&
 
 }
 
-FReply STEDSVisibilityWidget::OnDragDetected(const FGeometry& MyGeometry, const FPointerEvent& MouseEvent)
+FReply STedsVisibilityWidget::OnDragDetected(const FGeometry& MyGeometry, const FPointerEvent& MouseEvent)
 {
 	if (MouseEvent.IsMouseButtonDown(EKeys::LeftMouseButton))
 	{
@@ -155,7 +161,7 @@ FReply STEDSVisibilityWidget::OnDragDetected(const FGeometry& MyGeometry, const
 }
 
 /** If a visibility drag drop operation has entered this widget, set its item to the new visibility state */
-void STEDSVisibilityWidget::OnDragEnter(const FGeometry& MyGeometry, const FDragDropEvent& DragDropEvent)
+void STedsVisibilityWidget::OnDragEnter(const FGeometry& MyGeometry, const FDragDropEvent& DragDropEvent)
 {
 	auto VisibilityOp = DragDropEvent.GetOperationAs<FTEDSVisibilityDragDropOp>();
 	if (VisibilityOp.IsValid())
@@ -164,7 +170,7 @@ void STEDSVisibilityWidget::OnDragEnter(const FGeometry& MyGeometry, const FDrag
 	}
 }
 
-FReply STEDSVisibilityWidget::HandleClick()
+FReply STedsVisibilityWidget::HandleClick()
 {
 	if (!IsEnabled())
 	{
@@ -194,13 +200,13 @@ FReply STEDSVisibilityWidget::HandleClick()
 	return FReply::Handled().DetectDrag(SharedThis(this), EKeys::LeftMouseButton);
 }
 
-FReply STEDSVisibilityWidget::OnMouseButtonDoubleClick(const FGeometry& InMyGeometry, const FPointerEvent& InMouseEvent)
+FReply STedsVisibilityWidget::OnMouseButtonDoubleClick(const FGeometry& InMyGeometry, const FPointerEvent& InMouseEvent)
 {
 	return HandleClick();
 }
 
 /** Called when the mouse button is pressed down on this widget */
-FReply STEDSVisibilityWidget::OnMouseButtonDown(const FGeometry& MyGeometry, const FPointerEvent& MouseEvent)
+FReply STedsVisibilityWidget::OnMouseButtonDown(const FGeometry& MyGeometry, const FPointerEvent& MouseEvent)
 {
 	if (MouseEvent.GetEffectingButton() != EKeys::LeftMouseButton)
 	{
@@ -211,7 +217,7 @@ FReply STEDSVisibilityWidget::OnMouseButtonDown(const FGeometry& MyGeometry, con
 }
 
 /** Process a mouse up message */
-FReply STEDSVisibilityWidget::OnMouseButtonUp(const FGeometry& MyGeometry, const FPointerEvent& MouseEvent)
+FReply STedsVisibilityWidget::OnMouseButtonUp(const FGeometry& MyGeometry, const FPointerEvent& MouseEvent)
 {
 	if (MouseEvent.GetEffectingButton() == EKeys::LeftMouseButton)
 	{
@@ -223,13 +229,13 @@ FReply STEDSVisibilityWidget::OnMouseButtonUp(const FGeometry& MyGeometry, const
 }
 
 /** Called when this widget had captured the mouse, but that capture has been revoked for some reason. */
-void STEDSVisibilityWidget::OnMouseCaptureLost(const FCaptureLostEvent& CaptureLostEvent)
+void STedsVisibilityWidget::OnMouseCaptureLost(const FCaptureLostEvent& CaptureLostEvent)
 {
 	UndoTransaction.Reset();
 }
 
 /** Get the brush for this widget */
-const FSlateBrush* STEDSVisibilityWidget::GetBrush() const
+const FSlateBrush* STedsVisibilityWidget::GetBrush() const
 {
 	if (IsVisible())
 	{
@@ -241,7 +247,7 @@ const FSlateBrush* STEDSVisibilityWidget::GetBrush() const
 	}
 }
 
-FSlateColor STEDSVisibilityWidget::GetForegroundColor() const
+FSlateColor STedsVisibilityWidget::GetForegroundColor() const
 {
 	const bool bIsSelected = IsSelected();
 	const bool bIsVisible = IsVisible();
@@ -261,9 +267,9 @@ FSlateColor STEDSVisibilityWidget::GetForegroundColor() const
 }
 
 /** Check if our wrapped tree item is visible */
-bool STEDSVisibilityWidget::IsVisible() const
+bool STedsVisibilityWidget::IsVisible() const
 {
-	if (IEditorDataStorageProvider* DataStorage = GetDataStorage())
+	if (UE::Editor::DataStorage::ICoreProvider* DataStorage = GetDataStorage())
 	{
 		FVisibleInEditorColumn* HiddenInEditorColumn = DataStorage->GetColumn<FVisibleInEditorColumn>(TargetRow);
 		if (HiddenInEditorColumn)
@@ -275,9 +281,9 @@ bool STEDSVisibilityWidget::IsVisible() const
 	return true;
 }
 
-bool STEDSVisibilityWidget::IsSelected() const
+bool STedsVisibilityWidget::IsSelected() const
 {
-	if (IEditorDataStorageProvider* DataStorage = GetDataStorage())
+	if (UE::Editor::DataStorage::ICoreProvider* DataStorage = GetDataStorage())
 	{
 		return DataStorage->HasColumns<FTypedElementSelectionColumn>(TargetRow);
 	}
@@ -286,16 +292,16 @@ bool STEDSVisibilityWidget::IsSelected() const
 }
 
 /** Set the item this widget is responsible for to be hidden or shown */
-void STEDSVisibilityWidget::SetIsVisible(RowHandle InRow, const bool bVisible)
+void STedsVisibilityWidget::SetIsVisible(RowHandle InRow, const bool bVisible)
 {
 	TWeakPtr<ISceneOutlinerTreeItem> RowTreeItem = GetTreeItem(InRow);
-	if (IEditorDataStorageProvider* DataStorage = GetDataStorage(); RowTreeItem.IsValid())
+	if (UE::Editor::DataStorage::ICoreProvider* DataStorage = GetDataStorage(); RowTreeItem.IsValid())
 	{
 		SetVisibility_Recursive(*DataStorage, RowTreeItem.Pin(), bVisible);
 	}
 }
 
-void STEDSVisibilityWidget::CommitVisibility(IEditorDataStorageProvider& DataStorage, RowHandle Row, bool bVisible)
+void STedsVisibilityWidget::CommitVisibility(UE::Editor::DataStorage::ICoreProvider& DataStorage, RowHandle Row, bool bVisible)
 {
 	// Set new visibility on the TargetRow
 	FVisibleInEditorColumn* VisibileInEditorColumn = DataStorage.GetColumn<FVisibleInEditorColumn>(Row);
@@ -306,12 +312,12 @@ void STEDSVisibilityWidget::CommitVisibility(IEditorDataStorageProvider& DataSto
 	DataStorage.AddColumn<FTypedElementSyncBackToWorldTag>(Row);
 }
 
-void STEDSVisibilityWidget::SetVisibility_Recursive(IEditorDataStorageProvider& DataStorage, FSceneOutlinerTreeItemPtr TreeItem, bool bVisible)
+void STedsVisibilityWidget::SetVisibility_Recursive(UE::Editor::DataStorage::ICoreProvider& DataStorage, FSceneOutlinerTreeItemPtr TreeItem, bool bVisible)
 {
 	using namespace UE::Editor::DataStorage;
 	using namespace UE::Editor::Outliner;
 
-	if (IEditorDataStorageCompatibilityProvider* DataStorageCompat = GetDataStorageCompatibility())
+	if (ICompatibilityProvider* DataStorageCompat = GetDataStorageCompatibility())
 	{
 		if (FActorTreeItem* ActorTreeItem = TreeItem->CastTo<FActorTreeItem>())
 		{
@@ -334,9 +340,10 @@ void STEDSVisibilityWidget::SetVisibility_Recursive(IEditorDataStorageProvider&
 	}
 }
 
-FSceneOutlinerTreeItemPtr STEDSVisibilityWidget::GetTreeItem(RowHandle InRow) const
+FSceneOutlinerTreeItemPtr STedsVisibilityWidget::GetTreeItem(RowHandle InRow) const
 {
-	IEditorDataStorageProvider* DataStorage = GetDataStorage();
+	using namespace UE::Editor::DataStorage;
+	ICoreProvider* DataStorage = GetDataStorage();
 	if (!DataStorage)
 	{
 		return nullptr;
@@ -364,33 +371,33 @@ FSceneOutlinerTreeItemPtr STEDSVisibilityWidget::GetTreeItem(RowHandle InRow) co
 	return nullptr;
 }
 
-IEditorDataStorageProvider* STEDSVisibilityWidget::GetDataStorage()
+UE::Editor::DataStorage::ICoreProvider* STedsVisibilityWidget::GetDataStorage()
 {
 	using namespace UE::Editor::DataStorage;
-	return GetMutableDataStorageFeature<IEditorDataStorageProvider>(StorageFeatureName);
+	return GetMutableDataStorageFeature<ICoreProvider>(StorageFeatureName);
 }
 
-IEditorDataStorageProvider* STEDSVisibilityWidget::GetDataStorageUI()
+UE::Editor::DataStorage::ICoreProvider* STedsVisibilityWidget::GetDataStorageUI()
 {
 	using namespace UE::Editor::DataStorage;
-	return GetMutableDataStorageFeature<IEditorDataStorageProvider>(UiFeatureName);
+	return GetMutableDataStorageFeature<ICoreProvider>(UiFeatureName);
 }
 
-IEditorDataStorageCompatibilityProvider* STEDSVisibilityWidget::GetDataStorageCompatibility()
+UE::Editor::DataStorage::ICompatibilityProvider* STedsVisibilityWidget::GetDataStorageCompatibility()
 {
 	using namespace UE::Editor::DataStorage;
-	return GetMutableDataStorageFeature<IEditorDataStorageCompatibilityProvider>(CompatibilityFeatureName);
+	return GetMutableDataStorageFeature<ICompatibilityProvider>(CompatibilityFeatureName);
 }
 
-void STEDSVisibilityWidget::GetSelectedRows(TArray<RowHandle>& OutSelectedRows)
+void STedsVisibilityWidget::GetSelectedRows(TArray<RowHandle>& OutSelectedRows)
 {
 	using namespace UE::Editor::DataStorage;
 	using namespace UE::Editor::DataStorage::Queries;
-	if (IEditorDataStorageProvider* DataStorage = GetDataStorage())
+	if (ICoreProvider* DataStorage = GetDataStorage())
 	{
 		static QueryHandle AllSelectedItemsQuery = DataStorage->RegisterQuery(Select().Where().All<FTypedElementSelectionColumn>().Compile());
 		DataStorage->RunQuery(AllSelectedItemsQuery,
-			CreateDirectQueryCallbackBinding([&OutSelectedRows, &DataStorage](const IEditorDataStorageProvider::IDirectQueryContext& Context)
+			CreateDirectQueryCallbackBinding([&OutSelectedRows, &DataStorage](const IDirectQueryContext& Context)
 			{
 				// Only add selections from the level editor - for now.
 				// UE-231184: Support Visibility Column in multiple editor contexts
diff --git a/Engine/Plugins/Experimental/EditorDataStorageFeatures/Source/TedsOutliner/Private/Widgets/TransformHeadsUpWidget.cpp b/Engine/Plugins/Experimental/EditorDataStorageFeatures/Source/TedsOutliner/Private/Widgets/TransformHeadsUpWidget.cpp
--- a/Engine/Plugins/Experimental/EditorDataStorageFeatures/Source/TedsOutliner/Private/Widgets/TransformHeadsUpWidget.cpp
+++ b/Engine/Plugins/Experimental/EditorDataStorageFeatures/Source/TedsOutliner/Private/Widgets/TransformHeadsUpWidget.cpp
@@ -153,7 +153,7 @@ static void UpdateTransformHeadsUpDisplay(FTypedElementSlateWidgetReferenceColum
 // UTransformHeadsUpWidgetFactory
 //
 
-void UTransformHeadsUpWidgetFactory::RegisterQueries(IEditorDataStorageProvider& DataStorage)
+void UTransformHeadsUpWidgetFactory::RegisterQueries(UE::Editor::DataStorage::ICoreProvider& DataStorage)
 {
 	using namespace UE::Editor::DataStorage::Queries;
 		
@@ -185,8 +185,8 @@ void UTransformHeadsUpWidgetFactory::RegisterQueries(IEditorDataStorageProvider&
 
 }
 
-void UTransformHeadsUpWidgetFactory::RegisterWidgetConstructors(IEditorDataStorageProvider& DataStorage,
-	IEditorDataStorageUiProvider& DataStorageUi) const
+void UTransformHeadsUpWidgetFactory::RegisterWidgetConstructors(UE::Editor::DataStorage::ICoreProvider& DataStorage,
+	UE::Editor::DataStorage::IUiProvider& DataStorageUi) const
 {
 	using namespace UE::Editor::DataStorage::Queries;
 
@@ -207,6 +207,7 @@ FTransformHeadsUpWidgetConstructor::FTransformHeadsUpWidgetConstructor()
 
 TConstArrayView<const UScriptStruct*> FTransformHeadsUpWidgetConstructor::GetAdditionalColumnsList() const
 {
+	using namespace UE::Editor::DataStorage;
 	static TTypedElementColumnTypeList<FTypedElementRowReferenceColumn,
 		FTransformHeadsUpWidgetTag> Columns;
 	return Columns;
@@ -218,8 +219,8 @@ TSharedPtr<SWidget> FTransformHeadsUpWidgetConstructor::CreateWidget(const UE::E
 }
 
 bool FTransformHeadsUpWidgetConstructor::FinalizeWidget(
-	IEditorDataStorageProvider* DataStorage,
-	IEditorDataStorageUiProvider* DataStorageUi,
+	UE::Editor::DataStorage::ICoreProvider* DataStorage,
+	UE::Editor::DataStorage::IUiProvider* DataStorageUi,
 	UE::Editor::DataStorage::RowHandle Row,
 	const TSharedPtr<SWidget>& Widget)
 {
diff --git a/Engine/Plugins/Experimental/EditorDataStorageFeatures/Source/TedsOutliner/Private/Widgets/TypeInfoWidget.cpp b/Engine/Plugins/Experimental/EditorDataStorageFeatures/Source/TedsOutliner/Private/Widgets/TypeInfoWidget.cpp
--- a/Engine/Plugins/Experimental/EditorDataStorageFeatures/Source/TedsOutliner/Private/Widgets/TypeInfoWidget.cpp
+++ b/Engine/Plugins/Experimental/EditorDataStorageFeatures/Source/TedsOutliner/Private/Widgets/TypeInfoWidget.cpp
@@ -13,8 +13,8 @@
 #include "Widgets/Text/STextBlock.h"
 #include "Widgets/Images/SImage.h"
 
-void UTypeInfoWidgetFactory::RegisterWidgetConstructors(IEditorDataStorageProvider& DataStorage,
-	IEditorDataStorageUiProvider& DataStorageUi) const
+void UTypeInfoWidgetFactory::RegisterWidgetConstructors(UE::Editor::DataStorage::ICoreProvider& DataStorage,
+	UE::Editor::DataStorage::IUiProvider& DataStorageUi) const
 {
 	using namespace UE::Editor::DataStorage::Queries;
 
@@ -62,8 +62,9 @@ TSharedPtr<SWidget> FTypeInfoWidgetConstructor::CreateWidget(
 	}
 }
 
-bool FTypeInfoWidgetConstructor::FinalizeWidget(IEditorDataStorageProvider* DataStorage,
-	IEditorDataStorageUiProvider* DataStorageUi, UE::Editor::DataStorage::RowHandle Row, const TSharedPtr<SWidget>& Widget)
+bool FTypeInfoWidgetConstructor::FinalizeWidget(UE::Editor::DataStorage::ICoreProvider* DataStorage,
+	UE::Editor::DataStorage::IUiProvider* DataStorageUi,
+	UE::Editor::DataStorage::RowHandle Row, const TSharedPtr<SWidget>& Widget)
 {
 	checkf(Widget, TEXT("Referenced widget is not valid. A constructed widget may not have been cleaned up. This can "
 		"also happen if this processor is running in the same phase as the processors responsible for cleaning up old "
