[TEDS] Cleanup aliases, move core interface classes into namespace, and fix up one class name that wasn't following standards

diff --git a/Engine/Plugins/Experimental/EditorDataStorageFeatures/Source/TedsOutliner/Public/Compatibility/SceneOutlinerTedsBridge.h b/Engine/Plugins/Experimental/EditorDataStorageFeatures/Source/TedsOutliner/Public/Compatibility/SceneOutlinerTedsBridge.h
--- a/Engine/Plugins/Experimental/EditorDataStorageFeatures/Source/TedsOutliner/Public/Compatibility/SceneOutlinerTedsBridge.h
+++ b/Engine/Plugins/Experimental/EditorDataStorageFeatures/Source/TedsOutliner/Public/Compatibility/SceneOutlinerTedsBridge.h
@@ -13,9 +13,13 @@
 #include "SceneOutlinerTedsBridge.generated.h"
 
 class ISceneOutliner;
-class IEditorDataStorageProvider;
-class IEditorDataStorageUiProvider;
-class IEditorDataStorageCompatibilityProvider;
+
+namespace UE::Editor::DataStorage
+{
+	class ICoreProvider;
+	class IUiProvider;
+	class ICompatibilityProvider;
+} // namespace UE::Editor::DataStorage
 
 DECLARE_DELEGATE_RetVal_OneParam(FSceneOutlinerTreeItemID, FTreeItemIDDealiaser, UE::Editor::DataStorage::RowHandle);
 
@@ -61,9 +65,9 @@ private:
 	
 	TMap<TWeakPtr<ISceneOutliner>, TSharedPtr<FSceneOutlinerTedsBridge>> SceneOutliners;
 
-	IEditorDataStorageProvider* Storage{ nullptr };
-	IEditorDataStorageUiProvider* StorageUi{ nullptr };
-	IEditorDataStorageCompatibilityProvider* StorageCompatibility{ nullptr };
+	UE::Editor::DataStorage::ICoreProvider* Storage{ nullptr };
+	UE::Editor::DataStorage::IUiProvider* StorageUi{ nullptr };
+	UE::Editor::DataStorage::ICompatibilityProvider* StorageCompatibility{ nullptr };
 
 	TMap<TWeakObjectPtr<const UScriptStruct>, FName> TEDSToOutlinerDefaultColumnMapping;
 };
@@ -76,5 +80,5 @@ class USceneOutlinerTedsBridgeFactory : public UEditorDataStorageFactory
 public:
 	~USceneOutlinerTedsBridgeFactory() override = default;
 
-	void RegisterWidgetPurposes(IEditorDataStorageUiProvider& DataStorageUi) const override;
+	void RegisterWidgetPurposes(UE::Editor::DataStorage::IUiProvider& DataStorageUi) const override;
 };
diff --git a/Engine/Plugins/Experimental/EditorDataStorageFeatures/Source/TedsOutliner/Public/TedsOutlinerImpl.h b/Engine/Plugins/Experimental/EditorDataStorageFeatures/Source/TedsOutliner/Public/TedsOutlinerImpl.h
--- a/Engine/Plugins/Experimental/EditorDataStorageFeatures/Source/TedsOutliner/Public/TedsOutlinerImpl.h
+++ b/Engine/Plugins/Experimental/EditorDataStorageFeatures/Source/TedsOutliner/Public/TedsOutlinerImpl.h
@@ -7,8 +7,12 @@
 #include "ISceneOutlinerMode.h"
 #include "Elements/Columns/TypedElementHiearchyColumns.h"
 
-class IEditorDataStorageUiProvider;
-class IEditorDataStorageCompatibilityProvider;
+namespace UE::Editor::DataStorage
+{
+	class IUiProvider;
+	class ICompatibilityProvider;
+} // namespace UE::Editor::DataStorage
+
 struct FTypedElementWidgetConstructor;
 class SWidget;
 
@@ -130,9 +134,9 @@ public:
 	void Init();
 
 	// TEDS construct getters
-	IEditorDataStorageProvider* GetStorage() const;
-	IEditorDataStorageUiProvider* GetStorageUI() const;
-	IEditorDataStorageCompatibilityProvider* GetStorageCompatibility() const;
+	DataStorage::ICoreProvider* GetStorage() const;
+	DataStorage::IUiProvider* GetStorageUI() const;
+	DataStorage::ICompatibilityProvider* GetStorageCompatibility() const;
 
 	TOptional<FName> GetSelectionSetName() const;
 
@@ -196,9 +200,9 @@ protected:
 	
 protected:
 	// TEDS Storage Constructs
-	IEditorDataStorageProvider* Storage{ nullptr };
-	IEditorDataStorageUiProvider* StorageUi{ nullptr };
-	IEditorDataStorageCompatibilityProvider* StorageCompatibility{ nullptr };
+	DataStorage::ICoreProvider* Storage{ nullptr };
+	DataStorage::IUiProvider* StorageUi{ nullptr };
+	DataStorage::ICompatibilityProvider* StorageCompatibility{ nullptr };
 
 	FTedsOutlinerParams CreationParams;
 
diff --git a/Engine/Plugins/Experimental/EditorDataStorageFeatures/Source/TedsOutliner/Public/Widgets/OutlinerIconWidget.h b/Engine/Plugins/Experimental/EditorDataStorageFeatures/Source/TedsOutliner/Public/Widgets/OutlinerIconWidget.h
--- a/Engine/Plugins/Experimental/EditorDataStorageFeatures/Source/TedsOutliner/Public/Widgets/OutlinerIconWidget.h
+++ b/Engine/Plugins/Experimental/EditorDataStorageFeatures/Source/TedsOutliner/Public/Widgets/OutlinerIconWidget.h
@@ -19,8 +19,8 @@ class UOutlinerIconWidgetFactory : public UEditorDataStorageFactory
 public:
 	~UOutlinerIconWidgetFactory() override = default;
 
-	void RegisterWidgetConstructors(IEditorDataStorageProvider& DataStorage,
-		IEditorDataStorageUiProvider& DataStorageUi) const override;
+	void RegisterWidgetConstructors(UE::Editor::DataStorage::ICoreProvider& DataStorage,
+		UE::Editor::DataStorage::IUiProvider& DataStorageUi) const override;
 };
 
 USTRUCT()
@@ -33,8 +33,8 @@ public:
 	~FOutlinerIconWidgetConstructor() override = default;
 
 	TEDSOUTLINER_API virtual TSharedPtr<SWidget> CreateWidget(
-		IEditorDataStorageProvider* DataStorage,
-		IEditorDataStorageUiProvider* DataStorageUi,
+		UE::Editor::DataStorage::ICoreProvider* DataStorage,
+		UE::Editor::DataStorage::IUiProvider* DataStorageUi,
 		UE::Editor::DataStorage::RowHandle TargetRow,
 		UE::Editor::DataStorage::RowHandle WidgetRow,
 		const UE::Editor::DataStorage::FMetaDataView& Arguments) override;
diff --git a/Engine/Plugins/Experimental/EditorDataStorageFeatures/Source/TedsOutliner/Public/Widgets/OutlinerTextWidget.h b/Engine/Plugins/Experimental/EditorDataStorageFeatures/Source/TedsOutliner/Public/Widgets/OutlinerTextWidget.h
--- a/Engine/Plugins/Experimental/EditorDataStorageFeatures/Source/TedsOutliner/Public/Widgets/OutlinerTextWidget.h
+++ b/Engine/Plugins/Experimental/EditorDataStorageFeatures/Source/TedsOutliner/Public/Widgets/OutlinerTextWidget.h
@@ -22,8 +22,8 @@ class UOutlinerTextWidgetFactory : public UEditorDataStorageFactory
 public:
 	virtual ~UOutlinerTextWidgetFactory() override = default;
 
-	void RegisterWidgetConstructors(IEditorDataStorageProvider& DataStorage,
-		IEditorDataStorageUiProvider& DataStorageUi) const override;
+	void RegisterWidgetConstructors(UE::Editor::DataStorage::ICoreProvider& DataStorage,
+		UE::Editor::DataStorage::IUiProvider& DataStorageUi) const override;
 };
 
 USTRUCT()
@@ -37,22 +37,22 @@ public:
 	TEDSOUTLINER_API virtual ~FOutlinerTextWidgetConstructor() override = default;
 
 	TEDSOUTLINER_API virtual TSharedPtr<SWidget> CreateWidget(
-		IEditorDataStorageProvider* DataStorage,
-		IEditorDataStorageUiProvider* DataStorageUi,
+		UE::Editor::DataStorage::ICoreProvider* DataStorage,
+		UE::Editor::DataStorage::IUiProvider* DataStorageUi,
 		UE::Editor::DataStorage::RowHandle TargetRow,
 		UE::Editor::DataStorage::RowHandle WidgetRow,
 		const UE::Editor::DataStorage::FMetaDataView& Arguments) override;
 
 protected:
 	//! Updates the label columns with the given text.
-	static TEDSOUTLINER_API void OnCommitText(const FText& NewText, ETextCommit::Type Type, IEditorDataStorageProvider* DataStorage, UE::Editor::DataStorage::RowHandle TargetRow);
+	static TEDSOUTLINER_API void OnCommitText(const FText& NewText, ETextCommit::Type Type, UE::Editor::DataStorage::ICoreProvider* DataStorage, UE::Editor::DataStorage::RowHandle TargetRow);
 	static TEDSOUTLINER_API bool OnVerifyText(const FText& Label, FText& ErrorMessage);
 
 	//! Virtual relay function so that deriving Constructors may edit the editable
 	//! widget without having to override the whole CreateWidget procedure.
 	TEDSOUTLINER_API virtual TSharedPtr<SWidget> CreateEditableWidget(
-		IEditorDataStorageProvider* DataStorage,
-		IEditorDataStorageUiProvider* DataStorageUi,
+		UE::Editor::DataStorage::ICoreProvider* DataStorage,
+		UE::Editor::DataStorage::IUiProvider* DataStorageUi,
 		UE::Editor::DataStorage::RowHandle TargetRow,
 		UE::Editor::DataStorage::RowHandle WidgetRow,
 		const UE::Editor::DataStorage::FMetaDataView& Arguments);
@@ -60,21 +60,21 @@ protected:
 	//! Virtual relay function so that deriving Constructors may edit the non-editable
 	//! widget without having to override the whole CreateWidget procedure.
 	TEDSOUTLINER_API virtual TSharedPtr<SWidget> CreateNonEditableWidget(
-		IEditorDataStorageProvider* DataStorage,
-		IEditorDataStorageUiProvider* DataStorageUi,
+		UE::Editor::DataStorage::ICoreProvider* DataStorage,
+		UE::Editor::DataStorage::IUiProvider* DataStorageUi,
 		UE::Editor::DataStorage::RowHandle TargetRow,
 		UE::Editor::DataStorage::RowHandle WidgetRow,
 		const UE::Editor::DataStorage::FMetaDataView& Arguments);
 
 	//! The actual function that creates the editable widget.
 	TEDSOUTLINER_API TSharedPtr<SInlineEditableTextBlock> CreateEditableTextBlock(
-		IEditorDataStorageProvider* DataStorage,
+		UE::Editor::DataStorage::ICoreProvider* DataStorage,
 		UE::Editor::DataStorage::RowHandle TargetRow,
 		UE::Editor::DataStorage::RowHandle WidgetRow) const;
 
 	//! The actual function that creates the non-editable widget.
 	TEDSOUTLINER_API TSharedPtr<STextBlock> CreateNonEditableTextBlock(
-		IEditorDataStorageProvider* DataStorage,
+		UE::Editor::DataStorage::ICoreProvider* DataStorage,
 		UE::Editor::DataStorage::RowHandle TargetRow,
 		UE::Editor::DataStorage::RowHandle WidgetRow) const;
 };
diff --git a/Engine/Plugins/Experimental/EditorDataStorageFeatures/Source/TedsOutliner/Public/Widgets/OutlinerVisibilityWidget.h b/Engine/Plugins/Experimental/EditorDataStorageFeatures/Source/TedsOutliner/Public/Widgets/OutlinerVisibilityWidget.h
--- a/Engine/Plugins/Experimental/EditorDataStorageFeatures/Source/TedsOutliner/Public/Widgets/OutlinerVisibilityWidget.h
+++ b/Engine/Plugins/Experimental/EditorDataStorageFeatures/Source/TedsOutliner/Public/Widgets/OutlinerVisibilityWidget.h
@@ -12,9 +12,6 @@
 
 #include "OutlinerVisibilityWidget.generated.h"
 
-class IEditorDataStorageProvider;
-class UScriptStruct;
-
 UCLASS()
 class UOutlinerVisibilityWidgetFactory : public UEditorDataStorageFactory
 {
@@ -23,7 +20,8 @@ class UOutlinerVisibilityWidgetFactory : public UEditorDataStorageFactory
 public:
 	~UOutlinerVisibilityWidgetFactory() override = default;
 
-	TEDSOUTLINER_API void RegisterWidgetConstructors(IEditorDataStorageProvider& DataStorage, IEditorDataStorageUiProvider& DataStorageUi) const override;
+	TEDSOUTLINER_API void RegisterWidgetConstructors(UE::Editor::DataStorage::ICoreProvider& DataStorage,
+		UE::Editor::DataStorage::IUiProvider& DataStorageUi) const override;
 
 public:
 
@@ -38,7 +36,8 @@ class UOutlinerVisibilityHeaderFactory : public UEditorDataStorageFactory
 public:
 	~UOutlinerVisibilityHeaderFactory() override = default;
 
-	TEDSOUTLINER_API void RegisterWidgetConstructors(IEditorDataStorageProvider& DataStorage, IEditorDataStorageUiProvider& DataStorageUi) const override;
+	TEDSOUTLINER_API void RegisterWidgetConstructors(UE::Editor::DataStorage::ICoreProvider& DataStorage,
+		UE::Editor::DataStorage::IUiProvider& DataStorageUi) const override;
 
 public:
 
@@ -55,8 +54,8 @@ public:
 	~FOutlinerVisibilityWidgetConstructor() override = default;
 
 	TEDSOUTLINER_API virtual TSharedPtr<SWidget> CreateWidget(
-		IEditorDataStorageProvider* DataStorage,
-		IEditorDataStorageUiProvider* DataStorageUi,
+		UE::Editor::DataStorage::ICoreProvider* DataStorage,
+		UE::Editor::DataStorage::IUiProvider* DataStorageUi,
 		UE::Editor::DataStorage::RowHandle TargetRow,
 		UE::Editor::DataStorage::RowHandle WidgetRow, 
 		const UE::Editor::DataStorage::FMetaDataView& Arguments) override;
@@ -72,20 +71,20 @@ public:
 	~FOutlinerVisibilityHeaderConstructor() override = default;
 
 	TEDSOUTLINER_API virtual TSharedPtr<SWidget> CreateWidget(
-		IEditorDataStorageProvider* DataStorage,
-		IEditorDataStorageUiProvider* DataStorageUi,
+		UE::Editor::DataStorage::ICoreProvider* DataStorage,
+		UE::Editor::DataStorage::IUiProvider* DataStorageUi,
 		UE::Editor::DataStorage::RowHandle TargetRow,
 		UE::Editor::DataStorage::RowHandle WidgetRow, 
 		const UE::Editor::DataStorage::FMetaDataView& Arguments) override;
 };
 
 /** Widget responsible for managing the visibility for a single item */
-class TEDSOUTLINER_API STEDSVisibilityWidget : public SImage
+class TEDSOUTLINER_API STedsVisibilityWidget : public SImage
 {
 	using RowHandle = UE::Editor::DataStorage::RowHandle;
 
 public:
-	SLATE_BEGIN_ARGS(STEDSVisibilityWidget) {}
+	SLATE_BEGIN_ARGS(STedsVisibilityWidget) {}
 	SLATE_END_ARGS()
 
 	/** Construct this widget */
@@ -147,17 +146,17 @@ protected:
 private:
 
 	// Commit the visibility state into the Rows DataStorage. Will trigger a SyncBackToWorld manually but the object is never dirtied.
-	static void CommitVisibility(IEditorDataStorageProvider& DataStorage, RowHandle Row, bool bVisible);
+	static void CommitVisibility(UE::Editor::DataStorage::ICoreProvider& DataStorage, RowHandle Row, bool bVisible);
 
 	// Recursive internal method to set visibility on a TreeItem and all of it's children.
-	static void SetVisibility_Recursive(IEditorDataStorageProvider& DataStorage, FSceneOutlinerTreeItemPtr TreeItem, bool bVisible);
+	static void SetVisibility_Recursive(UE::Editor::DataStorage::ICoreProvider& DataStorage, FSceneOutlinerTreeItemPtr TreeItem, bool bVisible);
 
 	// Returns the tree item associated with this widget.
 	FSceneOutlinerTreeItemPtr GetTreeItem(RowHandle InRow) const;
 
-	static IEditorDataStorageProvider* GetDataStorage();
-	static IEditorDataStorageProvider* GetDataStorageUI();
-	static IEditorDataStorageCompatibilityProvider* GetDataStorageCompatibility();
+	static UE::Editor::DataStorage::ICoreProvider* GetDataStorage();
+	static UE::Editor::DataStorage::ICoreProvider* GetDataStorageUI();
+	static UE::Editor::DataStorage::ICompatibilityProvider* GetDataStorageCompatibility();
 
 	static void GetSelectedRows(TArray<RowHandle>& OutSelectedRows);
 };
\ No newline at end of file
diff --git a/Engine/Plugins/Experimental/EditorDataStorageFeatures/Source/TedsOutliner/Public/Widgets/TransformHeadsUpWidget.h b/Engine/Plugins/Experimental/EditorDataStorageFeatures/Source/TedsOutliner/Public/Widgets/TransformHeadsUpWidget.h
--- a/Engine/Plugins/Experimental/EditorDataStorageFeatures/Source/TedsOutliner/Public/Widgets/TransformHeadsUpWidget.h
+++ b/Engine/Plugins/Experimental/EditorDataStorageFeatures/Source/TedsOutliner/Public/Widgets/TransformHeadsUpWidget.h
@@ -26,9 +26,9 @@ class UTransformHeadsUpWidgetFactory : public UEditorDataStorageFactory
 public:
 	~UTransformHeadsUpWidgetFactory() override = default;
 
-	void RegisterQueries(IEditorDataStorageProvider& DataStorage) override;
-	void RegisterWidgetConstructors(IEditorDataStorageProvider& DataStorage,
-		IEditorDataStorageUiProvider& DataStorageUi) const override;
+	void RegisterQueries(UE::Editor::DataStorage::ICoreProvider& DataStorage) override;
+	void RegisterWidgetConstructors(UE::Editor::DataStorage::ICoreProvider& DataStorage,
+		UE::Editor::DataStorage::IUiProvider& DataStorageUi) const override;
 };
 
 USTRUCT()
@@ -45,8 +45,8 @@ public:
 protected:
 	TSharedPtr<SWidget> CreateWidget(const UE::Editor::DataStorage::FMetaDataView& Arguments) override;
 	bool FinalizeWidget(
-		IEditorDataStorageProvider* DataStorage,
-		IEditorDataStorageUiProvider* DataStorageUi,
+		UE::Editor::DataStorage::ICoreProvider* DataStorage,
+		UE::Editor::DataStorage::IUiProvider* DataStorageUi,
 		UE::Editor::DataStorage::RowHandle Row,
 		const TSharedPtr<SWidget>& Widget) override;
 };
diff --git a/Engine/Plugins/Experimental/EditorDataStorageFeatures/Source/TedsOutliner/Public/Widgets/TypeInfoWidget.h b/Engine/Plugins/Experimental/EditorDataStorageFeatures/Source/TedsOutliner/Public/Widgets/TypeInfoWidget.h
--- a/Engine/Plugins/Experimental/EditorDataStorageFeatures/Source/TedsOutliner/Public/Widgets/TypeInfoWidget.h
+++ b/Engine/Plugins/Experimental/EditorDataStorageFeatures/Source/TedsOutliner/Public/Widgets/TypeInfoWidget.h
@@ -8,9 +8,6 @@
 
 #include "TypeInfoWidget.generated.h"
 
-struct FSlateBrush;
-struct FTypedElementClassTypeInfoColumn;
-
 UCLASS()
 class UTypeInfoWidgetFactory : public UEditorDataStorageFactory
 {
@@ -19,8 +16,8 @@ class UTypeInfoWidgetFactory : public UEditorDataStorageFactory
 public:
 	~UTypeInfoWidgetFactory() override = default;
 
-	void RegisterWidgetConstructors(IEditorDataStorageProvider& DataStorage,
-		IEditorDataStorageUiProvider& DataStorageUi) const override;
+	void RegisterWidgetConstructors(UE::Editor::DataStorage::ICoreProvider& DataStorage,
+		UE::Editor::DataStorage::IUiProvider& DataStorageUi) const override;
 };
 
 USTRUCT()
@@ -35,7 +32,8 @@ public:
 protected:
 	explicit FTypeInfoWidgetConstructor(const UScriptStruct* InTypeInfo);
 	TSharedPtr<SWidget> CreateWidget(const UE::Editor::DataStorage::FMetaDataView& Arguments) override;
-	bool FinalizeWidget(IEditorDataStorageProvider* DataStorage, IEditorDataStorageUiProvider* DataStorageUi,
+	bool FinalizeWidget(UE::Editor::DataStorage::ICoreProvider* DataStorage,
+		UE::Editor::DataStorage::IUiProvider* DataStorageUi,
 		UE::Editor::DataStorage::RowHandle Row, const TSharedPtr<SWidget>& Widget) override;
 
 protected:
diff --git a/Engine/Plugins/Experimental/EditorDataStorageFeatures/Source/TedsQueryStack/Private/TedsQueryNode.cpp b/Engine/Plugins/Experimental/EditorDataStorageFeatures/Source/TedsQueryStack/Private/TedsQueryNode.cpp
--- a/Engine/Plugins/Experimental/EditorDataStorageFeatures/Source/TedsQueryStack/Private/TedsQueryNode.cpp
+++ b/Engine/Plugins/Experimental/EditorDataStorageFeatures/Source/TedsQueryStack/Private/TedsQueryNode.cpp
@@ -6,12 +6,12 @@
 
 namespace UE::Editor::DataStorage::QueryStack
 {
-	FQueryNode::FQueryNode(IEditorDataStorageProvider& Storage)
+	FQueryNode::FQueryNode(ICoreProvider& Storage)
 		: Storage(Storage)
 	{
 	}
 
-	FQueryNode::FQueryNode(IEditorDataStorageProvider& Storage, UE::Editor::DataStorage::FQueryDescription Query)
+	FQueryNode::FQueryNode(ICoreProvider& Storage, UE::Editor::DataStorage::FQueryDescription Query)
 		: QueryHandle(Storage.RegisterQuery(MoveTemp(Query)))
 		, Storage(Storage)
 	{
diff --git a/Engine/Plugins/Experimental/EditorDataStorageFeatures/Source/TedsQueryStack/Private/TedsRowMonitorNode.cpp b/Engine/Plugins/Experimental/EditorDataStorageFeatures/Source/TedsQueryStack/Private/TedsRowMonitorNode.cpp
--- a/Engine/Plugins/Experimental/EditorDataStorageFeatures/Source/TedsQueryStack/Private/TedsRowMonitorNode.cpp
+++ b/Engine/Plugins/Experimental/EditorDataStorageFeatures/Source/TedsQueryStack/Private/TedsRowMonitorNode.cpp
@@ -8,7 +8,7 @@
 namespace UE::Editor::DataStorage::QueryStack
 {
 	FRowMonitorNode::FRowMonitorNode(
-		IEditorDataStorageProvider& Storage, TSharedPtr<IRowNode> ParentRow, TSharedPtr<IQueryNode> QueryNode)
+		ICoreProvider& Storage, TSharedPtr<IRowNode> ParentRow, TSharedPtr<IQueryNode> QueryNode)
 		: QueryNode(MoveTemp(QueryNode))
 		, ParentRow(MoveTemp(ParentRow))
 		, Storage(Storage)
@@ -19,7 +19,7 @@ namespace UE::Editor::DataStorage::QueryStack
 	}
 
 	FRowMonitorNode::FRowMonitorNode(
-		IEditorDataStorageProvider& Storage, TSharedPtr<IRowNode> ParentRow, TArray<TObjectPtr<const UScriptStruct>> Columns)
+		ICoreProvider& Storage, TSharedPtr<IRowNode> ParentRow, TArray<TObjectPtr<const UScriptStruct>> Columns)
 		: MonitoredColumns(MoveTemp(Columns))
 		, ParentRow(MoveTemp(ParentRow))
 		, Storage(Storage)
@@ -29,7 +29,7 @@ namespace UE::Editor::DataStorage::QueryStack
 	}
 
 	FRowMonitorNode::FRowMonitorNode(
-		IEditorDataStorageProvider& Storage, TSharedPtr<IQueryNode> QueryNode, TSharedPtr<IRowNode> ParentRow, 
+		ICoreProvider& Storage, TSharedPtr<IQueryNode> QueryNode, TSharedPtr<IRowNode> ParentRow, 
 		TArray<TObjectPtr<const UScriptStruct>> MonitoredColumns)
 		: MonitoredColumns(MoveTemp(MonitoredColumns))
 		, QueryNode(MoveTemp(QueryNode))
@@ -42,7 +42,7 @@ namespace UE::Editor::DataStorage::QueryStack
 		UpdateMonitoredColumns();
 	}
 
-	FRowMonitorNode::FRowMonitorNode(IEditorDataStorageProvider& Storage, TSharedPtr<IQueryNode> QueryNode)
+	FRowMonitorNode::FRowMonitorNode(ICoreProvider& Storage, TSharedPtr<IQueryNode> QueryNode)
 		: QueryNode(MoveTemp(QueryNode))
 		, Storage(Storage)
 	{
@@ -50,7 +50,7 @@ namespace UE::Editor::DataStorage::QueryStack
 		UpdateMonitoredColumns();
 	}
 	
-	FRowMonitorNode::FRowMonitorNode(IEditorDataStorageProvider& Storage, TArray<TObjectPtr<const UScriptStruct>> Columns)
+	FRowMonitorNode::FRowMonitorNode(ICoreProvider& Storage, TArray<TObjectPtr<const UScriptStruct>> Columns)
 		: MonitoredColumns(MoveTemp(Columns))
 		, Storage(Storage)
 	{
@@ -58,7 +58,7 @@ namespace UE::Editor::DataStorage::QueryStack
 	}
 
 	FRowMonitorNode::FRowMonitorNode(
-		IEditorDataStorageProvider& Storage, TSharedPtr<IQueryNode> QueryNode, TArray<TObjectPtr<const UScriptStruct>> MonitoredColumns)
+		ICoreProvider& Storage, TSharedPtr<IQueryNode> QueryNode, TArray<TObjectPtr<const UScriptStruct>> MonitoredColumns)
 		: MonitoredColumns(MoveTemp(MonitoredColumns))
 		, QueryNode(MoveTemp(QueryNode))
 		, Storage(Storage)
diff --git a/Engine/Plugins/Experimental/EditorDataStorageFeatures/Source/TedsQueryStack/Private/TedsRowQueryResultsNode.cpp b/Engine/Plugins/Experimental/EditorDataStorageFeatures/Source/TedsQueryStack/Private/TedsRowQueryResultsNode.cpp
--- a/Engine/Plugins/Experimental/EditorDataStorageFeatures/Source/TedsQueryStack/Private/TedsRowQueryResultsNode.cpp
+++ b/Engine/Plugins/Experimental/EditorDataStorageFeatures/Source/TedsQueryStack/Private/TedsRowQueryResultsNode.cpp
@@ -7,7 +7,7 @@
 
 namespace UE::Editor::DataStorage::QueryStack
 {
-	FRowQueryResultsNode::FRowQueryResultsNode(IEditorDataStorageProvider& Storage, TSharedPtr<IQueryNode> InQueryNode, ESyncFlags SyncFlags)
+	FRowQueryResultsNode::FRowQueryResultsNode(ICoreProvider& Storage, TSharedPtr<IQueryNode> InQueryNode, ESyncFlags SyncFlags)
 		: QueryNode(MoveTemp(InQueryNode))
 		, Storage(Storage)
 		, SyncFlags(SyncFlags)
diff --git a/Engine/Plugins/Experimental/EditorDataStorageFeatures/Source/TedsQueryStack/Public/TedsQueryNode.h b/Engine/Plugins/Experimental/EditorDataStorageFeatures/Source/TedsQueryStack/Public/TedsQueryNode.h
--- a/Engine/Plugins/Experimental/EditorDataStorageFeatures/Source/TedsQueryStack/Public/TedsQueryNode.h
+++ b/Engine/Plugins/Experimental/EditorDataStorageFeatures/Source/TedsQueryStack/Public/TedsQueryNode.h
@@ -5,7 +5,10 @@
 #include "Elements/Common/TypedElementQueryDescription.h"
 #include "TedsQueryStackInterfaces.h"
 
-class IEditorDataStorageProvider;
+namespace UE::Editor::DataStorage
+{
+	class ICoreProvider;
+} // namespace UE::Editor::DataStorage
 
 namespace UE::Editor::DataStorage::QueryStack
 {
@@ -15,8 +18,8 @@ namespace UE::Editor::DataStorage::QueryStack
 	class FQueryNode : public IQueryNode
 	{
 	public:
-		TEDSQUERYSTACK_API explicit FQueryNode(IEditorDataStorageProvider& Storage);
-		TEDSQUERYSTACK_API FQueryNode(IEditorDataStorageProvider& Storage, UE::Editor::DataStorage::FQueryDescription Query);
+		TEDSQUERYSTACK_API explicit FQueryNode(ICoreProvider& Storage);
+		TEDSQUERYSTACK_API FQueryNode(ICoreProvider& Storage, UE::Editor::DataStorage::FQueryDescription Query);
 		TEDSQUERYSTACK_API virtual ~FQueryNode() override;
 
 		TEDSQUERYSTACK_API void SetQuery(UE::Editor::DataStorage::FQueryDescription Query);
@@ -28,7 +31,7 @@ namespace UE::Editor::DataStorage::QueryStack
 
 	private:
 		QueryHandle QueryHandle = InvalidQueryHandle;
-		IEditorDataStorageProvider& Storage;
+		ICoreProvider& Storage;
 		RevisionId Revision = 0;
 	};
 } // namespace UE::Editor::DataStorage::QueryStack
diff --git a/Engine/Plugins/Experimental/EditorDataStorageFeatures/Source/TedsQueryStack/Public/TedsRowMonitorNode.h b/Engine/Plugins/Experimental/EditorDataStorageFeatures/Source/TedsQueryStack/Public/TedsRowMonitorNode.h
--- a/Engine/Plugins/Experimental/EditorDataStorageFeatures/Source/TedsQueryStack/Public/TedsRowMonitorNode.h
+++ b/Engine/Plugins/Experimental/EditorDataStorageFeatures/Source/TedsQueryStack/Public/TedsRowMonitorNode.h
@@ -7,7 +7,10 @@
 #include "TedsQueryStackInterfaces.h"
 #include "UObject/ObjectPtr.h"
 
-class IEditorDataStorageProvider;
+namespace UE::Editor::DataStorage
+{
+	class ICoreProvider;
+} // namespace UE::Editor::DataStorage
 
 namespace UE::Editor::DataStorage::QueryStack
 {
@@ -18,18 +21,18 @@ namespace UE::Editor::DataStorage::QueryStack
 	{
 	public:
 		TEDSQUERYSTACK_API FRowMonitorNode(
-			IEditorDataStorageProvider& Storage, TSharedPtr<IRowNode> ParentRow, TSharedPtr<IQueryNode> QueryNode);
+			ICoreProvider& Storage, TSharedPtr<IRowNode> ParentRow, TSharedPtr<IQueryNode> QueryNode);
 		TEDSQUERYSTACK_API FRowMonitorNode(
-			IEditorDataStorageProvider& Storage, TSharedPtr<IRowNode> ParentRow, TArray<TObjectPtr<const UScriptStruct>> Columns);
+			ICoreProvider& Storage, TSharedPtr<IRowNode> ParentRow, TArray<TObjectPtr<const UScriptStruct>> Columns);
 		TEDSQUERYSTACK_API FRowMonitorNode(
-			IEditorDataStorageProvider& Storage, TSharedPtr<IQueryNode> QueryNode, TSharedPtr<IRowNode> ParentRow, 
+			ICoreProvider& Storage, TSharedPtr<IQueryNode> QueryNode, TSharedPtr<IRowNode> ParentRow, 
 			TArray<TObjectPtr<const UScriptStruct>> MonitoredColumns);
 		TEDSQUERYSTACK_API FRowMonitorNode(
-			IEditorDataStorageProvider& Storage, TSharedPtr<IQueryNode> QueryNode);
+			ICoreProvider& Storage, TSharedPtr<IQueryNode> QueryNode);
 		TEDSQUERYSTACK_API FRowMonitorNode(
-			IEditorDataStorageProvider& Storage, TArray<TObjectPtr<const UScriptStruct>> Columns);
+			ICoreProvider& Storage, TArray<TObjectPtr<const UScriptStruct>> Columns);
 		TEDSQUERYSTACK_API FRowMonitorNode(
-			IEditorDataStorageProvider& Storage, TSharedPtr<IQueryNode> QueryNode, TArray<TObjectPtr<const UScriptStruct>> MonitoredColumns);
+			ICoreProvider& Storage, TSharedPtr<IQueryNode> QueryNode, TArray<TObjectPtr<const UScriptStruct>> MonitoredColumns);
 
 		TEDSQUERYSTACK_API virtual ~FRowMonitorNode() override;
 
@@ -54,7 +57,7 @@ namespace UE::Editor::DataStorage::QueryStack
 		
 		TSharedPtr<IQueryNode> QueryNode;
 		TSharedPtr<IRowNode> ParentRow;
-		IEditorDataStorageProvider& Storage;
+		ICoreProvider& Storage;
 		RevisionId QueryRevision = 0;
 		RevisionId ParentRevision = 0;
 		RevisionId Revision = 0;
diff --git a/Engine/Plugins/Experimental/EditorDataStorageFeatures/Source/TedsQueryStack/Public/TedsRowQueryResultsNode.h b/Engine/Plugins/Experimental/EditorDataStorageFeatures/Source/TedsQueryStack/Public/TedsRowQueryResultsNode.h
--- a/Engine/Plugins/Experimental/EditorDataStorageFeatures/Source/TedsQueryStack/Public/TedsRowQueryResultsNode.h
+++ b/Engine/Plugins/Experimental/EditorDataStorageFeatures/Source/TedsQueryStack/Public/TedsRowQueryResultsNode.h
@@ -7,7 +7,10 @@
 #include "TedsQueryStackInterfaces.h"
 #include "Misc/EnumClassFlags.h"
 
-class IEditorDataStorageProvider;
+namespace UE::Editor::DataStorage
+{
+	class ICoreProvider;
+} // namespace UE::Editor::DataStorage
 
 namespace UE::Editor::DataStorage::QueryStack
 {
@@ -37,7 +40,7 @@ namespace UE::Editor::DataStorage::QueryStack
 			IncrementWhenDifferent = 1 << 2
 		};
 
-		TEDSQUERYSTACK_API FRowQueryResultsNode(IEditorDataStorageProvider& Storage, TSharedPtr<IQueryNode> InQueryNode, ESyncFlags SyncFlags);
+		TEDSQUERYSTACK_API FRowQueryResultsNode(ICoreProvider& Storage, TSharedPtr<IQueryNode> InQueryNode, ESyncFlags SyncFlags);
 		virtual ~FRowQueryResultsNode() override = default;
 
 		TEDSQUERYSTACK_API void Refresh();
@@ -52,7 +55,7 @@ namespace UE::Editor::DataStorage::QueryStack
 
 		FRowHandleArray Rows;
 		TSharedPtr<IQueryNode> QueryNode;
-		IEditorDataStorageProvider& Storage;
+		ICoreProvider& Storage;
 		RevisionId QueryRevision = 0;
 		RevisionId Revision = 0;
 		ESyncFlags SyncFlags = ESyncFlags::None;
diff --git a/Engine/Plugins/Experimental/EditorDataStorageFeatures/Source/TedsRevisionControl/Private/Processors/RevisionControlProcessors.cpp b/Engine/Plugins/Experimental/EditorDataStorageFeatures/Source/TedsRevisionControl/Private/Processors/RevisionControlProcessors.cpp
--- a/Engine/Plugins/Experimental/EditorDataStorageFeatures/Source/TedsRevisionControl/Private/Processors/RevisionControlProcessors.cpp
+++ b/Engine/Plugins/Experimental/EditorDataStorageFeatures/Source/TedsRevisionControl/Private/Processors/RevisionControlProcessors.cpp
@@ -25,7 +25,7 @@ namespace UE::Editor::RevisionControl::Private
 	void UpdateSCCOverlayStates(UScriptStruct* InColumn)
 	{
 		using namespace UE::Editor::DataStorage;
-		IEditorDataStorageProvider* DataStorage = GetMutableDataStorageFeature<IEditorDataStorageProvider>(StorageFeatureName);
+		ICoreProvider* DataStorage = GetMutableDataStorageFeature<ICoreProvider>(StorageFeatureName);
 
 		if(const URevisionControlDataStorageFactory* Factory = DataStorage->FindFactory<URevisionControlDataStorageFactory>())
 		{
@@ -37,7 +37,7 @@ namespace UE::Editor::RevisionControl::Private
 	void UpdateOverlayColors()
 	{
 		using namespace UE::Editor::DataStorage;
-		IEditorDataStorageProvider* DataStorage = GetMutableDataStorageFeature<IEditorDataStorageProvider>(StorageFeatureName);
+		ICoreProvider* DataStorage = GetMutableDataStorageFeature<ICoreProvider>(StorageFeatureName);
 
 		if(const URevisionControlDataStorageFactory* Factory = DataStorage->FindFactory<URevisionControlDataStorageFactory>())
 		{
@@ -229,8 +229,9 @@ namespace UE::Editor::RevisionControl::Private
 	}
 } // namespace UE::Editor::RevisionControl::Private
 
-void URevisionControlDataStorageFactory::RegisterTables(IEditorDataStorageProvider& DataStorage)
+void URevisionControlDataStorageFactory::RegisterTables(UE::Editor::DataStorage::ICoreProvider& DataStorage)
 {
+	using namespace UE::Editor::DataStorage;
 	DataStorage.RegisterTable(
 		TTypedElementColumnTypeList<
 			FTypedElementPackagePathColumn, FTypedElementPackageLoadedPathColumn,
@@ -238,7 +239,7 @@ void URevisionControlDataStorageFactory::RegisterTables(IEditorDataStorageProvid
 		FName("Editor_RevisionControlTable"));
 }
 
-void URevisionControlDataStorageFactory::RegisterQueries(IEditorDataStorageProvider& DataStorage)
+void URevisionControlDataStorageFactory::RegisterQueries(UE::Editor::DataStorage::ICoreProvider& DataStorage)
 {
 	using namespace UE::Editor::DataStorage;
 	using namespace UE::Editor::RevisionControl::Private;
@@ -310,17 +311,17 @@ void URevisionControlDataStorageFactory::RegisterQueries(IEditorDataStorageProvi
 	RegisterGeneralQueries(DataStorage);
 }
 
-void URevisionControlDataStorageFactory::UpdateOverlaysForSCCState(IEditorDataStorageProvider* DataStorage, const UScriptStruct* Column) const
+void URevisionControlDataStorageFactory::UpdateOverlaysForSCCState(UE::Editor::DataStorage::ICoreProvider* DataStorage, const UScriptStruct* Column) const
 {
 	DataStorage->ActivateQueries(TEXT("UpdateSCCForActors"));
 }
 
-void URevisionControlDataStorageFactory::UpdateOverlayColors(IEditorDataStorageProvider* DataStorage) const
+void URevisionControlDataStorageFactory::UpdateOverlayColors(UE::Editor::DataStorage::ICoreProvider* DataStorage) const
 {
 	DataStorage->ActivateQueries(TEXT("UpdateOverlayForActors"));
 }
 
-void URevisionControlDataStorageFactory::RegisterFetchUpdates(IEditorDataStorageProvider& DataStorage)
+void URevisionControlDataStorageFactory::RegisterFetchUpdates(UE::Editor::DataStorage::ICoreProvider& DataStorage)
 {
 	using namespace UE::Editor::DataStorage::Queries;
 	
@@ -348,7 +349,7 @@ void URevisionControlDataStorageFactory::RegisterFetchUpdates(IEditorDataStorage
 	}
 }
 
-void URevisionControlDataStorageFactory::RegisterApplyOverlays(IEditorDataStorageProvider& DataStorage)
+void URevisionControlDataStorageFactory::RegisterApplyOverlays(UE::Editor::DataStorage::ICoreProvider& DataStorage)
 {
 	using namespace UE::Editor::DataStorage::Queries;
 	using namespace UE::Editor::RevisionControl::Private;
@@ -545,7 +546,7 @@ void URevisionControlDataStorageFactory::RegisterApplyOverlays(IEditorDataStorag
 	}
 }
 
-void URevisionControlDataStorageFactory::RegisterRemoveOverlays(IEditorDataStorageProvider& DataStorage)
+void URevisionControlDataStorageFactory::RegisterRemoveOverlays(UE::Editor::DataStorage::ICoreProvider& DataStorage)
 {
 	using namespace UE::Editor::DataStorage::Queries;
 	
@@ -572,7 +573,7 @@ void URevisionControlDataStorageFactory::RegisterRemoveOverlays(IEditorDataStora
 	}
 }
 
-void URevisionControlDataStorageFactory::RegisterGeneralQueries(IEditorDataStorageProvider& DataStorage)
+void URevisionControlDataStorageFactory::RegisterGeneralQueries(UE::Editor::DataStorage::ICoreProvider& DataStorage)
 {
 	using namespace UE::Editor::DataStorage::Queries;
 	using namespace UE::Editor::RevisionControl::Private;
diff --git a/Engine/Plugins/Experimental/EditorDataStorageFeatures/Source/TedsRevisionControl/Private/Processors/RevisionControlProcessors.h b/Engine/Plugins/Experimental/EditorDataStorageFeatures/Source/TedsRevisionControl/Private/Processors/RevisionControlProcessors.h
--- a/Engine/Plugins/Experimental/EditorDataStorageFeatures/Source/TedsRevisionControl/Private/Processors/RevisionControlProcessors.h
+++ b/Engine/Plugins/Experimental/EditorDataStorageFeatures/Source/TedsRevisionControl/Private/Processors/RevisionControlProcessors.h
@@ -8,7 +8,10 @@
 
 #include "RevisionControlProcessors.generated.h"
 
-class IEditorDataStorageProvider;
+namespace UE::Editor::DataStorage
+{
+	class ICoreProvider;
+} // namepsace UE::Editor::DataStorage
 
 UCLASS()
 class URevisionControlDataStorageFactory : public UEditorDataStorageFactory
@@ -18,21 +21,21 @@ class URevisionControlDataStorageFactory : public UEditorDataStorageFactory
 public:
 	~URevisionControlDataStorageFactory() override = default;
 
-	void RegisterTables(IEditorDataStorageProvider& DataStorage) override;
-	void RegisterQueries(IEditorDataStorageProvider& DataStorage) override;
+	void RegisterTables(UE::Editor::DataStorage::ICoreProvider& DataStorage) override;
+	void RegisterQueries(UE::Editor::DataStorage::ICoreProvider& DataStorage) override;
 
 	// Update the overlays for all SCC rows that contain the Column specified
-	void UpdateOverlaysForSCCState(IEditorDataStorageProvider* DataStorage, const UScriptStruct* Column) const;
+	void UpdateOverlaysForSCCState(UE::Editor::DataStorage::ICoreProvider* DataStorage, const UScriptStruct* Column) const;
 
 	// Update the color for all actors that currently have an overlay
-	void UpdateOverlayColors(IEditorDataStorageProvider* DataStorage) const;
+	void UpdateOverlayColors(UE::Editor::DataStorage::ICoreProvider* DataStorage) const;
 
 
 private:
-	void RegisterFetchUpdates(IEditorDataStorageProvider& DataStorage);
-	void RegisterApplyOverlays(IEditorDataStorageProvider& DataStorage);
-	void RegisterRemoveOverlays(IEditorDataStorageProvider& DataStorage);
-	void RegisterGeneralQueries(IEditorDataStorageProvider& DataStorage);
+	void RegisterFetchUpdates(UE::Editor::DataStorage::ICoreProvider& DataStorage);
+	void RegisterApplyOverlays(UE::Editor::DataStorage::ICoreProvider& DataStorage);
+	void RegisterRemoveOverlays(UE::Editor::DataStorage::ICoreProvider& DataStorage);
+	void RegisterGeneralQueries(UE::Editor::DataStorage::ICoreProvider& DataStorage);
 
 	UE::Editor::DataStorage::QueryHandle FetchUpdates = UE::Editor::DataStorage::InvalidQueryHandle;
 	UE::Editor::DataStorage::QueryHandle ApplyOverlaysObjectToSCC = UE::Editor::DataStorage::InvalidQueryHandle;
diff --git a/Engine/Plugins/Experimental/EditorDataStorageFeatures/Source/TedsRevisionControl/Private/Queries/ObjectPackagePathToColumnQueries.cpp b/Engine/Plugins/Experimental/EditorDataStorageFeatures/Source/TedsRevisionControl/Private/Queries/ObjectPackagePathToColumnQueries.cpp
--- a/Engine/Plugins/Experimental/EditorDataStorageFeatures/Source/TedsRevisionControl/Private/Queries/ObjectPackagePathToColumnQueries.cpp
+++ b/Engine/Plugins/Experimental/EditorDataStorageFeatures/Source/TedsRevisionControl/Private/Queries/ObjectPackagePathToColumnQueries.cpp
@@ -21,7 +21,7 @@ namespace UE::Editor::RevisionControl::Private
 		TEXT("Automatically query revision control provider and fill information into TEDS")
 	);
 
-	static void ResolvePackageReference(IEditorDataStorageProvider::IQueryContext& Context, const UPackage* Package, RowHandle Row, RowHandle PackageRow)
+	static void ResolvePackageReference(IQueryContext& Context, const UPackage* Package, RowHandle Row, RowHandle PackageRow)
 	{
 		FTypedElementPackageReference PackageReference;
 		PackageReference.Row = PackageRow;
@@ -42,7 +42,7 @@ namespace UE::Editor::RevisionControl::Private
 	};
 } // namespace UE::Editor::DataStorage
 
-void UTypedElementUObjectPackagePathFactory::RegisterQueries(IEditorDataStorageProvider& DataStorage)
+void UTypedElementUObjectPackagePathFactory::RegisterQueries(UE::Editor::DataStorage::ICoreProvider& DataStorage)
 {
 	using namespace UE::Editor::DataStorage::Queries;
 	
@@ -88,7 +88,7 @@ void UTypedElementUObjectPackagePathFactory::RegisterQueries(IEditorDataStorageP
 	}
 }
 
-void UTypedElementUObjectPackagePathFactory::RegisterTryAddPackageRef(IEditorDataStorageProvider& DataStorage)
+void UTypedElementUObjectPackagePathFactory::RegisterTryAddPackageRef(UE::Editor::DataStorage::ICoreProvider& DataStorage)
 {
 	using namespace UE::Editor::DataStorage::Queries;
 	
diff --git a/Engine/Plugins/Experimental/EditorDataStorageFeatures/Source/TedsRevisionControl/Private/Queries/ObjectPackagePathToColumnQueries.h b/Engine/Plugins/Experimental/EditorDataStorageFeatures/Source/TedsRevisionControl/Private/Queries/ObjectPackagePathToColumnQueries.h
--- a/Engine/Plugins/Experimental/EditorDataStorageFeatures/Source/TedsRevisionControl/Private/Queries/ObjectPackagePathToColumnQueries.h
+++ b/Engine/Plugins/Experimental/EditorDataStorageFeatures/Source/TedsRevisionControl/Private/Queries/ObjectPackagePathToColumnQueries.h
@@ -7,7 +7,10 @@
 
 #include "ObjectPackagePathToColumnQueries.generated.h"
 
-class IEditorDataStorageProvider;
+namespace UE::Editor::DataStorage
+{
+	class ICoreProvider;
+} // namespace UE::Editor::DataStorage
 
 UCLASS()
 class UTypedElementUObjectPackagePathFactory : public UEditorDataStorageFactory
@@ -17,9 +20,9 @@ class UTypedElementUObjectPackagePathFactory : public UEditorDataStorageFactory
 public:
 	~UTypedElementUObjectPackagePathFactory() override = default;
 
-	void RegisterQueries(IEditorDataStorageProvider& DataStorage) override;
+	void RegisterQueries(UE::Editor::DataStorage::ICoreProvider& DataStorage) override;
 
 private:
-	void RegisterTryAddPackageRef(IEditorDataStorageProvider& DataStorage);
+	void RegisterTryAddPackageRef(UE::Editor::DataStorage::ICoreProvider& DataStorage);
 	UE::Editor::DataStorage::QueryHandle TryAddPackageRef = UE::Editor::DataStorage::InvalidQueryHandle;
 };
diff --git a/Engine/Plugins/Experimental/EditorDataStorageFeatures/Source/TedsSettings/Private/TedsSettingsFactory.cpp b/Engine/Plugins/Experimental/EditorDataStorageFeatures/Source/TedsSettings/Private/TedsSettingsFactory.cpp
--- a/Engine/Plugins/Experimental/EditorDataStorageFeatures/Source/TedsSettings/Private/TedsSettingsFactory.cpp
+++ b/Engine/Plugins/Experimental/EditorDataStorageFeatures/Source/TedsSettings/Private/TedsSettingsFactory.cpp
@@ -5,7 +5,7 @@
 #include "TedsSettingsColumns.h"
 #include "TedsSettingsWidgets.h"
 
-void UTedsSettingsFactory::RegisterWidgetConstructors(IEditorDataStorageProvider& DataStorage, IEditorDataStorageUiProvider& DataStorageUi) const
+void UTedsSettingsFactory::RegisterWidgetConstructors(UE::Editor::DataStorage::ICoreProvider& DataStorage, UE::Editor::DataStorage::IUiProvider& DataStorageUi) const
 {
 	using namespace UE::Editor::DataStorage::Queries;
 
diff --git a/Engine/Plugins/Experimental/EditorDataStorageFeatures/Source/TedsSettings/Private/TedsSettingsFactory.h b/Engine/Plugins/Experimental/EditorDataStorageFeatures/Source/TedsSettings/Private/TedsSettingsFactory.h
--- a/Engine/Plugins/Experimental/EditorDataStorageFeatures/Source/TedsSettings/Private/TedsSettingsFactory.h
+++ b/Engine/Plugins/Experimental/EditorDataStorageFeatures/Source/TedsSettings/Private/TedsSettingsFactory.h
@@ -16,6 +16,7 @@ public:
 
 	~UTedsSettingsFactory() override = default;
 
-	void RegisterWidgetConstructors(IEditorDataStorageProvider& DataStorage, IEditorDataStorageUiProvider& DataStorageUi) const override;
+	void RegisterWidgetConstructors(UE::Editor::DataStorage::ICoreProvider& DataStorage,
+		UE::Editor::DataStorage::IUiProvider& DataStorageUi) const override;
 
 };
diff --git a/Engine/Plugins/Experimental/EditorDataStorageFeatures/Source/TedsSettings/Private/TedsSettingsManager.cpp b/Engine/Plugins/Experimental/EditorDataStorageFeatures/Source/TedsSettings/Private/TedsSettingsManager.cpp
--- a/Engine/Plugins/Experimental/EditorDataStorageFeatures/Source/TedsSettings/Private/TedsSettingsManager.cpp
+++ b/Engine/Plugins/Experimental/EditorDataStorageFeatures/Source/TedsSettings/Private/TedsSettingsManager.cpp
@@ -68,7 +68,7 @@ void FTedsSettingsManager::Initialize()
 		using namespace UE::Editor::DataStorage;
 		auto OnDataStorage = [this]
 			{
-				IEditorDataStorageProvider* DataStorage = GetMutableDataStorageFeature<IEditorDataStorageProvider>(StorageFeatureName);
+				ICoreProvider* DataStorage = GetMutableDataStorageFeature<ICoreProvider>(StorageFeatureName);
 				check(DataStorage);
 
 				RegisterTables(*DataStorage);
@@ -98,7 +98,7 @@ void FTedsSettingsManager::Shutdown()
 
 		if (AreEditorDataStorageFeaturesEnabled())
 		{
-			IEditorDataStorageProvider* DataStorage = GetMutableDataStorageFeature<IEditorDataStorageProvider>(StorageFeatureName);
+			ICoreProvider* DataStorage = GetMutableDataStorageFeature<ICoreProvider>(StorageFeatureName);
 			check(DataStorage);
 
 			UnregisterActiveSettings();
@@ -117,7 +117,7 @@ UE::Editor::DataStorage::RowHandle FTedsSettingsManager::FindOrAddSettingsSectio
 
 	if (bIsInitialized)
 	{
-		if (IEditorDataStorageProvider* DataStorage = GetMutableDataStorageFeature<IEditorDataStorageProvider>(StorageFeatureName))
+		if (ICoreProvider* DataStorage = GetMutableDataStorageFeature<ICoreProvider>(StorageFeatureName))
 		{
 			// Ensure that the SettingsInactiveSectionTable is registered since it is possible for OnDataStorageInterfacesSet to be called
 			// for a delegate bound in another module before it has been called for our bound delegate.
@@ -152,7 +152,7 @@ bool FTedsSettingsManager::GetSettingsSectionFromRow(UE::Editor::DataStorage::Ro
 
 	if (bIsInitialized)
 	{
-		if (const IEditorDataStorageProvider* DataStorage = GetDataStorageFeature<IEditorDataStorageProvider>(StorageFeatureName))
+		if (const ICoreProvider* DataStorage = GetDataStorageFeature<ICoreProvider>(StorageFeatureName))
 		{
 			if (DataStorage->IsRowAvailable(Row) && (DataStorage->HasColumns<FSettingsSectionTag>(Row) || DataStorage->HasColumns<FSettingsInactiveSectionTag>(Row)))
 			{
@@ -175,23 +175,24 @@ bool FTedsSettingsManager::GetSettingsSectionFromRow(UE::Editor::DataStorage::Ro
 	return false;
 }
 
-void FTedsSettingsManager::RegisterTables(IEditorDataStorageProvider& DataStorage)
+void FTedsSettingsManager::RegisterTables(UE::Editor::DataStorage::ICoreProvider& DataStorage)
 {
-	if (SettingsContainerTable == UE::Editor::DataStorage::InvalidTableHandle)
+	using namespace UE::Editor::DataStorage;
+	if (SettingsContainerTable == InvalidTableHandle)
 	{
 		SettingsContainerTable = DataStorage.RegisterTable(
 			TTypedElementColumnTypeList<FNameColumn, FDisplayNameColumn, FDescriptionColumn, FSettingsContainerTag>(),
 			FName(TEXT("Editor_SettingsContainerTable")));
 	}
 
-	if (SettingsCategoryTable == UE::Editor::DataStorage::InvalidTableHandle)
+	if (SettingsCategoryTable == InvalidTableHandle)
 	{
 		SettingsCategoryTable = DataStorage.RegisterTable(
 			TTypedElementColumnTypeList<FSettingsContainerReferenceColumn, FNameColumn, FDisplayNameColumn, FDescriptionColumn, FSettingsCategoryTag>(),
 			FName(TEXT("Editor_SettingsCategoryTable")));
 	}
 
-	if (SettingsInactiveSectionTable == UE::Editor::DataStorage::InvalidTableHandle)
+	if (SettingsInactiveSectionTable == InvalidTableHandle)
 	{
 		SettingsInactiveSectionTable = DataStorage.RegisterTable(
 			TTypedElementColumnTypeList<FSettingsContainerReferenceColumn, FSettingsCategoryReferenceColumn, FNameColumn, FSettingsInactiveSectionTag>(),
@@ -199,7 +200,7 @@ void FTedsSettingsManager::RegisterTables(IEditorDataStorageProvider& DataStorag
 	}
 }
 
-void FTedsSettingsManager::RegisterQueries(IEditorDataStorageProvider& DataStorage)
+void FTedsSettingsManager::RegisterQueries(UE::Editor::DataStorage::ICoreProvider& DataStorage)
 {
 	using namespace UE::Editor::DataStorage::Queries;
 
@@ -224,7 +225,7 @@ void FTedsSettingsManager::RegisterQueries(IEditorDataStorageProvider& DataStora
 	}
 }
 
-void FTedsSettingsManager::UnregisterQueries(IEditorDataStorageProvider& DataStorage)
+void FTedsSettingsManager::UnregisterQueries(UE::Editor::DataStorage::ICoreProvider& DataStorage)
 {
 	DataStorage.UnregisterQuery(SelectAllActiveSettingsQuery);
 	SelectAllActiveSettingsQuery = UE::Editor::DataStorage::InvalidQueryHandle;
@@ -260,7 +261,7 @@ void FTedsSettingsManager::RegisterSettingsContainer(const FName& ContainerName)
 	ISettingsModule* SettingsModule = FModuleManager::GetModulePtr<ISettingsModule>("Settings");
 	check(SettingsModule);
 
-	IEditorDataStorageProvider* DataStorage = GetMutableDataStorageFeature<IEditorDataStorageProvider>(StorageFeatureName);
+	ICoreProvider* DataStorage = GetMutableDataStorageFeature<ICoreProvider>(StorageFeatureName);
 	check(DataStorage);
 
 	UE_LOG(LogTedsSettings, Verbose, TEXT("Register Settings Container : '%s'"), *ContainerName.ToString());
@@ -310,10 +311,10 @@ void FTedsSettingsManager::UnregisterActiveSettings()
 	ISettingsModule* SettingsModule = FModuleManager::GetModulePtr<ISettingsModule>("Settings");
 	check(SettingsModule);
 
-	IEditorDataStorageProvider* DataStorage = GetMutableDataStorageFeature<IEditorDataStorageProvider>(StorageFeatureName);
+	ICoreProvider* DataStorage = GetMutableDataStorageFeature<ICoreProvider>(StorageFeatureName);
 	check(DataStorage);
 
-	IEditorDataStorageCompatibilityProvider* DataStorageCompatibility = GetMutableDataStorageFeature<IEditorDataStorageCompatibilityProvider>(CompatibilityFeatureName);
+	ICompatibilityProvider* DataStorageCompatibility = GetMutableDataStorageFeature<ICompatibilityProvider>(CompatibilityFeatureName);
 	check(DataStorageCompatibility);
 
 	SettingsModule->OnContainerAdded().RemoveAll(this);
@@ -374,7 +375,7 @@ void FTedsSettingsManager::UnregisterInactiveSettings()
 	using namespace UE::Editor::DataStorage::Queries;
 	using namespace UE::Editor::Settings::Private;
 
-	IEditorDataStorageProvider* DataStorage = GetMutableDataStorageFeature<IEditorDataStorageProvider>(StorageFeatureName);
+	ICoreProvider* DataStorage = GetMutableDataStorageFeature<ICoreProvider>(StorageFeatureName);
 	check(DataStorage);
 	
 	TSet<RowHandle> RowHandles;
@@ -398,10 +399,10 @@ void FTedsSettingsManager::UpdateSettingsCategory(TSharedPtr<ISettingsCategory>
 
 	TRACE_CPUPROFILER_EVENT_SCOPE(TedsSettingsManager.UpdateSettingsCategory);
 
-	IEditorDataStorageProvider* DataStorage = GetMutableDataStorageFeature<IEditorDataStorageProvider>(StorageFeatureName);
+	ICoreProvider* DataStorage = GetMutableDataStorageFeature<ICoreProvider>(StorageFeatureName);
 	check(DataStorage);
 
-	IEditorDataStorageCompatibilityProvider* DataStorageCompatibility = GetMutableDataStorageFeature<IEditorDataStorageCompatibilityProvider>(CompatibilityFeatureName);
+	ICompatibilityProvider* DataStorageCompatibility = GetMutableDataStorageFeature<ICompatibilityProvider>(CompatibilityFeatureName);
 	check(DataStorageCompatibility);
 
 	const FName& ContainerName = DataStorage->GetColumn<FNameColumn>(ContainerRow)->Name;
@@ -562,7 +563,7 @@ void FTedsSettingsManager::UpdateSettingsCategory(TSharedPtr<ISettingsCategory>
 	}
 }
 
-void FTedsSettingsManager::AddColumns(IEditorDataStorageProvider& DataStorage, UE::Editor::DataStorage::RowHandle Row, const TArray<void*>& Columns, const TArray<const UScriptStruct*>& ColumnTypes)
+void FTedsSettingsManager::AddColumns(UE::Editor::DataStorage::ICoreProvider& DataStorage, UE::Editor::DataStorage::RowHandle Row, const TArray<void*>& Columns, const TArray<const UScriptStruct*>& ColumnTypes)
 {
 	for (int32 ColumnIndex = 0; ColumnIndex < Columns.Num(); ++ColumnIndex)
 	{
diff --git a/Engine/Plugins/Experimental/EditorDataStorageFeatures/Source/TedsSettings/Private/TedsSettingsManager.h b/Engine/Plugins/Experimental/EditorDataStorageFeatures/Source/TedsSettings/Private/TedsSettingsManager.h
--- a/Engine/Plugins/Experimental/EditorDataStorageFeatures/Source/TedsSettings/Private/TedsSettingsManager.h
+++ b/Engine/Plugins/Experimental/EditorDataStorageFeatures/Source/TedsSettings/Private/TedsSettingsManager.h
@@ -27,10 +27,10 @@ public:
 
 private:
 
-	void RegisterTables(IEditorDataStorageProvider& DataStorage);
+	void RegisterTables(UE::Editor::DataStorage::ICoreProvider& DataStorage);
 
-	void RegisterQueries(IEditorDataStorageProvider& DataStorage);
-	void UnregisterQueries(IEditorDataStorageProvider& DataStorage);
+	void RegisterQueries(UE::Editor::DataStorage::ICoreProvider& DataStorage);
+	void UnregisterQueries(UE::Editor::DataStorage::ICoreProvider& DataStorage);
 
 	void RegisterActiveSettings();
 	void UnregisterActiveSettings();
@@ -40,7 +40,7 @@ private:
 
 	void UpdateSettingsCategory(TSharedPtr<ISettingsCategory> SettingsCategory, UE::Editor::DataStorage::RowHandle ContainerRow, const bool bQueryExistingRows = true);
 
-	void AddColumns(IEditorDataStorageProvider& DataStorage, UE::Editor::DataStorage::RowHandle Row, const TArray<void*>& Columns, const TArray<const UScriptStruct*>& ColumnTypes);
+	void AddColumns(UE::Editor::DataStorage::ICoreProvider& DataStorage, UE::Editor::DataStorage::RowHandle Row, const TArray<void*>& Columns, const TArray<const UScriptStruct*>& ColumnTypes);
 
 	bool bIsInitialized;
 	UE::Editor::DataStorage::QueryHandle SelectAllActiveSettingsQuery;
diff --git a/Engine/Plugins/Experimental/EditorDataStorageFeatures/Source/TedsSettings/Private/TedsSettingsWidgets.cpp b/Engine/Plugins/Experimental/EditorDataStorageFeatures/Source/TedsSettings/Private/TedsSettingsWidgets.cpp
--- a/Engine/Plugins/Experimental/EditorDataStorageFeatures/Source/TedsSettings/Private/TedsSettingsWidgets.cpp
+++ b/Engine/Plugins/Experimental/EditorDataStorageFeatures/Source/TedsSettings/Private/TedsSettingsWidgets.cpp
@@ -16,8 +16,8 @@ FSettingsContainerReferenceWidgetConstructor::FSettingsContainerReferenceWidgetC
 }
 
 TSharedPtr<SWidget> FSettingsContainerReferenceWidgetConstructor::CreateWidget(
-	IEditorDataStorageProvider* DataStorage,
-	IEditorDataStorageUiProvider* DataStorageUi,
+	UE::Editor::DataStorage::ICoreProvider* DataStorage,
+	UE::Editor::DataStorage::IUiProvider* DataStorageUi,
 	UE::Editor::DataStorage::RowHandle TargetRow,
 	UE::Editor::DataStorage::RowHandle WidgetRow,
 	const UE::Editor::DataStorage::FMetaDataView& Arguments)
@@ -65,8 +65,8 @@ FSettingsCategoryReferenceWidgetConstructor::FSettingsCategoryReferenceWidgetCon
 }
 
 TSharedPtr<SWidget> FSettingsCategoryReferenceWidgetConstructor::CreateWidget(
-	IEditorDataStorageProvider* DataStorage,
-	IEditorDataStorageUiProvider* DataStorageUi,
+	UE::Editor::DataStorage::ICoreProvider* DataStorage,
+	UE::Editor::DataStorage::IUiProvider* DataStorageUi,
 	UE::Editor::DataStorage::RowHandle TargetRow,
 	UE::Editor::DataStorage::RowHandle WidgetRow,
 	const UE::Editor::DataStorage::FMetaDataView& Arguments)
@@ -114,8 +114,8 @@ FSettingsSectionWidgetConstructor::FSettingsSectionWidgetConstructor()
 }
 
 TSharedPtr<SWidget> FSettingsSectionWidgetConstructor::CreateWidget(
-	IEditorDataStorageProvider* DataStorage,
-	IEditorDataStorageUiProvider* DataStorageUi,
+	UE::Editor::DataStorage::ICoreProvider* DataStorage,
+	UE::Editor::DataStorage::IUiProvider* DataStorageUi,
 	UE::Editor::DataStorage::RowHandle TargetRow,
 	UE::Editor::DataStorage::RowHandle WidgetRow,
 	const UE::Editor::DataStorage::FMetaDataView& Arguments)
diff --git a/Engine/Plugins/Experimental/EditorDataStorageFeatures/Source/TedsSettings/Private/TedsSettingsWidgets.h b/Engine/Plugins/Experimental/EditorDataStorageFeatures/Source/TedsSettings/Private/TedsSettingsWidgets.h
--- a/Engine/Plugins/Experimental/EditorDataStorageFeatures/Source/TedsSettings/Private/TedsSettingsWidgets.h
+++ b/Engine/Plugins/Experimental/EditorDataStorageFeatures/Source/TedsSettings/Private/TedsSettingsWidgets.h
@@ -18,8 +18,8 @@ public:
 	~FSettingsContainerReferenceWidgetConstructor() override = default;
 
 	virtual TSharedPtr<SWidget> CreateWidget(
-		IEditorDataStorageProvider* DataStorage,
-		IEditorDataStorageUiProvider* DataStorageUi,
+		UE::Editor::DataStorage::ICoreProvider* DataStorage,
+		UE::Editor::DataStorage::IUiProvider* DataStorageUi,
 		UE::Editor::DataStorage::RowHandle TargetRow,
 		UE::Editor::DataStorage::RowHandle WidgetRow,
 		const UE::Editor::DataStorage::FMetaDataView& Arguments) override;
@@ -37,8 +37,8 @@ public:
 	~FSettingsCategoryReferenceWidgetConstructor() override = default;
 
 	virtual TSharedPtr<SWidget> CreateWidget(
-		IEditorDataStorageProvider* DataStorage,
-		IEditorDataStorageUiProvider* DataStorageUi,
+		UE::Editor::DataStorage::ICoreProvider* DataStorage,
+		UE::Editor::DataStorage::IUiProvider* DataStorageUi,
 		UE::Editor::DataStorage::RowHandle TargetRow,
 		UE::Editor::DataStorage::RowHandle WidgetRow,
 		const UE::Editor::DataStorage::FMetaDataView& Arguments) override;
@@ -56,8 +56,8 @@ public:
 	~FSettingsSectionWidgetConstructor() override = default;
 
 	virtual TSharedPtr<SWidget> CreateWidget(
-		IEditorDataStorageProvider* DataStorage,
-		IEditorDataStorageUiProvider* DataStorageUi,
+		UE::Editor::DataStorage::ICoreProvider* DataStorage,
+		UE::Editor::DataStorage::IUiProvider* DataStorageUi,
 		UE::Editor::DataStorage::RowHandle TargetRow,
 		UE::Editor::DataStorage::RowHandle WidgetRow,
 		const UE::Editor::DataStorage::FMetaDataView& Arguments) override;
diff --git a/Engine/Plugins/Experimental/EditorDataStorageFeatures/Source/TedsSettings/Tests/TedsSettingsTests.cpp b/Engine/Plugins/Experimental/EditorDataStorageFeatures/Source/TedsSettings/Tests/TedsSettingsTests.cpp
--- a/Engine/Plugins/Experimental/EditorDataStorageFeatures/Source/TedsSettings/Tests/TedsSettingsTests.cpp
+++ b/Engine/Plugins/Experimental/EditorDataStorageFeatures/Source/TedsSettings/Tests/TedsSettingsTests.cpp
@@ -21,8 +21,8 @@ namespace UE::Editor::Settings::Tests
 {
 	BEGIN_DEFINE_SPEC(FTedsSettingsTestFixture, "Editor.DataStorage.Settings", EAutomationTestFlags::EditorContext | EAutomationTestFlags::EngineFilter)
 	ISettingsModule* SettingsModule = nullptr;
-	IEditorDataStorageProvider* DataStorage = nullptr;
-	IEditorDataStorageCompatibilityProvider* DataStorageCompatibility = nullptr;
+	DataStorage::ICoreProvider* DataStorage = nullptr;
+	DataStorage::ICompatibilityProvider* DataStorageCompatibility = nullptr;
 	DataStorage::QueryHandle CountAllSettingsQuery = DataStorage::InvalidQueryHandle;
 
 	uint32 BeforeRowCount = 0;
@@ -75,10 +75,10 @@ namespace UE::Editor::Settings::Tests
 			SettingsModule = FModuleManager::GetModulePtr<ISettingsModule>("Settings");
 			check(SettingsModule != nullptr);
 
-			DataStorage = GetMutableDataStorageFeature<IEditorDataStorageProvider>(StorageFeatureName);
+			DataStorage = GetMutableDataStorageFeature<ICoreProvider>(StorageFeatureName);
 			check(DataStorage != nullptr);
 
-			DataStorageCompatibility = GetMutableDataStorageFeature<IEditorDataStorageCompatibilityProvider>(CompatibilityFeatureName);
+			DataStorageCompatibility = GetMutableDataStorageFeature<ICompatibilityProvider>(CompatibilityFeatureName);
 			check(DataStorageCompatibility != nullptr);
 
 			{
diff --git a/Engine/Plugins/Experimental/EditorDataStorageFeatures/Source/TedsTableViewer/Private/TedsTableViewerColumn.cpp b/Engine/Plugins/Experimental/EditorDataStorageFeatures/Source/TedsTableViewer/Private/TedsTableViewerColumn.cpp
--- a/Engine/Plugins/Experimental/EditorDataStorageFeatures/Source/TedsTableViewer/Private/TedsTableViewerColumn.cpp
+++ b/Engine/Plugins/Experimental/EditorDataStorageFeatures/Source/TedsTableViewer/Private/TedsTableViewerColumn.cpp
@@ -28,9 +28,9 @@ namespace UE::Editor::DataStorage
 		, MatchedColumns(InMatchedColumns)
 		, WidgetMetaData(InWidgetMetaData)
 	{
-		Storage = GetMutableDataStorageFeature<IEditorDataStorageProvider>(StorageFeatureName);
-		StorageUi = GetMutableDataStorageFeature<IEditorDataStorageUiProvider>(UiFeatureName);
-		StorageCompatibility = GetMutableDataStorageFeature<IEditorDataStorageCompatibilityProvider>(CompatibilityFeatureName);
+		Storage = GetMutableDataStorageFeature<ICoreProvider>(StorageFeatureName);
+		StorageUi = GetMutableDataStorageFeature<IUiProvider>(UiFeatureName);
+		StorageCompatibility = GetMutableDataStorageFeature<ICompatibilityProvider>(CompatibilityFeatureName);
 
 		// Store the matched columns as a query condition that requires all of them (i.e AND's them)
 		for(const TWeakObjectPtr<const UScriptStruct>& Column : MatchedColumns)
@@ -47,7 +47,7 @@ namespace UE::Editor::DataStorage
 		UnRegisterQueries();
 	}
 
-	TSharedPtr<SWidget> FTedsTableViewerColumn::ConstructRowWidget(RowHandle InRowHandle, TFunction<void(IEditorDataStorageProvider&, const RowHandle&)> WidgetRowSetupDelegate/* = nullptr*/) const
+	TSharedPtr<SWidget> FTedsTableViewerColumn::ConstructRowWidget(RowHandle InRowHandle, TFunction<void(ICoreProvider&, const RowHandle&)> WidgetRowSetupDelegate/* = nullptr*/) const
 	{
 		TSharedPtr<SWidget> RowWidget;
 		
@@ -350,7 +350,7 @@ namespace UE::Editor::DataStorage
 		return MatchedColumns;
 	}
 
-	IEditorDataStorageProvider* FTedsTableViewerColumn::GetStorage() const
+	ICoreProvider* FTedsTableViewerColumn::GetStorage() const
 	{
 		return Storage;
 	}
diff --git a/Engine/Plugins/Experimental/EditorDataStorageFeatures/Source/TedsTableViewer/Private/TedsTableViewerModel.cpp b/Engine/Plugins/Experimental/EditorDataStorageFeatures/Source/TedsTableViewer/Private/TedsTableViewerModel.cpp
--- a/Engine/Plugins/Experimental/EditorDataStorageFeatures/Source/TedsTableViewer/Private/TedsTableViewerModel.cpp
+++ b/Engine/Plugins/Experimental/EditorDataStorageFeatures/Source/TedsTableViewer/Private/TedsTableViewerModel.cpp
@@ -22,9 +22,9 @@ namespace UE::Editor::DataStorage
 		, CellWidgetPurposes(InCellWidgetPurposes)
 		, IsItemVisible(InIsItemVisibleDelegate)
 	{
-		Storage = GetMutableDataStorageFeature<IEditorDataStorageProvider>(StorageFeatureName);
-		StorageUi = GetMutableDataStorageFeature<IEditorDataStorageUiProvider>(UiFeatureName);
-		StorageCompatibility = GetMutableDataStorageFeature<IEditorDataStorageCompatibilityProvider>(CompatibilityFeatureName);
+		Storage = GetMutableDataStorageFeature<ICoreProvider>(StorageFeatureName);
+		StorageUi = GetMutableDataStorageFeature<IUiProvider>(UiFeatureName);
+		StorageCompatibility = GetMutableDataStorageFeature<ICompatibilityProvider>(CompatibilityFeatureName);
 		
 		TickerHandle = FTSTicker::GetCoreTicker().AddTicker(FTickerDelegate::CreateRaw(this, &FTedsTableViewerModel::Tick), 0);
 
@@ -151,12 +151,12 @@ namespace UE::Editor::DataStorage
 		ColumnsView.Insert(InColumn, 0);
 	}
 
-	IEditorDataStorageProvider* FTedsTableViewerModel::GetDataStorageInterface() const
+	ICoreProvider* FTedsTableViewerModel::GetDataStorageInterface() const
 	{
 		return Storage;
 	}
 
-	IEditorDataStorageUiProvider* FTedsTableViewerModel::GetDataStorageUiProvider() const
+	IUiProvider* FTedsTableViewerModel::GetDataStorageUiProvider() const
 	{
 		return StorageUi;
 	}
@@ -165,7 +165,7 @@ namespace UE::Editor::DataStorage
 	{
 		ValidateRequestedColumns();
 		
-		using MatchApproach = IEditorDataStorageUiProvider::EMatchApproach;
+		using MatchApproach = IUiProvider::EMatchApproach;
 		int32 IndexOffset = 0;
 
 		ColumnsView.Empty();
diff --git a/Engine/Plugins/Experimental/EditorDataStorageFeatures/Source/TedsTableViewer/Private/TedsTableViewerUtils.cpp b/Engine/Plugins/Experimental/EditorDataStorageFeatures/Source/TedsTableViewer/Private/TedsTableViewerUtils.cpp
--- a/Engine/Plugins/Experimental/EditorDataStorageFeatures/Source/TedsTableViewer/Private/TedsTableViewerUtils.cpp
+++ b/Engine/Plugins/Experimental/EditorDataStorageFeatures/Source/TedsTableViewer/Private/TedsTableViewerUtils.cpp
@@ -90,9 +90,9 @@ namespace UE::Editor::DataStorage::TableViewerUtils
 		return VerifiedColumnTypes;
 	}
 
-	TSharedPtr<FTypedElementWidgetConstructor> CreateHeaderWidgetConstructor(IEditorDataStorageUiProvider& StorageUi, const UE::Editor::DataStorage::FMetaDataView& InMetaData, TConstArrayView<TWeakObjectPtr<const UScriptStruct>> ColumnTypes, const TConstArrayView<FName> CellWidgetPurposes)
+	TSharedPtr<FTypedElementWidgetConstructor> CreateHeaderWidgetConstructor(IUiProvider& StorageUi, const UE::Editor::DataStorage::FMetaDataView& InMetaData, TConstArrayView<TWeakObjectPtr<const UScriptStruct>> ColumnTypes, const TConstArrayView<FName> CellWidgetPurposes)
 	{
-		using MatchApproach = IEditorDataStorageUiProvider::EMatchApproach;
+		using MatchApproach = IUiProvider::EMatchApproach;
 
 		TArray<TWeakObjectPtr<const UScriptStruct>> VerifiedColumnTypes = CreateVerifiedColumnTypeArray(ColumnTypes);
 		TSharedPtr<FTypedElementWidgetConstructor> Constructor;
@@ -145,7 +145,7 @@ namespace UE::Editor::DataStorage::TableViewerUtils
 		return nullptr;
 	}
 
-	const FSlateBrush* GetIconForRow(IEditorDataStorageProvider* DataStorage, RowHandle Row)
+	const FSlateBrush* GetIconForRow(ICoreProvider* DataStorage, RowHandle Row)
 	{
 		static TMap<FName, const FSlateBrush*> CachedIconMap;
 
@@ -201,7 +201,7 @@ namespace UE::Editor::DataStorage::TableViewerUtils
 	}
 }
 
-void UTypedElementTableViewerFactory::RegisterTables(IEditorDataStorageProvider& DataStorage)
+void UTypedElementTableViewerFactory::RegisterTables(UE::Editor::DataStorage::ICoreProvider& DataStorage)
 {
 	using namespace UE::Editor::DataStorage;
 	const TableHandle BaseWidgetTable = DataStorage.FindTable(FName(TEXT("Editor_WidgetTable")));
diff --git a/Engine/Plugins/Experimental/EditorDataStorageFeatures/Source/TedsTableViewer/Private/Widgets/RowHandleWidget.cpp b/Engine/Plugins/Experimental/EditorDataStorageFeatures/Source/TedsTableViewer/Private/Widgets/RowHandleWidget.cpp
--- a/Engine/Plugins/Experimental/EditorDataStorageFeatures/Source/TedsTableViewer/Private/Widgets/RowHandleWidget.cpp
+++ b/Engine/Plugins/Experimental/EditorDataStorageFeatures/Source/TedsTableViewer/Private/Widgets/RowHandleWidget.cpp
@@ -9,20 +9,21 @@
 #define LOCTEXT_NAMESPACE "RowHandleWidget"
 
 
-void URowHandleWidgetFactory::RegisterWidgetConstructors(IEditorDataStorageProvider& DataStorage, IEditorDataStorageUiProvider& DataStorageUi) const
+void URowHandleWidgetFactory::RegisterWidgetConstructors(UE::Editor::DataStorage::ICoreProvider& DataStorage, UE::Editor::DataStorage::IUiProvider& DataStorageUi) const
 {
 	DataStorageUi.RegisterWidgetFactory(FName(TEXT("General.Cell.RowHandle")), FRowHandleWidgetConstructor::StaticStruct());
 }
 
-void URowHandleWidgetFactory::RegisterWidgetPurposes(IEditorDataStorageUiProvider& DataStorageUi) const
+void URowHandleWidgetFactory::RegisterWidgetPurposes(UE::Editor::DataStorage::IUiProvider& DataStorageUi) const
 {
-	DataStorageUi.RegisterWidgetPurpose(FName(TEXT("General.Cell.RowHandle")), IEditorDataStorageUiProvider::EPurposeType::UniqueByName,
+	using namespace UE::Editor::DataStorage;
+	DataStorageUi.RegisterWidgetPurpose(FName(TEXT("General.Cell.RowHandle")), IUiProvider::EPurposeType::UniqueByName,
 	LOCTEXT("GeneralRowHandlePurpose", "Specific purpose to request a widget to display row handles."));
 
-	DataStorageUi.RegisterWidgetPurpose(FName(TEXT("RowDetails.Cell")), IEditorDataStorageUiProvider::EPurposeType::UniqueByNameAndColumn,
+	DataStorageUi.RegisterWidgetPurpose(FName(TEXT("RowDetails.Cell")), IUiProvider::EPurposeType::UniqueByNameAndColumn,
 	LOCTEXT("DetailsRowHandlePurpose", "Specific purpose to request a widget to display the details on a row (e.g SRowDetails)."));
 
-	DataStorageUi.RegisterWidgetPurpose(FName(TEXT("RowDetails.Cell.Large")), IEditorDataStorageUiProvider::EPurposeType::UniqueByNameAndColumn,
+	DataStorageUi.RegisterWidgetPurpose(FName(TEXT("RowDetails.Cell.Large")), IUiProvider::EPurposeType::UniqueByNameAndColumn,
 	LOCTEXT("LargeDetailsRowHandlePurpose", "Specific purpose to request a widget that is larger than a single cell to display the details on a row (e.g SRowDetails)"));
 
 
@@ -41,7 +42,9 @@ TSharedPtr<SWidget> FRowHandleWidgetConstructor::CreateWidget(const UE::Editor::
 		.VAlign(VAlign_Center)
 		.Padding(8, 0, 0, 0);}
 
-bool FRowHandleWidgetConstructor::FinalizeWidget(IEditorDataStorageProvider* DataStorage, IEditorDataStorageUiProvider* DataStorageUi, UE::Editor::DataStorage::RowHandle Row, const TSharedPtr<SWidget>& Widget)
+bool FRowHandleWidgetConstructor::FinalizeWidget(UE::Editor::DataStorage::ICoreProvider* DataStorage,
+	UE::Editor::DataStorage::IUiProvider* DataStorageUi, 
+	UE::Editor::DataStorage::RowHandle Row, const TSharedPtr<SWidget>& Widget)
 {
 
 	checkf(Widget->GetType() == SBox::StaticWidgetClass().GetWidgetType(),
diff --git a/Engine/Plugins/Experimental/EditorDataStorageFeatures/Source/TedsTableViewer/Private/Widgets/SRowDetails.cpp b/Engine/Plugins/Experimental/EditorDataStorageFeatures/Source/TedsTableViewer/Private/Widgets/SRowDetails.cpp
--- a/Engine/Plugins/Experimental/EditorDataStorageFeatures/Source/TedsTableViewer/Private/Widgets/SRowDetails.cpp
+++ b/Engine/Plugins/Experimental/EditorDataStorageFeatures/Source/TedsTableViewer/Private/Widgets/SRowDetails.cpp
@@ -38,8 +38,8 @@ namespace UE::Editor::DataStorage
 
 		checkf(AreEditorDataStorageFeaturesEnabled(), TEXT("Unable to initialize SRowDetails without the editor data storage interfaces."));
 
-		DataStorage = GetMutableDataStorageFeature<IEditorDataStorageProvider>(StorageFeatureName);
-		DataStorageUi = GetMutableDataStorageFeature<IEditorDataStorageUiProvider>(UiFeatureName);
+		DataStorage = GetMutableDataStorageFeature<ICoreProvider>(StorageFeatureName);
+		DataStorageUi = GetMutableDataStorageFeature<IUiProvider>(UiFeatureName);
 
 		
 		ChildSlot
@@ -79,7 +79,7 @@ namespace UE::Editor::DataStorage
 			
 			for (const FName& Purpose : WidgetPurposes)
 			{
-				DataStorageUi->CreateWidgetConstructors(Purpose, IEditorDataStorageUiProvider::EMatchApproach::LongestMatch,
+				DataStorageUi->CreateWidgetConstructors(Purpose, IUiProvider::EMatchApproach::LongestMatch,
 					Columns, {}, [this, Row](
 						TUniquePtr<FTypedElementWidgetConstructor> Constructor, TConstArrayView<TWeakObjectPtr<const UScriptStruct>> Columns)
 					{
@@ -128,8 +128,8 @@ namespace UE::Editor::DataStorage
 	//
 	// SRowDetailsRow
 	//
-	void SRowDetailsRow::Construct(const FArguments& Args, const TSharedRef<STableViewBase>& OwnerTableView,IEditorDataStorageProvider* InDataStorage,
-			IEditorDataStorageUiProvider* InDataStorageUi)
+	void SRowDetailsRow::Construct(const FArguments& Args, const TSharedRef<STableViewBase>& OwnerTableView,ICoreProvider* InDataStorage,
+			IUiProvider* InDataStorageUi)
 	{
 		Item = Args._Item;
 
diff --git a/Engine/Plugins/Experimental/EditorDataStorageFeatures/Source/TedsTableViewer/Private/Widgets/STedsTableViewer.cpp b/Engine/Plugins/Experimental/EditorDataStorageFeatures/Source/TedsTableViewer/Private/Widgets/STedsTableViewer.cpp
--- a/Engine/Plugins/Experimental/EditorDataStorageFeatures/Source/TedsTableViewer/Private/Widgets/STedsTableViewer.cpp
+++ b/Engine/Plugins/Experimental/EditorDataStorageFeatures/Source/TedsTableViewer/Private/Widgets/STedsTableViewer.cpp
@@ -26,7 +26,7 @@ namespace UE::Editor::DataStorage
 		HeaderRowWidget = SNew( SHeaderRow )
 							.CanSelectGeneratedColumn(true);
 		
-		IEditorDataStorageUiProvider* StorageUi = Model->GetDataStorageUiProvider();
+		IUiProvider* StorageUi = Model->GetDataStorageUiProvider();
 
 		if (ensure(StorageUi))
 		{
@@ -71,7 +71,7 @@ namespace UE::Editor::DataStorage
 
 	void STedsTableViewer::AddWidgetColumns()
 	{
-		if(IEditorDataStorageProvider* DataStorage = Model->GetDataStorageInterface())
+		if(ICoreProvider* DataStorage = Model->GetDataStorageInterface())
 		{
 			const RowHandle WidgetRowHandle = TedsWidget->GetRowHandle();
 		
diff --git a/Engine/Plugins/Experimental/EditorDataStorageFeatures/Source/TedsTableViewer/Private/Widgets/STedsTileViewer.cpp b/Engine/Plugins/Experimental/EditorDataStorageFeatures/Source/TedsTableViewer/Private/Widgets/STedsTileViewer.cpp
--- a/Engine/Plugins/Experimental/EditorDataStorageFeatures/Source/TedsTableViewer/Private/Widgets/STedsTileViewer.cpp
+++ b/Engine/Plugins/Experimental/EditorDataStorageFeatures/Source/TedsTableViewer/Private/Widgets/STedsTileViewer.cpp
@@ -30,7 +30,7 @@ namespace UE::Editor::DataStorage
 		Model = MakeShared<FTedsTableViewerModel>(InArgs._QueryStack, DummyColumns, Purposes,
 			FTedsTableViewerModel::FIsItemVisible::CreateSP(this, &STedsTileViewer::IsItemVisible));
 
-		IEditorDataStorageUiProvider* StorageUi = Model->GetDataStorageUiProvider();
+		IUiProvider* StorageUi = Model->GetDataStorageUiProvider();
 		
 		if (ensure(StorageUi))
 		{
@@ -83,7 +83,7 @@ namespace UE::Editor::DataStorage
 			return false;
 		};
 	
-		IEditorDataStorageUiProvider* StorageUi = GetMutableDataStorageFeature<IEditorDataStorageUiProvider>(UiFeatureName);
+		IUiProvider* StorageUi = GetMutableDataStorageFeature<IUiProvider>(UiFeatureName);
 
 		if (Columns.IsEmpty())
 		{
@@ -91,7 +91,7 @@ namespace UE::Editor::DataStorage
 		}
 		else
 		{
-			StorageUi->CreateWidgetConstructors(WidgetPurpose, IEditorDataStorageUiProvider::EMatchApproach::ExactMatch, Columns,
+			StorageUi->CreateWidgetConstructors(WidgetPurpose, IUiProvider::EMatchApproach::ExactMatch, Columns,
 				FMetaDataView(), AssignWidgetToColumn);
 		}
 
@@ -173,7 +173,7 @@ namespace UE::Editor::DataStorage
 
 	void STedsTileViewer::AddWidgetColumns()
 	{
-		if(IEditorDataStorageProvider* DataStorage = Model->GetDataStorageInterface())
+		if(ICoreProvider* DataStorage = Model->GetDataStorageInterface())
 		{
 			const RowHandle WidgetRowHandle = TedsWidget->GetRowHandle();
 		
diff --git a/Engine/Plugins/Experimental/EditorDataStorageFeatures/Source/TedsTableViewer/Private/Widgets/UObjectLabelWidget.cpp b/Engine/Plugins/Experimental/EditorDataStorageFeatures/Source/TedsTableViewer/Private/Widgets/UObjectLabelWidget.cpp
--- a/Engine/Plugins/Experimental/EditorDataStorageFeatures/Source/TedsTableViewer/Private/Widgets/UObjectLabelWidget.cpp
+++ b/Engine/Plugins/Experimental/EditorDataStorageFeatures/Source/TedsTableViewer/Private/Widgets/UObjectLabelWidget.cpp
@@ -14,8 +14,8 @@
 
 #define LOCTEXT_NAMESPACE "FUObjectLabelWidgetConstructor"
 
-void UUObjectLabelWidgetFactory::RegisterWidgetConstructors(IEditorDataStorageProvider& DataStorage,
-	IEditorDataStorageUiProvider& DataStorageUi) const
+void UUObjectLabelWidgetFactory::RegisterWidgetConstructors(UE::Editor::DataStorage::ICoreProvider& DataStorage,
+	UE::Editor::DataStorage::IUiProvider& DataStorageUi) const
 {
 	using namespace UE::Editor::DataStorage::Queries;
 	DataStorageUi.RegisterWidgetFactory<FUObjectLabelWidgetConstructor>(
@@ -33,8 +33,8 @@ FUObjectLabelWidgetConstructor::FUObjectLabelWidgetConstructor(const UScriptStru
 {
 }
 
-TSharedPtr<SWidget> FUObjectLabelWidgetConstructor::CreateWidget(IEditorDataStorageProvider* DataStorage,
-	IEditorDataStorageUiProvider* DataStorageUi, UE::Editor::DataStorage::RowHandle TargetRow, 
+TSharedPtr<SWidget> FUObjectLabelWidgetConstructor::CreateWidget(UE::Editor::DataStorage::ICoreProvider* DataStorage,
+	UE::Editor::DataStorage::IUiProvider* DataStorageUi, UE::Editor::DataStorage::RowHandle TargetRow,
 	UE::Editor::DataStorage::RowHandle WidgetRow, const UE::Editor::DataStorage::FMetaDataView& Arguments)
 {
 	using namespace UE::Editor::DataStorage;
diff --git a/Engine/Plugins/Experimental/EditorDataStorageFeatures/Source/TedsTableViewer/Public/TedsTableViewerColumn.h b/Engine/Plugins/Experimental/EditorDataStorageFeatures/Source/TedsTableViewer/Public/TedsTableViewerColumn.h
--- a/Engine/Plugins/Experimental/EditorDataStorageFeatures/Source/TedsTableViewer/Public/TedsTableViewerColumn.h
+++ b/Engine/Plugins/Experimental/EditorDataStorageFeatures/Source/TedsTableViewer/Public/TedsTableViewerColumn.h
@@ -8,14 +8,15 @@
 #include "UObject/NameTypes.h"
 #include "Widgets/Views/SHeaderRow.h"
 
-class IEditorDataStorageCompatibilityProvider;
-class IEditorDataStorageUiProvider;
-class IEditorDataStorageProvider;
 struct FTypedElementWidgetConstructor;
 class SWidget;
 
 namespace UE::Editor::DataStorage
 {
+	class ICompatibilityProvider;
+	class IUiProvider;
+	class ICoreProvider;
+
 	/*
 	 * Class representing a column in the UI of the table viewer. Can be constructed using a NameID and a WidgetConstructor to create the actual
 	 * widgets for rows (optionally supplying a header widget constructor and widget metadata to use)
@@ -37,7 +38,7 @@ namespace UE::Editor::DataStorage
 		TEDSTABLEVIEWER_API ~FTedsTableViewerColumn();
 		
 		TEDSTABLEVIEWER_API	TSharedPtr<SWidget> ConstructRowWidget(RowHandle InRowHandle, 
-			TFunction<void(IEditorDataStorageProvider&, const RowHandle&)> WidgetRowSetupDelegate = nullptr) const;
+			TFunction<void(ICoreProvider&, const RowHandle&)> WidgetRowSetupDelegate = nullptr) const;
 		
 		TEDSTABLEVIEWER_API SHeaderRow::FColumn::FArguments ConstructHeaderRowColumn() const;
 		
@@ -49,7 +50,7 @@ namespace UE::Editor::DataStorage
 
 		TEDSTABLEVIEWER_API TConstArrayView<TWeakObjectPtr<const UScriptStruct>> GetMatchedColumns() const;
 
-		TEDSTABLEVIEWER_API IEditorDataStorageProvider* GetStorage() const;
+		TEDSTABLEVIEWER_API ICoreProvider* GetStorage() const;
 	protected:
 
 		void RegisterQueries();
@@ -75,9 +76,9 @@ namespace UE::Editor::DataStorage
 		FMetaDataView WidgetMetaData;
 
 		// TEDS Constructs
-		IEditorDataStorageProvider* Storage;
-		IEditorDataStorageUiProvider* StorageUi;
-		IEditorDataStorageCompatibilityProvider* StorageCompatibility;
+		ICoreProvider* Storage;
+		IUiProvider* StorageUi;
+		ICompatibilityProvider* StorageCompatibility;
 
 		// Queries used to virtualize widgets when a column is added to/remove from a row
 		TArray<QueryHandle> InternalObserverQueries;
diff --git a/Engine/Plugins/Experimental/EditorDataStorageFeatures/Source/TedsTableViewer/Public/TedsTableViewerModel.h b/Engine/Plugins/Experimental/EditorDataStorageFeatures/Source/TedsTableViewer/Public/TedsTableViewerModel.h
--- a/Engine/Plugins/Experimental/EditorDataStorageFeatures/Source/TedsTableViewer/Public/TedsTableViewerModel.h
+++ b/Engine/Plugins/Experimental/EditorDataStorageFeatures/Source/TedsTableViewer/Public/TedsTableViewerModel.h
@@ -10,13 +10,12 @@
 #include "TedsQueryStackInterfaces.h"
 #include "UObject/NameTypes.h"
 
-class IEditorDataStorageCompatibilityProvider;
-class IEditorDataStorageUiProvider;
-class IEditorDataStorageProvider;
-
 namespace UE::Editor::DataStorage
 {
 	class FTedsTableViewerColumn;
+	class ICompatibilityProvider;
+	class IUiProvider;
+	class ICoreProvider;
 	
 	// Typedef for an item in the table viewer
 	using TableViewerItemPtr = FTedsRowHandle;
@@ -59,9 +58,9 @@ namespace UE::Editor::DataStorage
 		// Add a custom column to display in the table viewer, that doesn't necessarily map to a Teds column
 		TEDSTABLEVIEWER_API void AddCustomColumn(const TSharedRef<FTedsTableViewerColumn>& InColumn);
 
-		TEDSTABLEVIEWER_API IEditorDataStorageProvider* GetDataStorageInterface() const;
+		TEDSTABLEVIEWER_API ICoreProvider* GetDataStorageInterface() const;
 
-		TEDSTABLEVIEWER_API IEditorDataStorageUiProvider* GetDataStorageUiProvider() const;
+		TEDSTABLEVIEWER_API IUiProvider* GetDataStorageUiProvider() const;
 
 	protected:
 
@@ -109,8 +108,8 @@ namespace UE::Editor::DataStorage
 		FOnModelChanged OnModelChanged;
 		
 		// Teds Constructs
-		IEditorDataStorageProvider* Storage = nullptr;
-		IEditorDataStorageUiProvider* StorageUi = nullptr;
-		IEditorDataStorageCompatibilityProvider* StorageCompatibility = nullptr;
+		ICoreProvider* Storage = nullptr;
+		IUiProvider* StorageUi = nullptr;
+		ICompatibilityProvider* StorageCompatibility = nullptr;
 	};
 } // namespace UE::Editor::DataStorage
diff --git a/Engine/Plugins/Experimental/EditorDataStorageFeatures/Source/TedsTableViewer/Public/TedsTableViewerUtils.h b/Engine/Plugins/Experimental/EditorDataStorageFeatures/Source/TedsTableViewer/Public/TedsTableViewerUtils.h
--- a/Engine/Plugins/Experimental/EditorDataStorageFeatures/Source/TedsTableViewer/Public/TedsTableViewerUtils.h
+++ b/Engine/Plugins/Experimental/EditorDataStorageFeatures/Source/TedsTableViewer/Public/TedsTableViewerUtils.h
@@ -12,10 +12,10 @@
 namespace UE::Editor::DataStorage
 {
 	class FMetaDataView;
+	class ICoreProvider;
+	class IUiProvider;
 }
 
-class IEditorDataStorageProvider;
-class IEditorDataStorageUiProvider;
 struct FTypedElementWidgetConstructor;
 struct FSlateBrush;
 
@@ -28,7 +28,7 @@ namespace UE::Editor::DataStorage::TableViewerUtils
 	TEDSTABLEVIEWER_API FName FindLongestMatchingName(TConstArrayView<TWeakObjectPtr<const UScriptStruct>> ColumnTypes, int32 DefaultNameIndex);
 
 	// Create a header widget constructor for the given columns
-	TEDSTABLEVIEWER_API TSharedPtr<FTypedElementWidgetConstructor> CreateHeaderWidgetConstructor(IEditorDataStorageUiProvider& StorageUi, 
+	TEDSTABLEVIEWER_API TSharedPtr<FTypedElementWidgetConstructor> CreateHeaderWidgetConstructor(IUiProvider& StorageUi, 
 	const UE::Editor::DataStorage::FMetaDataView& InMetaData, TConstArrayView<TWeakObjectPtr<const UScriptStruct>> ColumnTypes,
 	const TConstArrayView<FName> CellWidgetPurposes);
 
@@ -36,7 +36,7 @@ namespace UE::Editor::DataStorage::TableViewerUtils
 	TEDSTABLEVIEWER_API TArray<TWeakObjectPtr<const UScriptStruct>> CreateVerifiedColumnTypeArray(TConstArrayView<TWeakObjectPtr<const UScriptStruct>> ColumnTypes);
 	
 	// Scans the row for columns and tries to find the best matching icon.
-	TEDSTABLEVIEWER_API const FSlateBrush* GetIconForRow(IEditorDataStorageProvider* DataStorage, RowHandle Row);
+	TEDSTABLEVIEWER_API const FSlateBrush* GetIconForRow(ICoreProvider* DataStorage, RowHandle Row);
 } // namespace UE::Editor::DataStorage::TableViewerUtils
 
 UCLASS()
@@ -47,6 +47,6 @@ class UTypedElementTableViewerFactory : public UEditorDataStorageFactory
 public:
 	~UTypedElementTableViewerFactory() override = default;
 
-	void RegisterTables(IEditorDataStorageProvider& DataStorage) override;
+	void RegisterTables(UE::Editor::DataStorage::ICoreProvider& DataStorage) override;
 };
 
diff --git a/Engine/Plugins/Experimental/EditorDataStorageFeatures/Source/TedsTableViewer/Public/Widgets/RowHandleWidget.h b/Engine/Plugins/Experimental/EditorDataStorageFeatures/Source/TedsTableViewer/Public/Widgets/RowHandleWidget.h
--- a/Engine/Plugins/Experimental/EditorDataStorageFeatures/Source/TedsTableViewer/Public/Widgets/RowHandleWidget.h
+++ b/Engine/Plugins/Experimental/EditorDataStorageFeatures/Source/TedsTableViewer/Public/Widgets/RowHandleWidget.h
@@ -8,7 +8,11 @@
 
 #include "RowHandleWidget.generated.h"
 
-class IEditorDataStorageProvider;
+namespace UE::Editor::DataStorage
+{
+	class ICoreProvider;
+} // namespace UE::Editor::DataStorage
+
 class UScriptStruct;
 
 UCLASS()
@@ -19,10 +23,10 @@ class TEDSTABLEVIEWER_API URowHandleWidgetFactory : public UEditorDataStorageFac
 public:
 	virtual ~URowHandleWidgetFactory() override = default;
 
-	virtual void RegisterWidgetConstructors(IEditorDataStorageProvider& DataStorage,
-		IEditorDataStorageUiProvider& DataStorageUi) const override;
+	virtual void RegisterWidgetConstructors(UE::Editor::DataStorage::ICoreProvider& DataStorage,
+		UE::Editor::DataStorage::IUiProvider& DataStorageUi) const override;
 
-	void RegisterWidgetPurposes(IEditorDataStorageUiProvider& DataStorageUi) const override;
+	void RegisterWidgetPurposes(UE::Editor::DataStorage::IUiProvider& DataStorageUi) const override;
 };
 
 // A custom widget to display the row handle of a row as text
@@ -38,8 +42,8 @@ public:
 protected:
 	virtual TSharedPtr<SWidget> CreateWidget(const UE::Editor::DataStorage::FMetaDataView& Arguments) override;
 	virtual bool FinalizeWidget(
-		IEditorDataStorageProvider* DataStorage,
-		IEditorDataStorageUiProvider* DataStorageUi,
+		UE::Editor::DataStorage::ICoreProvider* DataStorage,
+		UE::Editor::DataStorage::IUiProvider* DataStorageUi,
 		UE::Editor::DataStorage::RowHandle Row,
 		const TSharedPtr<SWidget>& Widget) override;
 };
diff --git a/Engine/Plugins/Experimental/EditorDataStorageFeatures/Source/TedsTableViewer/Public/Widgets/SRowDetails.h b/Engine/Plugins/Experimental/EditorDataStorageFeatures/Source/TedsTableViewer/Public/Widgets/SRowDetails.h
--- a/Engine/Plugins/Experimental/EditorDataStorageFeatures/Source/TedsTableViewer/Public/Widgets/SRowDetails.h
+++ b/Engine/Plugins/Experimental/EditorDataStorageFeatures/Source/TedsTableViewer/Public/Widgets/SRowDetails.h
@@ -6,12 +6,13 @@
 #include "Widgets/SCompoundWidget.h"
 #include "Widgets/Views/SListView.h"
 
-class IEditorDataStorageProvider;
-class IEditorDataStorageUiProvider;
 struct FTypedElementWidgetConstructor;
 
 namespace UE::Editor::DataStorage
 {
+	class ICoreProvider;
+	class IUiProvider;
+
 	// A row in the SRowDetails widget that represents a column on the TEDS row we are viewing
 	struct FRowDetailsItem
 	{
@@ -66,8 +67,8 @@ namespace UE::Editor::DataStorage
 
 		TArray<RowDetailsItemPtr> Items;
 
-		IEditorDataStorageProvider* DataStorage = nullptr; 
-		IEditorDataStorageUiProvider* DataStorageUi = nullptr;
+		ICoreProvider* DataStorage = nullptr; 
+		IUiProvider* DataStorageUi = nullptr;
 
 		bool bShowAllDetails = true;
 
@@ -85,14 +86,14 @@ namespace UE::Editor::DataStorage
 
 		SLATE_END_ARGS()
 		
-		void Construct(const FArguments& Args, const TSharedRef<STableViewBase>& OwnerTableView, IEditorDataStorageProvider* InDataStorage,
-			IEditorDataStorageUiProvider* InDataStorageUi);
+		void Construct(const FArguments& Args, const TSharedRef<STableViewBase>& OwnerTableView, ICoreProvider* InDataStorage,
+			IUiProvider* InDataStorageUi);
 		
 		TSharedRef<SWidget> GenerateWidgetForColumn(const FName& ColumnName) override;
 
 	private:
 		RowDetailsItemPtr Item;
-		IEditorDataStorageProvider* DataStorage = nullptr;
-		IEditorDataStorageUiProvider* DataStorageUi = nullptr;
+		ICoreProvider* DataStorage = nullptr;
+		IUiProvider* DataStorageUi = nullptr;
 	};
 } // namespace UE::Editor::DataStorage
diff --git a/Engine/Plugins/Experimental/EditorDataStorageFeatures/Source/TedsTableViewer/Public/Widgets/UObjectLabelWidget.h b/Engine/Plugins/Experimental/EditorDataStorageFeatures/Source/TedsTableViewer/Public/Widgets/UObjectLabelWidget.h
--- a/Engine/Plugins/Experimental/EditorDataStorageFeatures/Source/TedsTableViewer/Public/Widgets/UObjectLabelWidget.h
+++ b/Engine/Plugins/Experimental/EditorDataStorageFeatures/Source/TedsTableViewer/Public/Widgets/UObjectLabelWidget.h
@@ -9,7 +9,11 @@
 
 #include "UObjectLabelWidget.generated.h"
 
-class IEditorDataStorageProvider;
+namespace UE::Editor::DataStorage
+{
+	class ICoreProvider;
+} // UE::Editor::DataStorage
+
 class UScriptStruct;
 
 UCLASS()
@@ -20,8 +24,8 @@ class UUObjectLabelWidgetFactory : public UEditorDataStorageFactory
 public:
 	~UUObjectLabelWidgetFactory() override = default;
 
-	TEDSTABLEVIEWER_API void RegisterWidgetConstructors(IEditorDataStorageProvider& DataStorage,
-		IEditorDataStorageUiProvider& DataStorageUi) const override;
+	TEDSTABLEVIEWER_API void RegisterWidgetConstructors(UE::Editor::DataStorage::ICoreProvider& DataStorage,
+		UE::Editor::DataStorage::IUiProvider& DataStorageUi) const override;
 };
 
 // Widget to show labels for UObjects in TEDS
@@ -36,8 +40,8 @@ public:
 	~FUObjectLabelWidgetConstructor() override = default;
 
 	TEDSTABLEVIEWER_API virtual TSharedPtr<SWidget> CreateWidget(
-		IEditorDataStorageProvider* DataStorage,
-		IEditorDataStorageUiProvider* DataStorageUi,
+		UE::Editor::DataStorage::ICoreProvider* DataStorage,
+		UE::Editor::DataStorage::IUiProvider* DataStorageUi,
 		UE::Editor::DataStorage::RowHandle TargetRow,
 		UE::Editor::DataStorage::RowHandle WidgetRow, 
 		const UE::Editor::DataStorage::FMetaDataView& Arguments) override;
diff --git a/Engine/Plugins/Experimental/InstanceDataObjectFixupTool/Source/Private/InstanceDataObjectFixupPanel.cpp b/Engine/Plugins/Experimental/InstanceDataObjectFixupTool/Source/Private/InstanceDataObjectFixupPanel.cpp
--- a/Engine/Plugins/Experimental/InstanceDataObjectFixupTool/Source/Private/InstanceDataObjectFixupPanel.cpp
+++ b/Engine/Plugins/Experimental/InstanceDataObjectFixupTool/Source/Private/InstanceDataObjectFixupPanel.cpp
@@ -253,8 +253,8 @@ FInstanceDataObjectFixupPanel::~FInstanceDataObjectFixupPanel()
 {
 	using namespace UE::Editor::DataStorage;
 
-	IEditorDataStorageProvider* DataStorage = GetMutableDataStorageFeature<IEditorDataStorageProvider>(StorageFeatureName);
-	IEditorDataStorageCompatibilityProvider* DataStorageCompatibility = GetMutableDataStorageFeature<IEditorDataStorageCompatibilityProvider>(CompatibilityFeatureName);
+	ICoreProvider* DataStorage = GetMutableDataStorageFeature<ICoreProvider>(StorageFeatureName);
+	ICompatibilityProvider* DataStorageCompatibility = GetMutableDataStorageFeature<ICompatibilityProvider>(CompatibilityFeatureName);
 
 	if (DataStorageCompatibility != nullptr && DataStorage != nullptr)
 	{
diff --git a/Engine/Plugins/Experimental/InstanceDataObjectFixupTool/Source/Private/InstanceDataObjectFixupToolTedsQueries.cpp b/Engine/Plugins/Experimental/InstanceDataObjectFixupTool/Source/Private/InstanceDataObjectFixupToolTedsQueries.cpp
--- a/Engine/Plugins/Experimental/InstanceDataObjectFixupTool/Source/Private/InstanceDataObjectFixupToolTedsQueries.cpp
+++ b/Engine/Plugins/Experimental/InstanceDataObjectFixupTool/Source/Private/InstanceDataObjectFixupToolTedsQueries.cpp
@@ -16,7 +16,7 @@
 
 #define LOCTEXT_NAMESPACE "FixupToolTedsQueries"
 
-void UInstanceDataObjectFixupToolTedsQueryFactory::RegisterQueries(IEditorDataStorageProvider& DataStorage)
+void UInstanceDataObjectFixupToolTedsQueryFactory::RegisterQueries(UE::Editor::DataStorage::ICoreProvider& DataStorage)
 {
 	using namespace UE::Editor::DataStorage::Queries;
 	
@@ -63,7 +63,7 @@ void UInstanceDataObjectFixupToolTedsQueryFactory::ShowFixUpToolForLooseProperti
 void UInstanceDataObjectFixupToolTedsQueryFactory::ShowFixUpTool(UE::Editor::DataStorage::RowHandle Row, bool bRecurseIntoObject)
 {
 	using namespace UE::Editor::DataStorage;
-	IEditorDataStorageProvider* DataStorage = GetMutableDataStorageFeature<IEditorDataStorageProvider>(StorageFeatureName);
+	ICoreProvider* DataStorage = GetMutableDataStorageFeature<ICoreProvider>(StorageFeatureName);
 	if (FTypedElementUObjectColumn* ObjectColumn = DataStorage->GetColumn<FTypedElementUObjectColumn>(Row))
 	{
 		if (bRecurseIntoObject)
diff --git a/Engine/Plugins/Fab/Source/Fab/Private/Teds/FabMyFolderIntegration.cpp b/Engine/Plugins/Fab/Source/Fab/Private/Teds/FabMyFolderIntegration.cpp
--- a/Engine/Plugins/Fab/Source/Fab/Private/Teds/FabMyFolderIntegration.cpp
+++ b/Engine/Plugins/Fab/Source/Fab/Private/Teds/FabMyFolderIntegration.cpp
@@ -17,7 +17,7 @@
 #include "Serialization/JsonReader.h"
 #include "Serialization/JsonSerializer.h"
 
-void UFabFactory::RegisterTables(IEditorDataStorageProvider& DataStorage)
+void UFabFactory::RegisterTables(UE::Editor::DataStorage::ICoreProvider& DataStorage)
 {
 	Table = DataStorage.RegisterTable<FNameColumn, FFabObjectColumn, FEditorDataStorageUrlColumn>(FName("Fab"));
 }
@@ -31,7 +31,7 @@ void FFabTedsMyFolderIntegration::QueueSyncRequest(uint32 BatchSize)
 {
 	using namespace UE::Editor::DataStorage;
 
-	if (IEditorDataStorageProvider* Storage = GetMutableDataStorageFeature<IEditorDataStorageProvider>(StorageFeatureName))
+	if (ICoreProvider* Storage = GetMutableDataStorageFeature<ICoreProvider>(StorageFeatureName))
 	{
 		Storage->RemoveAllRowsWith<FFabObjectColumn>();
 
@@ -97,7 +97,7 @@ void FFabTedsMyFolderIntegration::ProcessSyncResults(uint32 BatchSize, HttpReque
 		TSharedRef<TJsonReader<TCHAR>> JsonReader = TJsonReaderFactory<TCHAR>::Create(Content);
 		if (FJsonSerializer::Deserialize(JsonReader, JsonResults))
 		{
-			if (IEditorDataStorageProvider* Storage = GetMutableDataStorageFeature<IEditorDataStorageProvider>(StorageFeatureName))
+			if (ICoreProvider* Storage = GetMutableDataStorageFeature<ICoreProvider>(StorageFeatureName))
 			{
 				if (const UFabFactory* FabFactory = Storage->FindFactory<UFabFactory>())
 				{
@@ -245,7 +245,7 @@ void FFabTedsMyFolderIntegration::SetUrlColumn(UE::Editor::DataStorage::FUrlColu
 }
 
 void FFabTedsMyFolderIntegration::AddDistributionMethod(
-	IEditorDataStorageProvider* Storage, UE::Editor::DataStorage::RowHandle Row, const FJsonObject& Object, FString& Temp)
+	UE::Editor::DataStorage::ICoreProvider* Storage, UE::Editor::DataStorage::RowHandle Row, const FJsonObject& Object, FString& Temp)
 {
 	if (Object.TryGetStringField(TEXT("distributionMethod"), Temp))
 	{
@@ -258,7 +258,7 @@ void FFabTedsMyFolderIntegration::AddDistributionMethod(
 }
 
 void FFabTedsMyFolderIntegration::AddImages(
-	IEditorDataStorageProvider* Storage, UE::Editor::DataStorage::RowHandle Row, const FJsonObject& Object, FString& Temp)
+	UE::Editor::DataStorage::ICoreProvider* Storage, UE::Editor::DataStorage::RowHandle Row, const FJsonObject& Object, FString& Temp)
 {
 	using namespace UE::Editor::DataStorage;
 
diff --git a/Engine/Plugins/Fab/Source/Fab/Private/Teds/FabMyFolderIntegration.h b/Engine/Plugins/Fab/Source/Fab/Private/Teds/FabMyFolderIntegration.h
--- a/Engine/Plugins/Fab/Source/Fab/Private/Teds/FabMyFolderIntegration.h
+++ b/Engine/Plugins/Fab/Source/Fab/Private/Teds/FabMyFolderIntegration.h
@@ -9,10 +9,14 @@
 #include "FabMyFolderIntegration.generated.h"
 
 struct FEditorDataStorageUrlColumn;
-class IEditorDataStorageProvider;
 class FJsonObject;
 struct FNameColumn;
 
+namespace UE::Editor::DataStorage
+{
+	class ICoreProvider;
+} // namespace UE::Editor::DataStorage
+
 UCLASS()
 class UFabFactory : public UEditorDataStorageFactory
 {
@@ -21,7 +25,7 @@ class UFabFactory : public UEditorDataStorageFactory
 public:
 	virtual ~UFabFactory() override = default;
 
-	virtual void RegisterTables(IEditorDataStorageProvider& DataStorage) override;
+	virtual void RegisterTables(UE::Editor::DataStorage::ICoreProvider& DataStorage) override;
 
 	UE::Editor::DataStorage::TableHandle Table;
 };
@@ -78,7 +82,7 @@ private:
 	static void SetFabObjectColumn(FFabObjectColumn& Target, const FJsonObject& Object, FString& Temp);
 	static void SetUrlColumn(FEditorDataStorageUrlColumn& Target, const FJsonObject& Object, FString& Temp);
 	static void AddDistributionMethod(
-		IEditorDataStorageProvider* Storage, UE::Editor::DataStorage::RowHandle Row, const FJsonObject& Object, FString& Temp);
+		UE::Editor::DataStorage::ICoreProvider* Storage, UE::Editor::DataStorage::RowHandle Row, const FJsonObject& Object, FString& Temp);
 	static void AddImages(
-		IEditorDataStorageProvider* Storage, UE::Editor::DataStorage::RowHandle Row, const FJsonObject& Object, FString& Temp);
+		UE::Editor::DataStorage::ICoreProvider* Storage, UE::Editor::DataStorage::RowHandle Row, const FJsonObject& Object, FString& Temp);
 };
\ No newline at end of file
diff --git a/Engine/Source/Editor/ContentBrowser/Private/ContentSources/Widgets/SContentSource.cpp b/Engine/Source/Editor/ContentBrowser/Private/ContentSources/Widgets/SContentSource.cpp
--- a/Engine/Source/Editor/ContentBrowser/Private/ContentSources/Widgets/SContentSource.cpp
+++ b/Engine/Source/Editor/ContentBrowser/Private/ContentSources/Widgets/SContentSource.cpp
@@ -18,8 +18,8 @@ namespace UE::Editor::ContentBrowser
     {
 		using namespace UE::Editor::DataStorage;
 		
-		DataStorage = GetMutableDataStorageFeature<IEditorDataStorageProvider>(StorageFeatureName);
-		DataStorageUi = GetMutableDataStorageFeature<IEditorDataStorageUiProvider>(UiFeatureName);
+		DataStorage = GetMutableDataStorageFeature<ICoreProvider>(StorageFeatureName);
+		DataStorageUi = GetMutableDataStorageFeature<IUiProvider>(UiFeatureName);
 
 		ChildSlot
 		[
diff --git a/Engine/Source/Editor/ContentBrowser/Private/ContentSources/Widgets/SContentSource.h b/Engine/Source/Editor/ContentBrowser/Private/ContentSources/Widgets/SContentSource.h
--- a/Engine/Source/Editor/ContentBrowser/Private/ContentSources/Widgets/SContentSource.h
+++ b/Engine/Source/Editor/ContentBrowser/Private/ContentSources/Widgets/SContentSource.h
@@ -5,8 +5,11 @@
 #include "ContentSources/IContentSource.h"
 #include "Widgets/SCompoundWidget.h"
 
-class IEditorDataStorageProvider;
-class IEditorDataStorageUiProvider;
+namespace UE::Editor::DataStorage
+{
+	class ICoreProvider;
+	class IUiProvider;
+} // namespace UE::Editor::DataStorage
 
 namespace UE::Editor::ContentBrowser
 {
@@ -27,8 +30,8 @@ namespace UE::Editor::ContentBrowser
 	private:
 		// The currently active content source
 		TSharedPtr<IContentSource> ContentSource;
-		IEditorDataStorageProvider* DataStorage = nullptr; 
-		IEditorDataStorageUiProvider* DataStorageUi = nullptr;
+		DataStorage::ICoreProvider* DataStorage = nullptr; 
+		DataStorage::IUiProvider* DataStorageUi = nullptr;
 	};
 
 }
\ No newline at end of file
diff --git a/Engine/Source/Editor/UnrealEd/Private/EditorActor.cpp b/Engine/Source/Editor/UnrealEd/Private/EditorActor.cpp
--- a/Engine/Source/Editor/UnrealEd/Private/EditorActor.cpp
+++ b/Engine/Source/Editor/UnrealEd/Private/EditorActor.cpp
@@ -1496,10 +1496,10 @@ static void SetActorVisibility(AActor* Actor, bool bVisible)
 	// we're only editing temporary, transient values
 	SaveToTransactionBuffer(Actor, false);
 
-	if (const IEditorDataStorageCompatibilityProvider* DataStoreCompat = GetDataStorageFeature<IEditorDataStorageCompatibilityProvider>(CompatibilityFeatureName))
+	if (const ICompatibilityProvider* DataStoreCompat = GetDataStorageFeature<ICompatibilityProvider>(CompatibilityFeatureName))
 	{
 		RowHandle Row = DataStoreCompat->FindRowWithCompatibleObject(Actor);
-		if (IEditorDataStorageProvider* DataStore = GetMutableDataStorageFeature<IEditorDataStorageProvider>(StorageFeatureName))
+		if (ICoreProvider* DataStore = GetMutableDataStorageFeature<ICoreProvider>(StorageFeatureName))
 		{
 			FVisibleInEditorColumn* VisibilityColumn = DataStore->GetColumn<FVisibleInEditorColumn>(Row);
 			if (VisibilityColumn)
diff --git a/Engine/Source/Editor/UnrealEd/Private/Elements/Actor/ActorElementEditorSelectionInterface.cpp b/Engine/Source/Editor/UnrealEd/Private/Elements/Actor/ActorElementEditorSelectionInterface.cpp
--- a/Engine/Source/Editor/UnrealEd/Private/Elements/Actor/ActorElementEditorSelectionInterface.cpp
+++ b/Engine/Source/Editor/UnrealEd/Private/Elements/Actor/ActorElementEditorSelectionInterface.cpp
@@ -54,12 +54,12 @@ bool UActorElementEditorSelectionInterface::SelectElement(const FTypedElementHan
 
 	const AActor* Actor = ActorElementDataUtil::GetActorFromHandle(InElementHandle);
 	// Add a selection column in TEDS
-	if (const IEditorDataStorageCompatibilityProvider* Compatibility = GetDataStorageFeature<IEditorDataStorageCompatibilityProvider>(CompatibilityFeatureName))
+	if (const ICompatibilityProvider* Compatibility = GetDataStorageFeature<ICompatibilityProvider>(CompatibilityFeatureName))
 	{
 		RowHandle Row = Compatibility->FindRowWithCompatibleObject(Actor);
 		if (Row != InvalidRowHandle)
 		{
-			if (IEditorDataStorageProvider* DataStorage = GetMutableDataStorageFeature<IEditorDataStorageProvider>(StorageFeatureName))
+			if (ICoreProvider* DataStorage = GetMutableDataStorageFeature<ICoreProvider>(StorageFeatureName))
 			{
 				DataStorage->AddColumn<FTypedElementSelectionColumn>(Row);
 
@@ -81,12 +81,12 @@ bool UActorElementEditorSelectionInterface::DeselectElement(const FTypedElementH
 	using namespace UE::Editor::DataStorage;
 
 	const AActor* Actor = ActorElementDataUtil::GetActorFromHandle(InElementHandle);
-	if (const IEditorDataStorageCompatibilityProvider* Compatibility = GetDataStorageFeature<IEditorDataStorageCompatibilityProvider>(CompatibilityFeatureName))
+	if (const ICompatibilityProvider* Compatibility = GetDataStorageFeature<ICompatibilityProvider>(CompatibilityFeatureName))
 	{
 		RowHandle Row = Compatibility->FindRowWithCompatibleObject(Actor);
 		if (Row != InvalidRowHandle)
 		{
-			if (IEditorDataStorageProvider* DataStorage = GetMutableDataStorageFeature<IEditorDataStorageProvider>(StorageFeatureName))
+			if (ICoreProvider* DataStorage = GetMutableDataStorageFeature<ICoreProvider>(StorageFeatureName))
             {
             	DataStorage->RemoveColumn(Row, FTypedElementSelectionColumn::StaticStruct());
             }
diff --git a/Engine/Source/Editor/UnrealEd/Private/LevelEditorViewport.cpp b/Engine/Source/Editor/UnrealEd/Private/LevelEditorViewport.cpp
--- a/Engine/Source/Editor/UnrealEd/Private/LevelEditorViewport.cpp
+++ b/Engine/Source/Editor/UnrealEd/Private/LevelEditorViewport.cpp
@@ -1674,7 +1674,7 @@ void FLevelEditorViewportClient::DestroyDropPreviewElements()
 	// deleting preview actors in UUnrealEdEngine::DeleteActors skips explicit handle deregistration,
 	// and although it does still happen in the immediately triggered garbage cleanup, that feels
 	// potentially brittle.
-	if (IEditorDataStorageCompatibilityProvider* TedsCompat = GetMutableDataStorageFeature<IEditorDataStorageCompatibilityProvider>(CompatibilityFeatureName))
+	if (ICompatibilityProvider* TedsCompat = GetMutableDataStorageFeature<ICompatibilityProvider>(CompatibilityFeatureName))
 	{
 		DropPreviewElements->ForEachElement<ITypedElementObjectInterface>(
 			[this, TedsCompat](const TTypedElement<ITypedElementObjectInterface>& InElement)
diff --git a/Engine/Source/Runtime/Engine/Private/ActorEditor.cpp b/Engine/Source/Runtime/Engine/Private/ActorEditor.cpp
--- a/Engine/Source/Runtime/Engine/Private/ActorEditor.cpp
+++ b/Engine/Source/Runtime/Engine/Private/ActorEditor.cpp
@@ -990,8 +990,8 @@ void AActor::SetIsTemporarilyHiddenInEditor( bool bIsHidden )
 
 		// When updating the visibility we need to add the sync tag to the row in TEDS to ensure the visibility is synced in TEDS properly
 		using namespace UE::Editor::DataStorage;
-		IEditorDataStorageProvider* Storage = GetMutableDataStorageFeature<IEditorDataStorageProvider>(StorageFeatureName);
-		IEditorDataStorageCompatibilityProvider* StorageCompatibility = GetMutableDataStorageFeature<IEditorDataStorageCompatibilityProvider>(CompatibilityFeatureName);
+		ICoreProvider* Storage = GetMutableDataStorageFeature<ICoreProvider>(StorageFeatureName);
+		ICompatibilityProvider* StorageCompatibility = GetMutableDataStorageFeature<ICompatibilityProvider>(CompatibilityFeatureName);
 
 		if (Storage && StorageCompatibility)
 		{
