Added a cvar to turn off the TedsAssetDataFactory. It's not yet optimized enough.

diff --git a/Engine/Plugins/Experimental/EditorDataStorageFeatures/Source/TedsAssetData/Private/Processors/AssetProcessors.cpp b/Engine/Plugins/Experimental/EditorDataStorageFeatures/Source/TedsAssetData/Private/Processors/AssetProcessors.cpp
--- a/Engine/Plugins/Experimental/EditorDataStorageFeatures/Source/TedsAssetData/Private/Processors/AssetProcessors.cpp
+++ b/Engine/Plugins/Experimental/EditorDataStorageFeatures/Source/TedsAssetData/Private/Processors/AssetProcessors.cpp
@@ -22,6 +22,10 @@
 
 namespace UE::TedsAssetDataFactory::Private
 {
+	TAutoConsoleVariable<bool> CVarTedsAssetDataFactory(TEXT("TEDS.TedsAssetDataFactory"),
+		false,
+		TEXT("When true this will enable some experimental features that are not optimized to work at scale yet. Note: The value need to be set a boot time to see the effect of this cvar."));
+
 	// @return true if any of the texture dimensions are not a power of 2
 	bool IsTextureNonSquare(FStringView Dimensions)
 	{
@@ -89,6 +93,11 @@ void UTedsAssetDataFactory::RegisterQueries(UE::Editor::DataStorage::ICoreProvid
 	using namespace UE::Editor::DataStorage::Columns;
 	using namespace UE::Editor::DataStorage::Queries;
 
+	if (!UE::TedsAssetDataFactory::Private::CVarTedsAssetDataFactory.GetValueOnGameThread())
+	{
+		return;
+	}
+
 	DataStorage.RegisterQuery(
 		Select(
 			TEXT("TedsAssetDataFactory: Sync folder color from world"),
@@ -191,6 +200,11 @@ void UTedsAssetDataFactory::RegisterQueries(UE::Editor::DataStorage::ICoreProvid
 
 void UTedsAssetDataFactory::PreRegister(UE::Editor::DataStorage::ICoreProvider& DataStorage)
 {
+	if (!UE::TedsAssetDataFactory::Private::CVarTedsAssetDataFactory.GetValueOnGameThread())
+	{
+		return;
+	}
+
 	if (FContentBrowserModule* ContentBrowserModule = FModuleManager::Get().GetModulePtr<FContentBrowserModule>("ContentBrowser"))
 	{
 		ContentBrowserModule->GetOnSetFolderColor().AddUObject(this, &UTedsAssetDataFactory::OnSetFolderColor, &DataStorage);
