Teds Outliner: Update FTedsOutlinerItem to hold weak ptrs to FTedsOutlinerImpl, unregister all queries from FTedsOutlinerImpl on FTedsOutlinerMode destruction

This fixes a bug where FTedsOutlinerImpl would persist for a bit longer than the owning SSceneOutliner but try to incorrectly access it due to FTedsOutlinerItem being destructed by STreeView a bit later.

diff --git a/Engine/Plugins/Experimental/EditorDataStorageFeatures/Source/TedsOutliner/Private/TedsOutlinerItem.cpp b/Engine/Plugins/Experimental/EditorDataStorageFeatures/Source/TedsOutliner/Private/TedsOutlinerItem.cpp
--- a/Engine/Plugins/Experimental/EditorDataStorageFeatures/Source/TedsOutliner/Private/TedsOutlinerItem.cpp
+++ b/Engine/Plugins/Experimental/EditorDataStorageFeatures/Source/TedsOutliner/Private/TedsOutlinerItem.cpp
@@ -24,7 +24,7 @@ namespace UE::Editor::Outliner
 const FSceneOutlinerTreeItemType FTedsOutlinerTreeItem::Type(&ISceneOutlinerTreeItem::Type);
 
 FTedsOutlinerTreeItem::FTedsOutlinerTreeItem(const DataStorage::RowHandle& InRowHandle,
-	const TSharedRef<const FTedsOutlinerImpl>& InTedsOutlinerImpl)
+	const TWeakPtr<const FTedsOutlinerImpl>& InTedsOutlinerImpl)
 	: ISceneOutlinerTreeItem(Type)
 	, RowHandle(InRowHandle)
 	, TedsOutlinerImpl(InTedsOutlinerImpl)
@@ -44,14 +44,18 @@ FSceneOutlinerTreeItemID FTedsOutlinerTreeItem::GetID() const
 
 FString FTedsOutlinerTreeItem::GetDisplayString() const
 {
-	if (const FTypedElementLabelColumn* LabelColumn = TedsOutlinerImpl->GetStorage()->GetColumn<FTypedElementLabelColumn>(RowHandle))
+	if (const TSharedPtr<const FTedsOutlinerImpl> TedsOutlinerImplPin = TedsOutlinerImpl.Pin())
 	{
-		return LabelColumn->Label;
-	}
-	if (const FNameColumn* NameColumn = TedsOutlinerImpl->GetStorage()->GetColumn<FNameColumn>(RowHandle))
-	{
-		return NameColumn->Name.ToString();
+		if (const FTypedElementLabelColumn* LabelColumn = TedsOutlinerImplPin->GetStorage()->GetColumn<FTypedElementLabelColumn>(RowHandle))
+		{
+			return LabelColumn->Label;
+		}
+		if (const FNameColumn* NameColumn = TedsOutlinerImplPin->GetStorage()->GetColumn<FNameColumn>(RowHandle))
+		{
+			return NameColumn->Name.ToString();
+		}
 	}
+	
 
 	return TEXT("TEDS Item");
 }
@@ -64,7 +68,7 @@ bool FTedsOutlinerTreeItem::CanInteract() const
 TSharedRef<SWidget> FTedsOutlinerTreeItem::GenerateLabelWidget(ISceneOutliner& Outliner,
 	const STableRow<FSceneOutlinerTreeItemPtr>& InRow)
 {
-	return TedsOutlinerImpl->CreateLabelWidgetForItem(RowHandle, InRow, CanInteract());
+	return TedsOutlinerImpl.IsValid() ? TedsOutlinerImpl.Pin()->CreateLabelWidgetForItem(RowHandle, InRow, CanInteract()) : SNullWidget::NullWidget;
 }
 
 void FTedsOutlinerTreeItem::GenerateContextMenu(UToolMenu* Menu, SSceneOutliner& Outliner)
diff --git a/Engine/Plugins/Experimental/EditorDataStorageFeatures/Source/TedsOutliner/Private/TedsOutlinerMode.cpp b/Engine/Plugins/Experimental/EditorDataStorageFeatures/Source/TedsOutliner/Private/TedsOutlinerMode.cpp
--- a/Engine/Plugins/Experimental/EditorDataStorageFeatures/Source/TedsOutliner/Private/TedsOutlinerMode.cpp
+++ b/Engine/Plugins/Experimental/EditorDataStorageFeatures/Source/TedsOutliner/Private/TedsOutlinerMode.cpp
@@ -45,6 +45,7 @@ FTedsOutlinerMode::~FTedsOutlinerMode()
 	if (TedsOutlinerImpl)
 	{
 		TedsOutlinerImpl->OnSelectionChanged().RemoveAll(this);
+		TedsOutlinerImpl->UnregisterQueries();
 	}
 }
 
diff --git a/Engine/Plugins/Experimental/EditorDataStorageFeatures/Source/TedsOutliner/Public/TedsOutlinerImpl.h b/Engine/Plugins/Experimental/EditorDataStorageFeatures/Source/TedsOutliner/Public/TedsOutlinerImpl.h
--- a/Engine/Plugins/Experimental/EditorDataStorageFeatures/Source/TedsOutliner/Public/TedsOutlinerImpl.h
+++ b/Engine/Plugins/Experimental/EditorDataStorageFeatures/Source/TedsOutliner/Public/TedsOutlinerImpl.h
@@ -182,13 +182,13 @@ public:
 	// Recompile all queries used by this table viewer
 	void RecompileQueries();
 
+	// Unregister all queries used by this table viewer
+	void UnregisterQueries() const;
 
 protected:
 	
 	void OnItemAdded(DataStorage::RowHandle ItemRowHandle);
 	void OnItemRemoved(DataStorage::RowHandle ItemRowHandle);
-
-	void UnregisterQueries() const;
 	void ClearSelection() const;
 	void Tick();
 
diff --git a/Engine/Plugins/Experimental/EditorDataStorageFeatures/Source/TedsOutliner/Public/TedsOutlinerItem.h b/Engine/Plugins/Experimental/EditorDataStorageFeatures/Source/TedsOutliner/Public/TedsOutlinerItem.h
--- a/Engine/Plugins/Experimental/EditorDataStorageFeatures/Source/TedsOutliner/Public/TedsOutlinerItem.h
+++ b/Engine/Plugins/Experimental/EditorDataStorageFeatures/Source/TedsOutliner/Public/TedsOutlinerItem.h
@@ -32,7 +32,7 @@ public:
 		return Pred.Execute(RowHandle);
 	}
 
-	TEDSOUTLINER_API FTedsOutlinerTreeItem(const DataStorage::RowHandle& InRowHandle, const TSharedRef<const FTedsOutlinerImpl>& InTedsOutlinerImpl);
+	TEDSOUTLINER_API FTedsOutlinerTreeItem(const DataStorage::RowHandle& InRowHandle, const TWeakPtr<const FTedsOutlinerImpl>& InTedsOutlinerImpl);
 
 	/* Begin ISceneOutlinerTreeItem Implementation */
 	TEDSOUTLINER_API virtual bool IsValid() const override;
@@ -50,6 +50,6 @@ public:
 
 private:
 	const DataStorage::RowHandle RowHandle;
-	const TSharedRef<const FTedsOutlinerImpl> TedsOutlinerImpl;
+	const TWeakPtr<const FTedsOutlinerImpl> TedsOutlinerImpl;
 };
 } // namespace UE::Editor::Outliner
