Teds Table Viewer: Use the cell widget constructor's display name as the column name if present, fallback to UScriptStruct name otherwise.

diff --git a/Engine/Plugins/Experimental/EditorDataStorageFeatures/Source/TedsTableViewer/Private/TedsTableViewerColumn.cpp b/Engine/Plugins/Experimental/EditorDataStorageFeatures/Source/TedsTableViewer/Private/TedsTableViewerColumn.cpp
--- a/Engine/Plugins/Experimental/EditorDataStorageFeatures/Source/TedsTableViewer/Private/TedsTableViewerColumn.cpp
+++ b/Engine/Plugins/Experimental/EditorDataStorageFeatures/Source/TedsTableViewer/Private/TedsTableViewerColumn.cpp
@@ -116,7 +116,7 @@ namespace UE::Editor::DataStorage
 		if (!Widget.IsValid())
 		{
 			Widget = SNew(STextBlock)
-				.Text(FText::FromString(ColumnName.ToString()));
+				.Text(CellWidgetConstructor ? CellWidgetConstructor->CreateWidgetDisplayNameText(Storage) : FText::FromString(ColumnName.ToString()));
 		}
 		
 		SHeaderRow::FColumn::FArguments Column = SHeaderRow::Column(ColumnName)
