TedsSettings : Override CreateWidgetDisplayNameText to shorten widget display names for container reference and category reference columns.

diff --git a/Engine/Plugins/Experimental/EditorDataStorageFeatures/Source/TedsSettings/Private/TedsSettingsWidgets.cpp b/Engine/Plugins/Experimental/EditorDataStorageFeatures/Source/TedsSettings/Private/TedsSettingsWidgets.cpp
--- a/Engine/Plugins/Experimental/EditorDataStorageFeatures/Source/TedsSettings/Private/TedsSettingsWidgets.cpp
+++ b/Engine/Plugins/Experimental/EditorDataStorageFeatures/Source/TedsSettings/Private/TedsSettingsWidgets.cpp
@@ -59,6 +59,13 @@ TSharedPtr<SWidget> FSettingsContainerReferenceWidgetConstructor::CreateWidget(
 	return SNullWidget::NullWidget;
 }
 
+FText FSettingsContainerReferenceWidgetConstructor::CreateWidgetDisplayNameText(
+	UE::Editor::DataStorage::ICoreProvider* DataStorage,
+	UE::Editor::DataStorage::RowHandle Row) const
+{
+	return NSLOCTEXT("TedsSettingsWidgets", "SettingsContainerReferenceWidgetDisplayNameText", "Settings Container");
+}
+
 FSettingsCategoryReferenceWidgetConstructor::FSettingsCategoryReferenceWidgetConstructor()
 	: FSimpleWidgetConstructor(FSettingsCategoryReferenceWidgetConstructor::StaticStruct())
 {
@@ -108,6 +115,13 @@ TSharedPtr<SWidget> FSettingsCategoryReferenceWidgetConstructor::CreateWidget(
 	return SNullWidget::NullWidget;
 }
 
+FText FSettingsCategoryReferenceWidgetConstructor::CreateWidgetDisplayNameText(
+	UE::Editor::DataStorage::ICoreProvider* DataStorage,
+	UE::Editor::DataStorage::RowHandle Row) const
+{
+	return NSLOCTEXT("TedsSettingsWidgets", "SettingsCategoryReferenceWidgetDisplayNameText", "Settings Category");
+}
+
 FSettingsSectionWidgetConstructor::FSettingsSectionWidgetConstructor()
 	: FSimpleWidgetConstructor(FSettingsSectionWidgetConstructor::StaticStruct())
 {
diff --git a/Engine/Plugins/Experimental/EditorDataStorageFeatures/Source/TedsSettings/Private/TedsSettingsWidgets.h b/Engine/Plugins/Experimental/EditorDataStorageFeatures/Source/TedsSettings/Private/TedsSettingsWidgets.h
--- a/Engine/Plugins/Experimental/EditorDataStorageFeatures/Source/TedsSettings/Private/TedsSettingsWidgets.h
+++ b/Engine/Plugins/Experimental/EditorDataStorageFeatures/Source/TedsSettings/Private/TedsSettingsWidgets.h
@@ -24,6 +24,10 @@ public:
 		UE::Editor::DataStorage::RowHandle WidgetRow,
 		const UE::Editor::DataStorage::FMetaDataView& Arguments) override;
 
+	virtual FText CreateWidgetDisplayNameText(
+		UE::Editor::DataStorage::ICoreProvider* DataStorage,
+		UE::Editor::DataStorage::RowHandle Row = UE::Editor::DataStorage::InvalidRowHandle) const override;
+
 };
 
 USTRUCT()
@@ -43,6 +47,10 @@ public:
 		UE::Editor::DataStorage::RowHandle WidgetRow,
 		const UE::Editor::DataStorage::FMetaDataView& Arguments) override;
 
+	virtual FText CreateWidgetDisplayNameText(
+		UE::Editor::DataStorage::ICoreProvider* DataStorage,
+		UE::Editor::DataStorage::RowHandle Row = UE::Editor::DataStorage::InvalidRowHandle) const override;
+
 };
 
 USTRUCT()
