TedsOutlinerImpl: Use the correct query condition syntax to match the syntax used by the provided row query. This fixes a bug where you could not query for dynamic columns in the Teds Outliner

diff --git a/Engine/Plugins/Experimental/EditorDataStorageFeatures/Source/TedsOutliner/Private/TedsOutlinerImpl.cpp b/Engine/Plugins/Experimental/EditorDataStorageFeatures/Source/TedsOutliner/Private/TedsOutlinerImpl.cpp
--- a/Engine/Plugins/Experimental/EditorDataStorageFeatures/Source/TedsOutliner/Private/TedsOutlinerImpl.cpp
+++ b/Engine/Plugins/Experimental/EditorDataStorageFeatures/Source/TedsOutliner/Private/TedsOutlinerImpl.cpp
@@ -118,6 +118,8 @@ FTedsOutlinerImpl::FTedsOutlinerImpl(const FTedsOutlinerParams& InParams, IScene
 	Storage = GetMutableDataStorageFeature<DataStorage::ICoreProvider>(StorageFeatureName);
 	StorageUi = GetMutableDataStorageFeature<DataStorage::IUiProvider>(UiFeatureName);
 	StorageCompatibility = GetMutableDataStorageFeature<DataStorage::ICompatibilityProvider>(CompatibilityFeatureName);
+
+	bUsingQueryConditionsSyntax = InitialQueryDescription.Get().Conditions && !InitialQueryDescription.Get().Conditions->IsEmpty();
 }
 
 void FTedsOutlinerImpl::CreateFilterQueries()
@@ -141,11 +143,22 @@ void FTedsOutlinerImpl::CreateFilterQueries()
 				if (Struct->IsChildOf(TedsColumn) || Struct->IsChildOf(TedsTag))
 				{
 					// Create a query description to filter for this tag/column
-					DataStorage::FQueryDescription FilterQueryDesc =
-					Select()
-					.Where()
-						.All(Struct)
-					.Compile();
+					FQueryDescription FilterQueryDesc;
+
+					if (bUsingQueryConditionsSyntax)
+					{
+						FilterQueryDesc = Select()
+							.Where(TColumn(Struct))
+							.Compile();
+					}
+					else
+					{
+						FilterQueryDesc =
+							Select()
+							.Where()
+								.All(Struct)
+							.Compile();
+					}
 
 					// Create the filter
 					TSharedRef<FTedsOutlinerFilter> TedsFilter = MakeShared<FTedsOutlinerFilter>(Struct->GetFName(), Struct->GetDisplayNameText(),
@@ -310,6 +323,18 @@ void FTedsOutlinerImpl::AppendQuery(DataStorage::FQueryDescription& Query1, cons
 			Query1.ConditionTypes.Add(Query2.ConditionTypes[i]);
 		}
 	}
+
+	if (Query2.Conditions)
+	{
+		if (Query1.Conditions)
+		{
+			Query1.Conditions = Query1.Conditions.GetValue() && Query2.Conditions.GetValue();
+		}
+		else
+		{
+			Query1.Conditions = Query2.Conditions;
+		}
+	}
 }
 
 void FTedsOutlinerImpl::AddExternalQuery(FName QueryName, const DataStorage::FQueryDescription& InQueryDescription)
@@ -550,6 +575,9 @@ void FTedsOutlinerImpl::RecompileQueries()
 		return;
 	}
 
+	// Since InitialQueryDescription is an attribute we have to check this every time it changes 
+	bUsingQueryConditionsSyntax = InitialQueryDescription.Get().Conditions && !InitialQueryDescription.Get().Conditions->IsEmpty();
+
 	// Our final query to collect rows to populate the Outliner - currently the same as the initial query the user provided
 	FQueryDescription FinalQueryDescription(InitialQueryDescription.Get());
 
@@ -588,14 +616,25 @@ void FTedsOutlinerImpl::RecompileQueries()
 	if (HierarchyData.IsSet())
 	{
 		const UScriptStruct* ParentColumnType = HierarchyData.GetValue().HierarchyColumn;
-		
-		// Query to get all rows that match our conditions with a parent column (i.e all child rows)
-		FQueryDescription ChildHandleQueryDescription =
-							Select()
-							.Where()
-								.All(ParentColumnType)
-							.Compile();
 
+		// Query to get all rows that match our conditions with a parent column (i.e all child rows)
+		FQueryDescription ChildHandleQueryDescription;
+		
+		if (bUsingQueryConditionsSyntax)
+		{
+			ChildHandleQueryDescription = Select()
+				.Where(TColumn(ParentColumnType))
+				.Compile();
+		}
+		else
+		{
+			ChildHandleQueryDescription =
+				Select()
+				.Where()
+					.All(ParentColumnType)
+				.Compile();
+		}
+		
 		// Add the conditions from FinalQueryDescription to ensure we are tracking removal of the rows the user requested
 		AppendQuery(ChildHandleQueryDescription, FinalQueryDescription);
 
@@ -633,9 +672,17 @@ void FTedsOutlinerImpl::RecompileQueries()
 			})
 			.ReadOnly(ParentColumnType, EOptional::Yes)
 			.ReadOnly<FTedsOutlinerColumn>(EOptional::Yes)
-		.Where()
-			.All<FTypedElementSyncBackToWorldTag>()
 		.Compile();
+
+		if (bUsingQueryConditionsSyntax)
+		{
+			UpdateParentQueryDescription.Conditions.Emplace(TColumn<FTypedElementSyncBackToWorldTag>());
+		}
+		else
+		{
+			UpdateParentQueryDescription.ConditionTypes.Add(FQueryDescription::EOperatorType::SimpleAll);
+			UpdateParentQueryDescription.ConditionOperators.AddZeroed_GetRef().Type = FTypedElementSyncBackToWorldTag::StaticStruct();
+		}
 		
 		// Add the conditions from FinalQueryDescription to ensure we are the rows the user requested
 		AppendQuery(UpdateParentQueryDescription, FinalQueryDescription);
@@ -647,11 +694,22 @@ void FTedsOutlinerImpl::RecompileQueries()
 	if (SelectionSetName.IsSet())
 	{
 		// Query to grab all selected rows
-		FQueryDescription SelectedRowsQueryDescription =
-						Select()
-							.Where()
-								.All<FTypedElementSelectionColumn>()
-							.Compile();
+		FQueryDescription SelectedRowsQueryDescription;
+		
+		if (bUsingQueryConditionsSyntax)
+		{
+			SelectedRowsQueryDescription = Select()
+				.Where(TColumn<FTypedElementSelectionColumn>())
+				.Compile();
+		}
+		else
+		{
+			SelectedRowsQueryDescription =
+				Select()
+				.Where()
+					.All(FTypedElementSelectionColumn::StaticStruct())
+				.Compile();
+		}
 	
 		// Query to track when a row gets selected
 		FQueryDescription SelectionAddedQueryDescription =
diff --git a/Engine/Plugins/Experimental/EditorDataStorageFeatures/Source/TedsOutliner/Public/TedsOutlinerImpl.h b/Engine/Plugins/Experimental/EditorDataStorageFeatures/Source/TedsOutliner/Public/TedsOutlinerImpl.h
--- a/Engine/Plugins/Experimental/EditorDataStorageFeatures/Source/TedsOutliner/Public/TedsOutlinerImpl.h
+++ b/Engine/Plugins/Experimental/EditorDataStorageFeatures/Source/TedsOutliner/Public/TedsOutlinerImpl.h
@@ -196,7 +196,7 @@ protected:
 
 	// Check if this row can be displayed in this table viewer
 	bool CanDisplayRow(DataStorage::RowHandle ItemRowHandle) const;
-	
+
 protected:
 	// TEDS Storage Constructs
 	DataStorage::ICoreProvider* Storage{ nullptr };
@@ -256,5 +256,8 @@ protected:
 
 	// Delegate to check if an item is compatible with this table viewer
 	FIsItemCompatible IsItemCompatibleWithTeds;
+
+	// Whether the row query is using the FConditions syntax (TColumn<FMyColumn>()) or the old syntax (.All().Any().None())
+	bool bUsingQueryConditionsSyntax = false;
 };
 } // namsepace UE::Editor::Outliner
