TEDS UI: Add a new SHierarchyViewer widget to view TEDS rows in a tree view

diff --git a/Engine/Plugins/Experimental/EditorDataStorageFeatures/Source/TedsQueryStack/Public/TedsRowFilterNode.h b/Engine/Plugins/Experimental/EditorDataStorageFeatures/Source/TedsQueryStack/Public/TedsRowFilterNode.h
new file mode 100644
--- /dev/null
+++ b/Engine/Plugins/Experimental/EditorDataStorageFeatures/Source/TedsQueryStack/Public/TedsRowFilterNode.h
@@ -0,0 +1,96 @@
+// Copyright Epic Games, Inc. All Rights Reserved
+ 
+#pragma once
+
+#include "Elements/Common/TypedElementCommonTypes.h"
+#include "Elements/Interfaces/TypedElementDataStorageInterface.h"
+#include "TedsQueryStackInterfaces.h"
+
+namespace UE::Editor::DataStorage
+{
+	class ICoreProvider;
+}
+
+namespace UE::Editor::DataStorage::QueryStack
+{
+	/**
+	 * A specialized query stack node that takes another row node as input and only keeps rows that contain (or don't contain) a specific row
+	 */
+	template<TColumnType ColumnType>
+	class FRowFilterNode : public IRowNode
+	{
+	public:
+
+		/**
+		 * Construct an FHierarchyTopLevelRowNode
+		 *
+		 * @param InStorage	The root elements of the new graph to be generated
+		 * @param InParentRowNode Different visualization settings, such as whether it should display the referencers or the dependencies of the NewGraphRootIdentifiers
+		 * @param bInRowsShouldHaveColumn If true, only keep rows that have the required column. If false, only keep rows that DON'T have the required column 
+		 */
+		FRowFilterNode(ICoreProvider* InStorage, const TSharedPtr<IRowNode>& InParentRowNode, bool bInRowsShouldHaveColumn)
+		: Storage(InStorage)
+		, ParentRowNode(InParentRowNode)
+		, CachedParentRevisionID(InParentRowNode->GetRevision())
+		, bRowsShouldHaveColumn(bInRowsShouldHaveColumn)
+		{
+			
+		}
+		
+		virtual ~FRowFilterNode() override = default;
+
+		// IRowNode interface
+		virtual FRowHandleArrayView GetRows() const override
+		{
+			return Rows.GetRows();
+		}
+		
+		virtual FRowHandleArray* GetMutableRows() override
+		{
+			return nullptr;
+		}
+		
+		virtual RevisionId GetRevision() const override
+		{
+			return ParentRowNode->GetRevision();
+		}
+		
+		virtual void Update() override
+		{
+			if (CachedParentRevisionID != ParentRowNode->GetRevision())
+			{
+				UpdateRows();
+				CachedParentRevisionID = ParentRowNode->GetRevision();
+			}
+		}
+		// ~IRowNode interface
+
+	protected:
+
+		// Update our list of top level rows
+		void UpdateRows()
+		{
+			Rows.Reset();
+
+			FRowHandleArrayView ParentRows = ParentRowNode->GetRows();
+
+			for (RowHandle Row : ParentRows)
+			{
+				const bool bHasColumn = Storage->HasColumns<ColumnType>(Row);
+
+				if (bRowsShouldHaveColumn ? bHasColumn : !bHasColumn)
+				{
+					Rows.Add(Row);
+				}
+			}
+		}
+
+	protected:
+
+		FRowHandleArray Rows;
+		ICoreProvider* Storage;
+		TSharedPtr<IRowNode> ParentRowNode;
+		RevisionId CachedParentRevisionID;
+		bool bRowsShouldHaveColumn;
+	};
+}
diff --git a/Engine/Plugins/Experimental/EditorDataStorageFeatures/Source/TedsTableViewer/Private/TedsTableViewerModel.cpp b/Engine/Plugins/Experimental/EditorDataStorageFeatures/Source/TedsTableViewer/Private/TedsTableViewerModel.cpp
--- a/Engine/Plugins/Experimental/EditorDataStorageFeatures/Source/TedsTableViewer/Private/TedsTableViewerModel.cpp
+++ b/Engine/Plugins/Experimental/EditorDataStorageFeatures/Source/TedsTableViewer/Private/TedsTableViewerModel.cpp
@@ -126,6 +126,15 @@ namespace UE::Editor::DataStorage
 		return nullptr;
 	}
 
+	int32 FTedsTableViewerModel::GetColumnIndex(const FName& ColumnName) const
+	{
+		return ColumnsView.IndexOfByPredicate([ColumnName]
+			(const TSharedRef<FTedsTableViewerColumn>& InColumn)
+		{
+			return InColumn->GetColumnName() == ColumnName;
+		});
+	}
+
 	void FTedsTableViewerModel::ForEachColumn(const TFunctionRef<void(const TSharedRef<FTedsTableViewerColumn>&)>& Delegate) const
 	{
 		for(const TSharedRef<FTedsTableViewerColumn>& Column : ColumnsView)
diff --git a/Engine/Plugins/Experimental/EditorDataStorageFeatures/Source/TedsTableViewer/Private/Widgets/STedsHierarchyViewer.cpp b/Engine/Plugins/Experimental/EditorDataStorageFeatures/Source/TedsTableViewer/Private/Widgets/STedsHierarchyViewer.cpp
new file mode 100644
--- /dev/null
+++ b/Engine/Plugins/Experimental/EditorDataStorageFeatures/Source/TedsTableViewer/Private/Widgets/STedsHierarchyViewer.cpp
@@ -0,0 +1,247 @@
+// Copyright Epic Games, Inc. All Rights Reserved.
+
+#include "Widgets/STedsHierarchyViewer.h"
+
+#include "TedsTableViewerColumn.h"
+#include "Columns/SlateDelegateColumns.h"
+#include "Elements/Columns/TypedElementHiearchyColumns.h"
+#include "Elements/Columns/TypedElementUIColumns.h"
+#include "Elements/Common/EditorDataStorageFeatures.h"
+#include "Elements/Framework/TypedElementAttributeBinding.h"
+#include "Elements/Framework/TypedElementDataStorageWidget.h"
+#include "Widgets/STedsTableViewerRow.h"
+#include "Widgets/Views/SListView.h"
+#include "Widgets/Text/STextBlock.h"
+#include "Widgets/STedsTreeView.h"
+
+#define LOCTEXT_NAMESPACE "SHierarchyViewer"
+
+namespace UE::Editor::DataStorage
+{
+	void SHierarchyViewer::Construct(const FArguments& InArgs)
+	{
+		OnSelectionChanged = InArgs._OnSelectionChanged;
+		EmptyRowsMessage = InArgs._EmptyRowsMessage;
+		
+		IUiProvider::FPurposeID CellWigetPurpose = InArgs._CellWidgetPurpose;
+
+		if (!CellWigetPurpose.IsSet())
+		{
+			IUiProvider* StorageUi = GetMutableDataStorageFeature<IUiProvider>(UiFeatureName);
+			CellWigetPurpose = StorageUi->GetGeneralWidgetPurposeID();
+		}
+
+		IUiProvider::FPurposeID HeaderWidgetPurpose = InArgs._HeaderWidgetPurpose;
+
+		if (!HeaderWidgetPurpose.IsSet())
+		{
+			HeaderWidgetPurpose = IUiProvider::FPurposeInfo("General", "Header", NAME_None).GeneratePurposeID();
+		}
+
+		Model = MakeShared<FTedsTableViewerModel>(InArgs._QueryStack, InArgs._Columns, CellWigetPurpose, HeaderWidgetPurpose,
+			FTedsTableViewerModel::FIsItemVisible::CreateSP(this, &SHierarchyViewer::IsItemVisible));
+
+		constexpr bool bRowsShouldHaveColumn = false;
+		HierarchyNode = MakeUnique<QueryStack::FRowFilterNode<FTableRowParentColumn>>(Model->GetDataStorageInterface(), InArgs._QueryStack, bRowsShouldHaveColumn);
+		
+		HeaderRowWidget = SNew( SHeaderRow )
+							.CanSelectGeneratedColumn(true);
+		
+		IUiProvider* StorageUi = Model->GetDataStorageUiProvider();
+
+		if (ensure(StorageUi))
+		{
+			TedsWidget = StorageUi->CreateContainerTedsWidget(InvalidRowHandle);
+			
+			ChildSlot
+			[
+				TedsWidget->AsWidget()
+			];
+		}
+		
+		AddWidgetColumns();
+
+		TreeView = SNew(STedsTreeView, STedsTreeView::FOnGetParent::CreateSP(this, &SHierarchyViewer::GetParentRow))
+			.HeaderRow(HeaderRowWidget)
+			.TopLevelRowsSource(&TopLevelRows)
+			.RowsSource(&Model->GetItems())
+			.OnGenerateRow(this, &SHierarchyViewer::MakeTableRowWidget)
+			.OnSelectionChanged(this, &SHierarchyViewer::OnListSelectionChanged)
+			.SelectionMode(InArgs._ListSelectionMode);
+
+		CreateInternalWidget();
+		
+		// Add each Teds column from the model to our header row widget
+		Model->ForEachColumn([this](const TSharedRef<FTedsTableViewerColumn>& Column)
+		{
+			HeaderRowWidget->AddColumn(Column->ConstructHeaderRowColumn());
+		});
+
+		// Whenever the model changes, refresh the list to update the UI
+		Model->GetOnModelChanged().AddRaw(this, &SHierarchyViewer::OnModelChanged);
+	}
+
+	void SHierarchyViewer::OnModelChanged()
+	{
+		TopLevelRows.Empty();
+
+		HierarchyNode->Update();
+		
+		FRowHandleArrayView Rows = HierarchyNode->GetRows();
+
+		for(const RowHandle RowHandle : Rows)
+		{
+			if(Model->IsRowDisplayable(RowHandle))
+			{
+				FTedsRowHandle TedsRowHandle{ .RowHandle = RowHandle };
+				TopLevelRows.Add(TedsRowHandle);
+			}
+		}
+		
+		TreeView->RequestListRefresh();
+		CreateInternalWidget();
+	}
+
+	void SHierarchyViewer::AddWidgetColumns()
+	{
+		if(ICoreProvider* DataStorage = Model->GetDataStorageInterface())
+		{
+			const RowHandle WidgetRowHandle = TedsWidget->GetRowHandle();
+		
+			if(DataStorage->IsRowAvailable(WidgetRowHandle))
+			{
+				// FHideRowFromUITag - The table viewer should not show up as a row in a table viewer because that will cause all sorts of recursion issues
+				// The others are Columns we are going to bind to attributes on STreeView
+				DataStorage->AddColumns<FHideRowFromUITag, FWidgetContextMenuColumn, FWidgetRowScrolledIntoView, FWidgetDoubleClickedColumn>(WidgetRowHandle);
+			}
+		}
+	}
+
+	TableViewerItemPtr SHierarchyViewer::GetParentRow(TableViewerItemPtr InItem) const
+	{
+		if (FTableRowParentColumn* ParentColumn = Model->GetDataStorageInterface()->GetColumn<FTableRowParentColumn>(InItem))
+		{
+			return FTedsRowHandle{.RowHandle = ParentColumn->Parent};
+		}
+
+		return FTedsRowHandle{.RowHandle = InvalidRowHandle};
+	}
+
+	void SHierarchyViewer::CreateInternalWidget()
+	{
+		TSharedPtr<SWidget> ContentWidget;
+
+		// If there are no rows and the table viewer wants to show a message
+		if(Model->GetRowCount() == 0 && EmptyRowsMessage.IsSet())
+		{
+			ContentWidget = 
+				SNew(SBox)
+					.HAlign(HAlign_Center)
+					.VAlign(VAlign_Center)
+					[
+						SNew(STextBlock)
+							.Text(EmptyRowsMessage)
+					];
+		}
+		else if(Model->GetColumnCount() == 0)
+		{
+			ContentWidget =
+				SNew(SBox)
+					.HAlign(HAlign_Center)
+					.VAlign(VAlign_Center)
+					[
+						SNew(STextBlock)
+							.Text(LOCTEXT("EmptyTableViewerColumnsText", "No columns found to display."))
+					];
+		}
+		else
+		{
+			ContentWidget = TreeView.ToSharedRef();
+		}
+
+		TedsWidget->SetContent(ContentWidget.ToSharedRef());
+	}
+
+	void SHierarchyViewer::RefreshColumnWidgets()
+	{
+		HeaderRowWidget->ClearColumns();
+		Model->ForEachColumn([this](const TSharedRef<FTedsTableViewerColumn>& Column)
+		{
+			HeaderRowWidget->AddColumn(Column->ConstructHeaderRowColumn());
+		});
+
+		CreateInternalWidget();
+	}
+
+	void SHierarchyViewer::OnListSelectionChanged(TableViewerItemPtr Item, ESelectInfo::Type SelectInfo)
+	{
+		if(OnSelectionChanged.IsBound())
+		{
+			OnSelectionChanged.Execute(Item);
+		}
+	}
+
+	void SHierarchyViewer::SetColumns(const TArray<TWeakObjectPtr<const UScriptStruct>>& Columns)
+	{
+		Model->SetColumns(Columns);
+		RefreshColumnWidgets();
+	}
+
+	void SHierarchyViewer::AddCustomColumn(const TSharedRef<FTedsTableViewerColumn>& InColumn)
+	{
+		Model->AddCustomColumn(InColumn);
+		RefreshColumnWidgets();
+	}
+
+	void SHierarchyViewer::ForEachSelectedRow(TFunctionRef<void(RowHandle)> InCallback) const
+	{
+		TArray<TableViewerItemPtr> SelectedRows;
+		TreeView->GetSelectedItems(SelectedRows);
+
+		for(TableViewerItemPtr& Row : SelectedRows)
+		{
+			InCallback(Row);
+		}
+	}
+
+	RowHandle SHierarchyViewer::GetWidgetRowHandle() const
+	{
+		return TedsWidget->GetRowHandle();
+	}
+
+	void SHierarchyViewer::SetSelection(RowHandle Row, bool bSelected, const ESelectInfo::Type SelectInfo) const
+	{
+		FTedsRowHandle TedsRowHandle{ .RowHandle = Row };
+		TreeView->SetItemSelection(TedsRowHandle, bSelected, SelectInfo);
+	}
+
+	void SHierarchyViewer::ScrollIntoView(RowHandle Row) const
+	{
+		FTedsRowHandle TedsRowHandle{ .RowHandle = Row };
+		TreeView->RequestScrollIntoView(TedsRowHandle);
+	}
+
+	void SHierarchyViewer::ClearSelection() const
+	{
+		TreeView->ClearSelection();
+	}
+
+	TSharedRef<SWidget> SHierarchyViewer::AsWidget()
+	{
+		return AsShared();
+	}
+
+	bool SHierarchyViewer::IsItemVisible(TableViewerItemPtr InItem) const
+	{
+		return TreeView->IsItemVisible(InItem);
+	}
+
+	TSharedRef<ITableRow> SHierarchyViewer::MakeTableRowWidget(TableViewerItemPtr InItem, const TSharedRef<STableViewBase>& OwnerTable) const
+	{
+		return SNew(SHierarchyViewerRow, OwnerTable, Model.ToSharedRef())
+				.Item(InItem);
+	}
+} // namespace UE::Editor::DataStorage
+
+#undef LOCTEXT_NAMESPACE //"SHierarchyViewer"
+
diff --git a/Engine/Plugins/Experimental/EditorDataStorageFeatures/Source/TedsTableViewer/Private/Widgets/STedsTableViewerRow.cpp b/Engine/Plugins/Experimental/EditorDataStorageFeatures/Source/TedsTableViewer/Private/Widgets/STedsTableViewerRow.cpp
--- a/Engine/Plugins/Experimental/EditorDataStorageFeatures/Source/TedsTableViewer/Private/Widgets/STedsTableViewerRow.cpp
+++ b/Engine/Plugins/Experimental/EditorDataStorageFeatures/Source/TedsTableViewer/Private/Widgets/STedsTableViewerRow.cpp
@@ -36,4 +36,32 @@ namespace UE::Editor::DataStorage
 
 		return SNullWidget::NullWidget;
 	}
+
+	TSharedRef<SWidget> SHierarchyViewerRow::GenerateWidgetForColumn(const FName& ColumnName)
+	{
+		TSharedRef<SWidget> ActualWidget = STedsTableViewerRow::GenerateWidgetForColumn(ColumnName);
+
+		// We show the expander arrow on the first column only
+		if (TableViewerModel->GetColumnIndex(ColumnName) == 0)
+		{
+			return SNew(SBox)
+				.MinDesiredHeight(20.0f)
+				[
+					SNew(SHorizontalBox)
+					+ SHorizontalBox::Slot()
+					.AutoWidth()
+					.Padding(6.f, 0.f, 0.f, 0.f)
+					[
+						SNew(SExpanderArrow, SharedThis(this)).IndentAmount(12)
+					]
+					+ SHorizontalBox::Slot()
+					.FillWidth(1.0f)
+					[
+						ActualWidget
+					]
+				];
+		}
+
+		return ActualWidget;
+	}
 } // namespace UE::Editor::DataStorage
diff --git a/Engine/Plugins/Experimental/EditorDataStorageFeatures/Source/TedsTableViewer/Private/Widgets/STedsTreeView.cpp b/Engine/Plugins/Experimental/EditorDataStorageFeatures/Source/TedsTableViewer/Private/Widgets/STedsTreeView.cpp
new file mode 100644
--- /dev/null
+++ b/Engine/Plugins/Experimental/EditorDataStorageFeatures/Source/TedsTableViewer/Private/Widgets/STedsTreeView.cpp
@@ -0,0 +1,81 @@
+// Copyright Epic Games, Inc. All Rights Reserved
+ 
+#pragma once
+
+#include "STedsTreeView.h"
+
+namespace UE::Editor::DataStorage
+{
+	void STedsTreeView::Construct( const FArguments& InArgs, const FOnGetParent& InOnGetParent )
+	{
+		checkf(InOnGetParent.IsBound(), TEXT("Must provide FOnGetParent delegate."));
+		checkf(InArgs._RowsSource, TEXT("Must provide a source for the rows to be displayed."));
+
+		OnGetParent = InOnGetParent;
+		RowsSource = InArgs._RowsSource;
+			
+		STreeView::Construct(STreeView<TableViewerItemPtr>::FArguments()
+			.HeaderRow(InArgs._HeaderRow)
+			.TreeItemsSource( InArgs._TopLevelRowsSource)
+			.OnGenerateRow(InArgs._OnGenerateRow)
+			.OnSelectionChanged(InArgs._OnSelectionChanged)
+			.SelectionMode(InArgs._SelectionMode)
+			.OnGetChildren(this, &STedsTreeView::GetChildren_Internal)
+			);
+	}
+
+	void STedsTreeView::Tick( const FGeometry& AllottedGeometry, const double InCurrentTime, const float InDeltaTime )
+	{
+		if (bDirty)
+		{
+			UpdateTreeMap();
+			bDirty = false;
+		}
+
+		STreeView::Tick(AllottedGeometry, InCurrentTime, InDeltaTime);
+
+	}
+		
+	void STedsTreeView::RequestListRefresh()
+	{
+		bDirty = true;
+		STreeView<TableViewerItemPtr>::RequestListRefresh();
+	}
+	
+	void STedsTreeView::UpdateTreeMap()
+	{
+		/*
+		 * Currently, whenever a tree refresh is requested we re-calulate the hierarchy of all items and store it, and simply re-use the existing
+		 * top down logic in STreeView to provide the actual hierarchy to the widget. In the future this can be optimized to override the whole logic
+		 * contained in STreeView's Tick() function and use a completely bottom up approach.
+		 */
+		TreeMap.Empty();
+
+		const TArrayView<const TableViewerItemPtr> ItemsSourceRef(*RowsSource);
+
+		// For each item, grab the parent and add the item to the parent item's array of children
+		for (TableViewerItemPtr Item : ItemsSourceRef)
+		{
+			if (TableViewerItemPtr Parent = OnGetParent.Execute(Item))
+			{
+				if (TArray<TableViewerItemPtr>* FoundChildren = TreeMap.Find(Parent))
+				{
+					FoundChildren->Add(Item);
+				}
+				else
+				{
+					TreeMap.Emplace(Parent, {Item});
+				}
+			}
+		}
+	}
+
+	void STedsTreeView::GetChildren_Internal(TableViewerItemPtr InParent, TArray<TableViewerItemPtr>& OutChildren)
+	{
+		// Simply look up the child in our internally cached array for the row
+		if (TArray<TableViewerItemPtr>* FoundArray = TreeMap.Find(InParent))
+		{
+			OutChildren = *FoundArray; 
+		}
+	}
+}
\ No newline at end of file
diff --git a/Engine/Plugins/Experimental/EditorDataStorageFeatures/Source/TedsTableViewer/Private/Widgets/STedsTreeView.h b/Engine/Plugins/Experimental/EditorDataStorageFeatures/Source/TedsTableViewer/Private/Widgets/STedsTreeView.h
new file mode 100644
--- /dev/null
+++ b/Engine/Plugins/Experimental/EditorDataStorageFeatures/Source/TedsTableViewer/Private/Widgets/STedsTreeView.h
@@ -0,0 +1,70 @@
+// Copyright Epic Games, Inc. All Rights Reserved
+ 
+#pragma once
+
+#include "Widgets/Views/STreeView.h"
+#include "TedsTableViewerModel.h"
+
+namespace UE::Editor::DataStorage
+{
+	/**
+	 * A tree view widget specialized for TableViewerItemPtr to support constructing the widget using a bottom up hierarchy (FOnGetParent)
+	 */
+	class STedsTreeView : public STreeView<TableViewerItemPtr>
+	{
+	public:
+
+		/** Delegate that gets the parent row handle for a given row */
+		DECLARE_DELEGATE_RetVal_OneParam(TableViewerItemPtr /** OutParent */, FOnGetParent, TableViewerItemPtr /** InChild */)
+
+		SLATE_BEGIN_ARGS(STedsTreeView) {}
+
+			/** All the rows that can be displayed by the widget */
+			SLATE_ARGUMENT( const TArray<TableViewerItemPtr>*, RowsSource )
+
+			/** Only the top level rows being displayed by the widget */
+			SLATE_ARGUMENT( const TArray<TableViewerItemPtr>*, TopLevelRowsSource )
+
+			/** Delegate to generate the actual row widget */
+			SLATE_EVENT( FOnGenerateRow, OnGenerateRow )
+
+			/** Delegate fired on selection change */
+			SLATE_EVENT( FOnSelectionChanged, OnSelectionChanged )
+
+			/** Delegate that determines the selection mode */
+			SLATE_ATTRIBUTE( ESelectionMode::Type, SelectionMode )
+
+			/** The header row widget to use */
+			SLATE_ARGUMENT( TSharedPtr<SHeaderRow>, HeaderRow )
+			
+		SLATE_END_ARGS()
+		
+		void Construct( const FArguments& InArgs, const FOnGetParent& InOnGetParent );
+
+		virtual void Tick( const FGeometry& AllottedGeometry, const double InCurrentTime, const float InDeltaTime ) override;
+		virtual void RequestListRefresh() override;
+
+	protected:
+
+		// Update the internal tree map used to contain the hierarchy
+		void UpdateTreeMap();
+
+		// For a given row, get all the children
+		void GetChildren_Internal(TableViewerItemPtr InParent, TArray<TableViewerItemPtr>& OutChildren);
+
+		protected:
+		
+		FOnGetParent OnGetParent;
+
+		// All rows that can be shown in this tree, including the full hieararchy regardless of whether it is currently expanded
+		const TArray<TableViewerItemPtr>* RowsSource = nullptr;
+
+		// The internal tree map used to contain the hierarchy
+		TMap<TableViewerItemPtr, TArray<TableViewerItemPtr>> TreeMap;
+
+		// Whether the TreeMap needs to be re-calculated
+		bool bDirty = false;
+	};
+}
+
+
diff --git a/Engine/Plugins/Experimental/EditorDataStorageFeatures/Source/TedsTableViewer/Public/TedsTableViewerModel.h b/Engine/Plugins/Experimental/EditorDataStorageFeatures/Source/TedsTableViewer/Public/TedsTableViewerModel.h
--- a/Engine/Plugins/Experimental/EditorDataStorageFeatures/Source/TedsTableViewer/Public/TedsTableViewerModel.h
+++ b/Engine/Plugins/Experimental/EditorDataStorageFeatures/Source/TedsTableViewer/Public/TedsTableViewerModel.h
@@ -47,6 +47,9 @@ namespace UE::Editor::DataStorage
 		// Get a specific column that the table viewer is displaying by name
 		TEDSTABLEVIEWER_API TSharedPtr<FTedsTableViewerColumn> GetColumn(const FName& ColumnName) const;
 
+		// Get the index of a column by name (INDEX_NONE if column doesn't exist)
+		TEDSTABLEVIEWER_API int32 GetColumnIndex(const FName& ColumnName) const;
+
 		// Execute a delegate for each column in the model
 		TEDSTABLEVIEWER_API void ForEachColumn(const TFunctionRef<void(const TSharedRef<FTedsTableViewerColumn>&)>& Delegate) const;
 
@@ -63,6 +66,9 @@ namespace UE::Editor::DataStorage
 
 		TEDSTABLEVIEWER_API IUiProvider* GetDataStorageUiProvider() const;
 
+		// Check whether a row is allowed to be displayed in the table viewer
+		TEDSTABLEVIEWER_API bool IsRowDisplayable(RowHandle InRowHandle) const;
+
 	protected:
 
 		// Generate the actual columns to display in the UI using TEDS UI
@@ -75,9 +81,6 @@ namespace UE::Editor::DataStorage
 
 		void Refresh();
 
-		// Check whether a row is allowed to be displayed in the table viewer
-		bool IsRowDisplayable(RowHandle InRowHandle) const;
-
 		void ValidateRequestedColumns();
 
 	private:
diff --git a/Engine/Plugins/Experimental/EditorDataStorageFeatures/Source/TedsTableViewer/Public/Widgets/STedsHierarchyViewer.h b/Engine/Plugins/Experimental/EditorDataStorageFeatures/Source/TedsTableViewer/Public/Widgets/STedsHierarchyViewer.h
new file mode 100644
--- /dev/null
+++ b/Engine/Plugins/Experimental/EditorDataStorageFeatures/Source/TedsTableViewer/Public/Widgets/STedsHierarchyViewer.h
@@ -0,0 +1,147 @@
+// Copyright Epic Games, Inc. All Rights Reserved.
+#pragma once
+
+#include "ITedsTableViewer.h"
+#include "Elements/Common/TypedElementHandles.h"
+#include "TedsRowFilterNode.h"
+#include "TedsTableViewerModel.h"
+#include "TedsQueryStackInterfaces.h"
+#include "Elements/Columns/TypedElementHiearchyColumns.h"
+#include "Widgets/SCompoundWidget.h"
+#include "Widgets/Views/STableRow.h"
+
+class ITableRow;
+class STableViewBase;
+class SHeaderRow;
+
+namespace UE::Editor::DataStorage
+{
+	class ITedsWidget;
+	class STedsTreeView;
+
+}
+
+namespace UE::Editor::DataStorage
+{
+	class FTedsTableViewerColumn;
+
+	/*
+	 * A table viewer widget can be used to show a visual representation of data in TEDS.
+	 * The rows to display can be specified using a RowQueryStack, and the columns to display are directly input into the widget
+	 * Example usage:
+	 * 
+	 *	SNew(SHierarchyViewer)
+     *		.QueryStack(MakeShared<UE::Editor::DataStorage::FQueryStackNode_RowView>(&Rows))
+	 *		.Columns({FTypedElementLabelColumn::StaticStruct(), FTypedElementClassTypeInfoColumn::StaticStruct());
+	 */
+	class SHierarchyViewer : public SCompoundWidget, public ITableViewer
+	{
+	public:
+		
+		// Delegate fired when the selection in the table viewer changes
+		DECLARE_DELEGATE_OneParam(FOnSelectionChanged, RowHandle)
+
+		SLATE_BEGIN_ARGS(SHierarchyViewer)
+			: _ListSelectionMode(ESelectionMode::Type::Single)
+		{
+			
+		}
+		
+
+		// Query Stack that will supply the rows to be displayed
+		SLATE_ARGUMENT(TSharedPtr<QueryStack::IRowNode>, QueryStack)
+
+		// The Columns that this table viewer will display
+		// Table Viewer TODO: How do we specify column metadata (ReadOnly or ReadWrite)?
+		SLATE_ARGUMENT(TArray<TWeakObjectPtr<const UScriptStruct>>, Columns)
+
+		// The widget purpose to use to create the cell widgets
+		SLATE_ARGUMENT(IUiProvider::FPurposeID, CellWidgetPurpose)
+			
+		// The widget purpose to use to create the header widgets
+		SLATE_ARGUMENT(IUiProvider::FPurposeID, HeaderWidgetPurpose)
+		
+		// Delegate called when the selection changes
+		SLATE_EVENT(FOnSelectionChanged, OnSelectionChanged)
+
+		// The selection mode for the table viewer (single/multi etc)
+		SLATE_ARGUMENT(ESelectionMode::Type, ListSelectionMode)
+
+		// The message to show in place of the table viewer when there are no rows provided by the current query stack
+		// Empty = simply show the column headers instead of a message
+		SLATE_ATTRIBUTE(FText, EmptyRowsMessage)
+
+		SLATE_END_ARGS()
+
+	public:
+		
+		TEDSTABLEVIEWER_API void Construct(const FArguments& InArgs);
+
+		// Clear the current list of columns being displayed and set it to the given list
+		TEDSTABLEVIEWER_API void SetColumns(const TArray<TWeakObjectPtr<const UScriptStruct>>& Columns);
+
+		// Add a custom column to display in the table viewer, that doesn't necessarily map to a Teds column
+		TEDSTABLEVIEWER_API void AddCustomColumn(const TSharedRef<FTedsTableViewerColumn>& InColumn);
+
+		// Execute the given callback for each row that is selected in the table viewer
+		TEDSTABLEVIEWER_API virtual void ForEachSelectedRow(TFunctionRef<void(RowHandle)> InCallback) const override;
+
+		// Get the row handle for the widget row the table viewer's contents are stored in
+		TEDSTABLEVIEWER_API virtual RowHandle GetWidgetRowHandle() const override;
+		
+		// Select the given row in the table viewer
+		TEDSTABLEVIEWER_API virtual void SetSelection(RowHandle Row, bool bSelected, const ESelectInfo::Type SelectInfo) const override;
+
+		// Scroll the given row into view in the table viewer
+		TEDSTABLEVIEWER_API virtual void ScrollIntoView(RowHandle Row) const override;
+
+		TEDSTABLEVIEWER_API virtual void ClearSelection() const override;
+		
+		TEDSTABLEVIEWER_API virtual TSharedRef<SWidget> AsWidget() override;
+
+	protected:
+		
+		TSharedRef<ITableRow> MakeTableRowWidget(TableViewerItemPtr InItem, const TSharedRef<STableViewBase>& OwnerTable) const;
+
+		bool IsItemVisible(TableViewerItemPtr InItem) const;
+
+		void CreateInternalWidget();
+
+		void RefreshColumnWidgets();
+
+		void OnListSelectionChanged(TableViewerItemPtr Item, ESelectInfo::Type SelectInfo);
+
+		void AddWidgetColumns();
+
+		TableViewerItemPtr GetParentRow(TableViewerItemPtr InItem) const;
+
+		void OnModelChanged();
+
+	private:
+
+		// The actual ListView widget that displays the rows
+		TSharedPtr<STedsTreeView> TreeView;
+
+		// The actual header widget
+		TSharedPtr<SHeaderRow> HeaderRowWidget;
+
+		// Our model class
+		TSharedPtr<FTedsTableViewerModel> Model;
+
+		// Delegate fired when the selection changes
+		FOnSelectionChanged OnSelectionChanged;
+
+		// Wrapper Teds Widget around our contents so we can use Teds columns to specify behavior
+		TSharedPtr<ITedsWidget> TedsWidget;
+
+		// The message to show in place of the table viewer when there are no rows provided by the current query stack
+		TAttribute<FText> EmptyRowsMessage;
+
+		// A specialized query stack node to only contain the top level rows present in the actual query stack we are using.
+		TUniquePtr<QueryStack::FRowFilterNode<FTableRowParentColumn>> HierarchyNode;
+		
+		// The cached list of top level rows
+		TArray<TableViewerItemPtr> TopLevelRows;
+
+	};
+}
\ No newline at end of file
diff --git a/Engine/Plugins/Experimental/EditorDataStorageFeatures/Source/TedsTableViewer/Public/Widgets/STedsTableViewerRow.h b/Engine/Plugins/Experimental/EditorDataStorageFeatures/Source/TedsTableViewer/Public/Widgets/STedsTableViewerRow.h
--- a/Engine/Plugins/Experimental/EditorDataStorageFeatures/Source/TedsTableViewer/Public/Widgets/STedsTableViewerRow.h
+++ b/Engine/Plugins/Experimental/EditorDataStorageFeatures/Source/TedsTableViewer/Public/Widgets/STedsTableViewerRow.h
@@ -10,7 +10,7 @@ namespace UE::Editor::DataStorage
 {
 	class STedsTableViewer;
 	
-	/** Widget that represents a row.  Generates widgets for each column on demand. */
+	/** Widget that represents a row in the table viewer.  Generates widgets for each column on demand. */
 	class STedsTableViewerRow
 		: public SMultiColumnTableRow< TableViewerItemPtr >
 	{
@@ -35,4 +35,11 @@ namespace UE::Editor::DataStorage
 		TSharedPtr<FTedsTableViewerModel> TableViewerModel;
 		TableViewerItemPtr Item;
 	};
+
+	/** Widget that represents a row in the hierarchy viewers.  Generates widgets for each column and adds an expander arrow to the very first column. */
+	class SHierarchyViewerRow : public STedsTableViewerRow
+	{
+	public:
+		virtual TSharedRef<SWidget> GenerateWidgetForColumn( const FName& ColumnName ) override;
+	};
 } // namespace UE::Editor::DataStorage
