Introduced EnterEditMode column to Teds Outliner.

diff --git a/Engine/Plugins/Experimental/EditorDataStorage/Source/TedsUI/Private/Widgets/LabelWidget.cpp b/Engine/Plugins/Experimental/EditorDataStorage/Source/TedsUI/Private/Widgets/LabelWidget.cpp
--- a/Engine/Plugins/Experimental/EditorDataStorage/Source/TedsUI/Private/Widgets/LabelWidget.cpp
+++ b/Engine/Plugins/Experimental/EditorDataStorage/Source/TedsUI/Private/Widgets/LabelWidget.cpp
@@ -11,7 +11,6 @@
 #include "Elements/Framework/TypedElementAttributeBinding.h"
 #include "Elements/Framework/TypedElementQueryBuilder.h"
 #include "Elements/Interfaces/TypedElementDataStorageInterface.h"
-#include "Elements/Interfaces/Capabilities/TypedElementUiEditableCapability.h"
 #include "Elements/Interfaces/Capabilities/TypedElementUiTextCapability.h"
 #include "Elements/Interfaces/Capabilities/TypedElementUiTooltipCapability.h"
 #include "Elements/Interfaces/Capabilities/TypedElementUiStyleOverrideCapability.h"
diff --git a/Engine/Plugins/Experimental/EditorDataStorageFeatures/Source/TedsOutliner/Private/Widgets/OutlinerTextWidget.cpp b/Engine/Plugins/Experimental/EditorDataStorageFeatures/Source/TedsOutliner/Private/Widgets/OutlinerTextWidget.cpp
--- a/Engine/Plugins/Experimental/EditorDataStorageFeatures/Source/TedsOutliner/Private/Widgets/OutlinerTextWidget.cpp
+++ b/Engine/Plugins/Experimental/EditorDataStorageFeatures/Source/TedsOutliner/Private/Widgets/OutlinerTextWidget.cpp
@@ -7,7 +7,6 @@
 #include "Elements/Columns/TypedElementLabelColumns.h"
 #include "Elements/Columns/TypedElementMiscColumns.h"
 #include "Elements/Framework/TypedElementAttributeBinding.h"
-#include "Elements/Interfaces/Capabilities/TypedElementUiEditableCapability.h"
 #include "Elements/Interfaces/Capabilities/TypedElementUiStyleOverrideCapability.h"
 #include "Settings/EditorStyleSettings.h"
 #include "Widgets/Text/SInlineEditableTextBlock.h"
@@ -128,7 +127,13 @@ TSharedPtr<SInlineEditableTextBlock> FOutlinerTextWidgetConstructor::CreateEdita
 		.OverflowPolicy(bShouldUseMiddleEllipsis ? ETextOverflowPolicy::MiddleEllipsis : TOptional<ETextOverflowPolicy>())
 		.IsSelected(WidgetRowBinder.BindEvent(&FExternalWidgetSelectionColumn::IsSelected));
 
-	TextBlock->AddMetadata(MakeShared<TTypedElementUiEditableCapability<SInlineEditableTextBlock>>(*TextBlock));
+	DataStorage->AddColumn<FWidgetEnterEditModeColumn>(TargetRow, FWidgetEnterEditModeColumn{
+		.OnEnterEditMode = FSimpleDelegate::CreateLambda([TextBlock]()
+		{
+			TextBlock->EnterEditingMode();
+		})
+	});
+
 	TextBlock->AddMetadata(MakeShared<TTypedElementUiStyleOverrideCapability<SInlineEditableTextBlock>>(*TextBlock));
 	return TextBlock;
 }
diff --git a/Engine/Source/Developer/ToolWidgets/Public/Columns/SlateDelegateColumns.h b/Engine/Source/Developer/ToolWidgets/Public/Columns/SlateDelegateColumns.h
--- a/Engine/Source/Developer/ToolWidgets/Public/Columns/SlateDelegateColumns.h
+++ b/Engine/Source/Developer/ToolWidgets/Public/Columns/SlateDelegateColumns.h
@@ -57,4 +57,14 @@ struct FWidgetDoubleClickedColumn : public FEditorDataStorageColumn
 	
 	GENERATED_BODY()
 	FOnMouseButtonDoubleClick OnMouseButtonDoubleClick;
+};
+
+/**
+ * Column that triggers a signal when the label is edited.
+ */
+USTRUCT(meta = (DisplayName = "On Label Edit"))
+struct FWidgetEnterEditModeColumn : public FEditorDataStorageColumn
+{
+	GENERATED_BODY()
+	FSimpleDelegate OnEnterEditMode;
 };
\ No newline at end of file
diff --git a/Engine/Source/Runtime/TypedElementFramework/Public/Elements/Interfaces/Capabilities/TypedElementUiEditableCapability.h b/Engine/Source/Runtime/TypedElementFramework/Public/Elements/Interfaces/Capabilities/TypedElementUiEditableCapability.h
deleted file mode 100644
