EntityFramework: Introduced derived SceneOutliner classes to separate Entity name id suffix in EntityPrefabEditor.

diff --git a/Engine/Plugins/Experimental/EditorDataStorageFeatures/Source/TedsOutliner/Private/TedsOutlinerImpl.cpp b/Engine/Plugins/Experimental/EditorDataStorageFeatures/Source/TedsOutliner/Private/TedsOutlinerImpl.cpp
--- a/Engine/Plugins/Experimental/EditorDataStorageFeatures/Source/TedsOutliner/Private/TedsOutlinerImpl.cpp
+++ b/Engine/Plugins/Experimental/EditorDataStorageFeatures/Source/TedsOutliner/Private/TedsOutlinerImpl.cpp
@@ -224,14 +224,21 @@ void FTedsOutlinerImpl::SetSelection(const TArray<DataStorage::RowHandle>& InSel
 	}
 }
 
-TSharedRef<SWidget> FTedsOutlinerImpl::CreateLabelWidgetForItem(DataStorage::RowHandle InRowHandle, const STableRow<FSceneOutlinerTreeItemPtr>& InRow, bool bIsInteractable) const
+TSharedRef<SWidget> FTedsOutlinerImpl::CreateLabelWidget(
+	DataStorage::ICoreProvider& Storage,
+	DataStorage::IUiProvider& StorageUi,
+	ISceneOutliner& Outliner,
+	DataStorage::IUiProvider::FPurposeID LabelWidgetPurpose,
+	DataStorage::RowHandle Row,
+	const STableRow<FSceneOutlinerTreeItemPtr>& RowItem,
+	bool bIsInteractable)
 {
 	// Find the best matching label widget for a given list of columns
-	auto CreateWidgetConstructor = [this](TArray<TWeakObjectPtr<const UScriptStruct>>& ColumnTypes) -> TSharedPtr<FTypedElementWidgetConstructor>
+	auto CreateWidgetConstructor = [&StorageUi, LabelWidgetPurpose](TArray<TWeakObjectPtr<const UScriptStruct>>& ColumnTypes) -> TSharedPtr<FTypedElementWidgetConstructor>
 	{
 		TSharedPtr<FTypedElementWidgetConstructor> OutWidgetConstructorPtr = nullptr;
 
-		StorageUi->CreateWidgetConstructors(StorageUi->FindPurpose(LabelWidgetPurpose), DataStorage::IUiProvider::EMatchApproach::LongestMatch, ColumnTypes, {},
+		StorageUi.CreateWidgetConstructors(StorageUi.FindPurpose(LabelWidgetPurpose), DataStorage::IUiProvider::EMatchApproach::LongestMatch, ColumnTypes, {},
 					[&OutWidgetConstructorPtr, ColumnTypes](
 					TUniquePtr<FTypedElementWidgetConstructor> CreatedConstructor, 
 					TConstArrayView<TWeakObjectPtr<const UScriptStruct>> MatchedColumnTypes)
@@ -251,26 +258,26 @@ TSharedRef<SWidget> FTedsOutlinerImpl::CreateLabelWidgetForItem(DataStorage::Row
 	static FQueryDescription MetaDataQueryRead = Select().ReadOnly<FTypedElementLabelColumn>().Where().Compile();
 
 	// Create a widget for this row using the given widget constructor
-	auto CreateWidget = [InRowHandle, this, &InRow, bIsInteractable](const TSharedPtr<FTypedElementWidgetConstructor>& WidgetConstructor) -> TSharedPtr<SWidget>
+	auto CreateWidget = [&Storage, &StorageUi, &Outliner, Row, &RowItem, bIsInteractable](const TSharedPtr<FTypedElementWidgetConstructor>& WidgetConstructor) -> TSharedPtr<SWidget>
 	{
 		// Create metadata for the query
 		FQueryMetaDataView QueryMetaDataView = bIsInteractable ?
 			FQueryMetaDataView(MetaDataQueryReadWrite) : FQueryMetaDataView(MetaDataQueryRead);
 
-		RowHandle UiRowHandle = Storage->AddRow(Storage->FindTable(DataStorage::TableViewerUtils::GetWidgetTableName()));
+		RowHandle UiRowHandle = Storage.AddRow(Storage.FindTable(DataStorage::TableViewerUtils::GetWidgetTableName()));
 
-		if (FTypedElementRowReferenceColumn* RowReference = Storage->GetColumn<FTypedElementRowReferenceColumn>(UiRowHandle))
+		if (FTypedElementRowReferenceColumn* RowReference = Storage.GetColumn<FTypedElementRowReferenceColumn>(UiRowHandle))
 		{
-			RowReference->Row = InRowHandle;
+			RowReference->Row = Row;
 		}
 
-		Storage->AddColumn(UiRowHandle, FTedsOutlinerColumn{.Outliner = StaticCastSharedRef<ISceneOutliner>(SceneOutliner->AsShared())});
+		Storage.AddColumn(UiRowHandle, FTedsOutlinerColumn{.Outliner = StaticCastSharedRef<ISceneOutliner>(Outliner.AsShared())});
 		
-		TSharedPtr<SWidget> Widget = StorageUi->ConstructWidget(UiRowHandle, *WidgetConstructor, QueryMetaDataView);
+		TSharedPtr<SWidget> Widget = StorageUi.ConstructWidget(UiRowHandle, *WidgetConstructor, QueryMetaDataView);
 		
-		if (FExternalWidgetSelectionColumn* ExternalWidgetSelectionColumn = Storage->GetColumn<FExternalWidgetSelectionColumn>(UiRowHandle))
+		if (FExternalWidgetSelectionColumn* ExternalWidgetSelectionColumn = Storage.GetColumn<FExternalWidgetSelectionColumn>(UiRowHandle))
 		{
-			ExternalWidgetSelectionColumn->IsSelected = FIsSelected::CreateSP(&InRow, &STableRow<FSceneOutlinerTreeItemPtr>::IsSelectedExclusively);
+			ExternalWidgetSelectionColumn->IsSelected = FIsSelected::CreateSP(&RowItem, &STableRow<FSceneOutlinerTreeItemPtr>::IsSelectedExclusively);
 		}
 		return Widget;
 		
@@ -278,7 +285,7 @@ TSharedRef<SWidget> FTedsOutlinerImpl::CreateLabelWidgetForItem(DataStorage::Row
 
 	// Get all the columns on the given row
 	TArray<TWeakObjectPtr<const UScriptStruct>> Columns;
-	Storage->ListColumns(InRowHandle, [&Columns](const UScriptStruct& ColumnType)
+	Storage.ListColumns(Row, [&Columns](const UScriptStruct& ColumnType)
 	{
 		Columns.Add(&ColumnType);
 	});
@@ -300,6 +307,11 @@ TSharedRef<SWidget> FTedsOutlinerImpl::CreateLabelWidgetForItem(DataStorage::Row
 	return SNullWidget::NullWidget;
 }
 
+TSharedRef<SWidget> FTedsOutlinerImpl::CreateLabelWidgetForItem(DataStorage::RowHandle InRowHandle, const STableRow<FSceneOutlinerTreeItemPtr>& InRow, bool bIsInteractable) const
+{
+	return CreateLabelWidget(*Storage, *StorageUi, *SceneOutliner, LabelWidgetPurpose, InRowHandle, InRow, bIsInteractable);
+}
+
 void FTedsOutlinerImpl::AppendQuery(DataStorage::FQueryDescription& Query1, const DataStorage::FQueryDescription& Query2)
 {
 	// TEDS-Outliner TODO: We simply discard duplicate types for now but we probably want a more robust system to detect duplicates and conflicting conditions
diff --git a/Engine/Plugins/Experimental/EditorDataStorageFeatures/Source/TedsOutliner/Public/TedsOutlinerImpl.h b/Engine/Plugins/Experimental/EditorDataStorageFeatures/Source/TedsOutliner/Public/TedsOutlinerImpl.h
--- a/Engine/Plugins/Experimental/EditorDataStorageFeatures/Source/TedsOutliner/Public/TedsOutlinerImpl.h
+++ b/Engine/Plugins/Experimental/EditorDataStorageFeatures/Source/TedsOutliner/Public/TedsOutlinerImpl.h
@@ -126,6 +126,15 @@ class TEDSOUTLINER_API FTedsOutlinerImpl : public TSharedFromThis<FTedsOutlinerI
 {
 
 public:
+	// Helper function to create a label widget for a given row
+	static TSharedRef<SWidget> CreateLabelWidget(
+		DataStorage::ICoreProvider& Storage,
+		DataStorage::IUiProvider& StorageUi,
+		ISceneOutliner& Outliner,
+		DataStorage::IUiProvider::FPurposeID LabelWidgetPurpose,
+		DataStorage::RowHandle Row,
+		const STableRow<FSceneOutlinerTreeItemPtr>& RowItem,
+		bool bIsInteractable = true);
 
 	FTedsOutlinerImpl(const FTedsOutlinerParams& InParams, ISceneOutlinerMode* InMode);
 	virtual ~FTedsOutlinerImpl();
