[EditorDataStorageFeatures] Teds Tile View Widgets

diff --git a/Engine/Plugins/Experimental/EditorDataStorage/Source/TedsUI/Private/Widgets/NameWidget.cpp b/Engine/Plugins/Experimental/EditorDataStorage/Source/TedsUI/Private/Widgets/NameWidget.cpp
new file mode 100644
--- /dev/null
+++ b/Engine/Plugins/Experimental/EditorDataStorage/Source/TedsUI/Private/Widgets/NameWidget.cpp
@@ -0,0 +1,36 @@
+// Copyright Epic Games, Inc. All Rights Reserved.
+
+#include "Widgets/NameWidget.h"
+
+#include "Elements/Framework/TypedElementAttributeBinding.h"
+#include "Elements/Columns/TypedElementMiscColumns.h"
+#include "Widgets/Text/STextBlock.h"
+
+#define LOCTEXT_NAMESPACE "NameWidget"
+
+void UNameWidgetFactory::RegisterWidgetConstructors(UE::Editor::DataStorage::ICoreProvider& DataStorage,
+	UE::Editor::DataStorage::IUiProvider& DataStorageUi) const
+{
+	using namespace UE::Editor::DataStorage::Queries;
+	DataStorageUi.RegisterWidgetFactory<FNameWidgetConstructor>(
+		DataStorageUi.FindPurpose(DataStorageUi.GetGeneralWidgetPurposeID()),
+		TColumn<FNameColumn>());
+}
+
+FNameWidgetConstructor::FNameWidgetConstructor()
+	: FSimpleWidgetConstructor(StaticStruct())
+{
+}
+
+TSharedPtr<SWidget> FNameWidgetConstructor::CreateWidget(UE::Editor::DataStorage::ICoreProvider* DataStorage,
+	UE::Editor::DataStorage::IUiProvider* DataStorageUi, UE::Editor::DataStorage::RowHandle TargetRow, UE::Editor::DataStorage::RowHandle WidgetRow,
+	const UE::Editor::DataStorage::FMetaDataView& Arguments)
+{
+	UE::Editor::DataStorage::FAttributeBinder Binder(TargetRow, DataStorage);
+
+	return SNew(STextBlock)
+				.Text(Binder.BindText(&FNameColumn::Name))
+				.ToolTipText(Binder.BindText(&FNameColumn::Name));
+}
+
+#undef LOCTEXT_NAMESPACE
diff --git a/Engine/Plugins/Experimental/EditorDataStorage/Source/TedsUI/Public/Widgets/NameWidget.h b/Engine/Plugins/Experimental/EditorDataStorage/Source/TedsUI/Public/Widgets/NameWidget.h
new file mode 100644
--- /dev/null
+++ b/Engine/Plugins/Experimental/EditorDataStorage/Source/TedsUI/Public/Widgets/NameWidget.h
@@ -0,0 +1,42 @@
+// Copyright Epic Games, Inc. All Rights Reserved.
+
+#pragma once
+
+#include "Elements/Interfaces/TypedElementDataStorageFactory.h"
+#include "Elements/Interfaces/TypedElementDataStorageUiInterface.h"
+#include "Internationalization/Text.h"
+#include "UObject/ObjectMacros.h"
+
+#include "NameWidget.generated.h"
+
+class IEditorDataStorageProvider;
+class UScriptStruct;
+
+UCLASS()
+class UNameWidgetFactory : public UEditorDataStorageFactory
+{
+	GENERATED_BODY()
+
+public:
+	virtual ~UNameWidgetFactory() override = default;
+
+	TEDSUI_API virtual void RegisterWidgetConstructors(UE::Editor::DataStorage::ICoreProvider& DataStorage,
+		UE::Editor::DataStorage::IUiProvider& DataStorageUi) const override;
+};
+
+// Label for assets in TEDS
+USTRUCT()
+struct FNameWidgetConstructor : public FSimpleWidgetConstructor
+{
+	GENERATED_BODY()
+
+public:
+	TEDSUI_API FNameWidgetConstructor();
+
+	TEDSUI_API virtual TSharedPtr<SWidget> CreateWidget(
+		UE::Editor::DataStorage::ICoreProvider* DataStorage,
+		UE::Editor::DataStorage::IUiProvider* DataStorageUi,
+		UE::Editor::DataStorage::RowHandle TargetRow,
+		UE::Editor::DataStorage::RowHandle WidgetRow, 
+		const UE::Editor::DataStorage::FMetaDataView& Arguments) override;
+};
diff --git a/Engine/Plugins/Experimental/EditorDataStorageFeatures/Source/TedsAssetData/Private/TedsAssetDataHelper.cpp b/Engine/Plugins/Experimental/EditorDataStorageFeatures/Source/TedsAssetData/Private/TedsAssetDataHelper.cpp
--- a/Engine/Plugins/Experimental/EditorDataStorageFeatures/Source/TedsAssetData/Private/TedsAssetDataHelper.cpp
+++ b/Engine/Plugins/Experimental/EditorDataStorageFeatures/Source/TedsAssetData/Private/TedsAssetDataHelper.cpp
@@ -4,6 +4,28 @@
 
 namespace TedsAssetDataHelper
 {
+	namespace MetaDataNames
+	{
+		static const FLazyName ThumbnailStatus("bAllowAssetStatusThumbnailOverlay");
+		static const FLazyName ThumbnailFadeIn("bAllowFadeIn");
+		static const FLazyName ThumbnailHintText("bAllowHintText");
+		static const FLazyName ThumbnailRealTimeOnHovered("bAllowRealTimeOnHovered");
+		static const FLazyName ThumbnailSizeOffset("ThumbnailSizeOffset");
+
+		FName GetThumbnailStatusMetaDataName() { return ThumbnailStatus; }
+		FName GetThumbnailFadeInMetaDataName() { return ThumbnailFadeIn; }
+		FName GetThumbnailHintTextMetaDataName() { return ThumbnailHintText; }
+		FName GetThumbnailRealTimeOnHoveredMetaDataName() { return ThumbnailRealTimeOnHovered; }
+		FName GetThumbnailSizeOffsetMetaDataName() { return ThumbnailSizeOffset; }
+	}
+
+	namespace TableView
+	{
+		static FLazyName TableViewerWidgetTableName("Editor_TableViewerWidgetTable");
+
+		FName GetWidgetTableName() { return TableViewerWidgetTableName; }
+	}
+
 	FString RemoveSlashFromStart(const FString& InStringToModify)
 	{
 		FString StringToReturn = InStringToModify;
diff --git a/Engine/Plugins/Experimental/EditorDataStorageFeatures/Source/TedsAssetData/Private/TedsAssetDataHelper.h b/Engine/Plugins/Experimental/EditorDataStorageFeatures/Source/TedsAssetData/Private/TedsAssetDataHelper.h
deleted file mode 100644
diff --git a/Engine/Plugins/Experimental/EditorDataStorageFeatures/Source/TedsAssetData/Private/Widgets/AssetDataItemTypeWidget.cpp b/Engine/Plugins/Experimental/EditorDataStorageFeatures/Source/TedsAssetData/Private/Widgets/AssetDataItemTypeWidget.cpp
--- a/Engine/Plugins/Experimental/EditorDataStorageFeatures/Source/TedsAssetData/Private/Widgets/AssetDataItemTypeWidget.cpp
+++ b/Engine/Plugins/Experimental/EditorDataStorageFeatures/Source/TedsAssetData/Private/Widgets/AssetDataItemTypeWidget.cpp
@@ -2,9 +2,13 @@
 
 #include "Widgets/AssetDataItemTypeWidget.h"
 
+#include "AssetDefinition.h"
+#include "AssetDefinitionRegistry.h"
 #include "Elements/Framework/TypedElementAttributeBinding.h"
 #include "Internationalization/Text.h"
 #include "TedsAssetDataColumns.h"
+#include "TedsAssetDataWidgetColumns.h"
+#include "Columns/SlateDelegateColumns.h"
 #include "Elements/Columns/TypedElementFolderColumns.h"
 #include "UObject/ObjectMacros.h"
 #include "Widgets/Layout/SBox.h"
@@ -31,33 +35,37 @@ TSharedPtr<SWidget> FAssetDataItemTypeWidgetConstructor::CreateWidget(UE::Editor
 	const UE::Editor::DataStorage::FMetaDataView& Arguments)
 {
 	UE::Editor::DataStorage::FAttributeBinder Binder(TargetRow, DataStorage);
+	UE::Editor::DataStorage::FAttributeBinder WidgetBinder(WidgetRow, DataStorage);
 
-	static const FMargin ColumnItemPadding(5, 0, 5, 0);
-
-	bool bIsFolder = DataStorage->HasColumns<FFolderTag>(TargetRow);
-	TAttribute<FText> ItemTypeAttribute;
-	if (bIsFolder)
-	{
-		ItemTypeAttribute = LOCTEXT("AssetClassFolder", "Folder");
-	}
-	else
-	{
-		ItemTypeAttribute = Binder.BindData(&FAssetClassColumn::ClassPath, [](FTopLevelAssetPath InClassPath)
+	const bool bIsFolder = DataStorage->HasColumns<FFolderTag>(TargetRow);
+
+	return SNew(STextBlock)
+			.Font(WidgetBinder.BindData(&FFontStyleColumn_Experimental::FontInfo, FAppStyle::GetFontStyle("NormalFont")))
+			.Visibility(WidgetBinder.BindData(&FWidgetVisibilityColumn_Experimental::Visibility, EVisibility::Visible))
+			.OverflowPolicy(WidgetBinder.BindData(&FTextOverflowPolicyColumn_Experimental::OverflowPolicy, ETextOverflowPolicy::Ellipsis).Get())
+			.ColorAndOpacity(WidgetBinder.BindData(&FOnGetWidgetColorAndOpacityColumn_Experimental::OnGetWidgetColorAndOpacity, [] (FOnGetWidgetColorAndOpacity InOnGetWidgetColorAndOpacity)
+			{
+				return InOnGetWidgetColorAndOpacity.IsBound() ? InOnGetWidgetColorAndOpacity.Execute() : FSlateColor::UseForeground();
+			}))
+			.Text(Binder.BindData(&FAssetClassColumn::ClassPath, [bIsFolder](FTopLevelAssetPath InClassPath)
 			{
-				const FName AssetClass = InClassPath.GetAssetName();
-				if (!AssetClass.IsNone())
+				if (bIsFolder)
+				{
+					return LOCTEXT("AssetClassFolder", "Folder");
+				}
+
+				if (const UAssetDefinitionRegistry* AssetDefinitionRegistry = UAssetDefinitionRegistry::Get())
 				{
-					return FText::FromName(InClassPath.GetAssetName());
+					if (const UClass* FoundClass = FindObject<UClass>(InClassPath))
+					{
+						if (const UAssetDefinition* AssetDefinition = AssetDefinitionRegistry->GetAssetDefinitionForClass(FoundClass))
+						{
+							return AssetDefinition->GetAssetDisplayName();
+						}
+					}
 				}
 				return LOCTEXT("AssetClass_Invalid", "Invalid or not found");
-			});
-	}
-	return SNew(SBox)
-			.Padding(ColumnItemPadding)
-			[
-				SNew(STextBlock)
-				.Text(ItemTypeAttribute)
-			];
+			}));
 }
 
 #undef LOCTEXT_NAMESPACE
diff --git a/Engine/Plugins/Experimental/EditorDataStorageFeatures/Source/TedsAssetData/Private/Widgets/AssetDataThumbnailWidget.cpp b/Engine/Plugins/Experimental/EditorDataStorageFeatures/Source/TedsAssetData/Private/Widgets/AssetDataThumbnailWidget.cpp
--- a/Engine/Plugins/Experimental/EditorDataStorageFeatures/Source/TedsAssetData/Private/Widgets/AssetDataThumbnailWidget.cpp
+++ b/Engine/Plugins/Experimental/EditorDataStorageFeatures/Source/TedsAssetData/Private/Widgets/AssetDataThumbnailWidget.cpp
@@ -6,14 +6,242 @@
 #include "AssetViewTypes.h"
 #include "Elements/Framework/TypedElementAttributeBinding.h"
 #include "TedsAssetDataColumns.h"
+#include "TedsAssetDataHelper.h"
+#include "TedsAssetDataWidgetColumns.h"
+#include "Columns/SlateDelegateColumns.h"
 #include "Elements/Columns/TypedElementFolderColumns.h"
+#include "Elements/Columns/TypedElementHiearchyColumns.h"
+#include "Elements/Columns/TypedElementMiscColumns.h"
 #include "Elements/Columns/TypedElementSlateWidgetColumns.h"
+#include "Interfaces/IPluginManager.h"
+#include "Styling/StyleColors.h"
 #include "ThumbnailRendering/ThumbnailManager.h"
 #include "Widgets/Images/SImage.h"
+#include "Widgets/Layout/SBorder.h"
 #include "Widgets/Layout/SBox.h"
+#include "Widgets/Layout/SSeparator.h"
+#include "Widgets/SBoxPanel.h"
+#include "Widgets/SToolTip.h"
+#include "Widgets/Text/STextBlock.h"
 
 #define LOCTEXT_NAMESPACE "FAssetThumbnailWidgetConstructor"
 
+namespace Private
+{
+	bool IsTopLevelFolder(const FStringView InFolderPath)
+	{
+		int32 SlashCount = 0;
+		for (const TCHAR PathChar : InFolderPath)
+		{
+			if (PathChar == TEXT('/'))
+			{
+				if (++SlashCount > 1)
+				{
+					break;
+				}
+			}
+		}
+
+		return SlashCount == 1;
+	}
+
+	void AddToToolTipInfoBox(const TSharedRef<SVerticalBox>& InfoBox, const FText& Key, const FText& Value)
+	{
+		InfoBox->AddSlot()
+        	.Padding(0.f, 0.f, 0.f, 6.f)
+        	.AutoHeight()
+        	[
+        		SNew(SHorizontalBox)
+        		+SHorizontalBox::Slot()
+        		.AutoWidth()
+        		.Padding(0, 0, 4, 0)
+        		[
+        			SNew(STextBlock)
+        			.Font(FAppStyle::GetFontStyle("ContentBrowser.Tooltip.EntryFont"))
+        			.Text(FText::Format(NSLOCTEXT("AssetThumbnailToolTip", "AssetViewTooltipFormat", "{0}:"), Key))
+        		]
+        
+        		+SHorizontalBox::Slot()
+        		.AutoWidth()
+        		[
+        			SNew(STextBlock)
+        			.Font(FAppStyle::GetFontStyle("ContentBrowser.Tooltip.EntryFont"))
+        			.ColorAndOpacity(FStyleColors::White)
+        			.Text(Value)
+        		]
+        	];
+	}
+
+	TSharedRef<SWidget> GetFolderTooltip(UE::Editor::DataStorage::ICoreProvider* InDataStorage, UE::Editor::DataStorage::RowHandle InTargetRow, TAttribute<const FSlateBrush*> InFolderImage)
+	{
+		UE::Editor::DataStorage::FAttributeBinder Binder(InTargetRow, InDataStorage);
+
+		// Create a box to hold every line of info in the body of the tooltip
+		TSharedRef<SVerticalBox> InfoBox = SNew(SVerticalBox);
+
+		if (const FAssetPathColumn_Experimental* AssetPathColumn = InDataStorage->GetColumn<FAssetPathColumn_Experimental>(InTargetRow))
+		{
+			const FText FolderPath = FText::FromName(AssetPathColumn->Path);
+			AddToToolTipInfoBox(InfoBox, LOCTEXT("TileViewTooltipPath", "Path"), FolderPath);
+		}
+
+		if (const FAssetPathColumn_Experimental* InternalPathColumn = InDataStorage->GetColumn<FAssetPathColumn_Experimental>(InTargetRow))
+		{
+			const FName InternalPath = InternalPathColumn->Path;
+			if (!InternalPath.IsNone())
+			{
+				FNameBuilder FolderPathBuilder(InternalPath);
+				if (IsTopLevelFolder(FStringView(FolderPathBuilder)))
+				{
+					FStringView PluginName(FolderPathBuilder);
+					PluginName.RightChopInline(1);
+
+					if (TSharedPtr<IPlugin> Plugin = IPluginManager::Get().FindPlugin(PluginName))
+					{
+						if (Plugin->GetDescriptor().Description.Len() > 0)
+						{
+							AddToToolTipInfoBox(InfoBox, LOCTEXT("TileViewTooltipPluginDescription", "Plugin Description"), FText::FromString(Plugin->GetDescriptor().Description));
+						}
+					}
+				}
+			}
+		}
+
+		return SNew(SVerticalBox)
+
+				+ SVerticalBox::Slot()
+				.AutoHeight()
+				[
+					SNew(SVerticalBox)
+
+					+ SVerticalBox::Slot()
+					.AutoHeight()
+					.Padding(0.f, 0.f, 0.f, 6.f)
+					[
+						SNew(SHorizontalBox)
+
+						+ SHorizontalBox::Slot()
+						.AutoWidth()
+						.VAlign(VAlign_Center)
+						[
+							SNew(STextBlock)
+							.Text(Binder.BindData(&FNameColumn::Name, [] (FName InName)
+							{
+								return FText::FromString(TedsAssetDataHelper::RemoveSlashFromStart(InName.ToString()));
+							}))
+							.ColorAndOpacity(FStyleColors::White)
+							.Font(FAppStyle::GetFontStyle("ContentBrowser.Tooltip.EntryFont"))
+						]
+					]
+
+					+ SVerticalBox::Slot()
+					.AutoHeight()
+					[
+						SNew(SHorizontalBox)
+
+						+ SHorizontalBox::Slot()
+						.AutoWidth()
+						.VAlign(VAlign_Center)
+						.Padding(0.f, 0.f, 4.f, 0.f)
+						[
+							SNew(SBox)
+							.WidthOverride(16.f)
+							.HeightOverride(16.f)
+							[
+								SNew(SImage)
+								.ColorAndOpacity(Binder.BindData(&FSlateColorColumn::Color))
+								.Image(InFolderImage)
+							]
+						]
+
+						+ SHorizontalBox::Slot()
+						.AutoWidth()
+						.VAlign(VAlign_Center)
+						[
+							SNew(STextBlock)
+							.Font(FAppStyle::GetFontStyle("ContentBrowser.Tooltip.EntryFont"))
+							.Text(LOCTEXT("FolderNameBracketedLabel", "Folder"))
+						]
+					]
+				]
+
+				+ SVerticalBox::Slot()
+				.Padding(FMargin(0.f,6.f))
+				.AutoHeight()
+				[
+					SNew(SSeparator)
+					.Orientation(Orient_Horizontal)
+					.Thickness(1.f)
+					.ColorAndOpacity(COLOR("#484848FF"))
+					.SeparatorImage(FAppStyle::Get().GetBrush("WhiteBrush"))
+				]
+
+				+ SVerticalBox::Slot()
+				.AutoHeight()
+				[
+					InfoBox
+				];
+	}
+
+	class SFolderItemToolTip : public SToolTip
+	{
+	public:
+		SLATE_BEGIN_ARGS(SFolderItemToolTip)
+		{ }
+
+			SLATE_ARGUMENT(UE::Editor::DataStorage::RowHandle, TargetRow)
+			SLATE_ARGUMENT(UE::Editor::DataStorage::RowHandle, WidgetRow)
+			SLATE_ARGUMENT(UE::Editor::DataStorage::ICoreProvider*, DataStorage)
+			SLATE_ATTRIBUTE(const FSlateBrush*, FolderImage)
+
+		SLATE_END_ARGS()
+
+		void Construct(const FArguments& InArgs)
+		{
+			DataStorage = InArgs._DataStorage;
+			TargetRow = InArgs._TargetRow;
+			WidgetRow = InArgs._WidgetRow;
+			FolderImage = InArgs._FolderImage;
+
+			SToolTip::Construct(
+				SToolTip::FArguments()
+				.TextMargin(FMargin(12.f, 8.f, 12.f, 2.f))
+				.BorderImage(FAppStyle::GetBrush("AssetThumbnail.Tooltip.Border")));
+		}
+
+		// IToolTip interface
+		virtual bool IsEmpty() const override
+		{
+			return	TargetRow == UE::Editor::DataStorage::InvalidRowHandle ||
+					WidgetRow == UE::Editor::DataStorage::InvalidRowHandle ||
+					DataStorage == nullptr ||
+					!FolderImage.IsSet();
+		}
+
+		virtual void OnOpening() override
+		{
+			SetContentWidget(Private::GetFolderTooltip(DataStorage, TargetRow, FolderImage));
+
+			// When opening the tooltip update it on the Row as well
+			if (FLocalWidgetTooltipColumn_Experimental* WidgetTooltipColumn = DataStorage->GetColumn<FLocalWidgetTooltipColumn_Experimental>(WidgetRow))
+			{
+				WidgetTooltipColumn->Tooltip = SharedThis(this);
+			}
+		}
+
+		virtual void OnClosed() override
+		{
+			ResetContentWidget();
+		}
+
+	private:
+		UE::Editor::DataStorage::RowHandle TargetRow = UE::Editor::DataStorage::InvalidRowHandle;
+		UE::Editor::DataStorage::RowHandle WidgetRow = UE::Editor::DataStorage::InvalidRowHandle;
+		UE::Editor::DataStorage::ICoreProvider* DataStorage = nullptr;
+		TAttribute<const FSlateBrush*> FolderImage;
+	};
+}
+
 void UAssetThumbnailWidgetFactory::RegisterWidgetConstructors(UE::Editor::DataStorage::ICoreProvider& DataStorage,
 	UE::Editor::DataStorage::IUiProvider& DataStorageUi) const
 {
@@ -39,18 +267,37 @@ TSharedPtr<SWidget> FAssetThumbnailWidgetConstructor::CreateWidget(UE::Editor::D
 	const UE::Editor::DataStorage::FMetaDataView& Arguments)
 {
 	UE::Editor::DataStorage::FAttributeBinder Binder(TargetRow, DataStorage);
+	UE::Editor::DataStorage::FAttributeBinder WidgetRowBinder(WidgetRow, DataStorage);
 
 	bool bIsAsset = DataStorage->HasColumns<FAssetTag>(TargetRow);
+	UE::Editor::DataStorage::RowHandle ParentWidgetRowHandle = UE::Editor::DataStorage::InvalidRowHandle;
+	if (FTableRowParentColumn* ParentWidgetRow = DataStorage->GetColumn<FTableRowParentColumn>(WidgetRow))
+	{
+		ParentWidgetRowHandle = ParentWidgetRow->Parent;
+	}
+	UE::Editor::DataStorage::FAttributeBinder ParentWidgetBinder(ParentWidgetRowHandle, DataStorage);
 
-	constexpr uint32 ThumbnailSize = 64;
-	constexpr float ThumbnailSizeAsFloat = static_cast<float>(ThumbnailSize);
-	constexpr float ThumbnailVerticalPadding = 4.f;
+	float ThumbnailSizeOffset = 0.f;
+	UE::Editor::DataStorage::FMetaDataEntryView ThumbnailSizeOffsetMetaData = Arguments.FindGeneric(TedsAssetDataHelper::MetaDataNames::GetThumbnailSizeOffsetMetaDataName());
+	if (ThumbnailSizeOffsetMetaData.IsSet() && ThumbnailSizeOffsetMetaData.IsType<double>())
+	{
+		ThumbnailSizeOffset = (float)(*ThumbnailSizeOffsetMetaData.TryGetExact<double>());
+	}
+	constexpr int32 DefaultThumbnailSize = 64;
 
 	// Thumbnail Box container
 	const TSharedRef<SBox> ThumbnailBox = SNew(SBox)
-		.WidthOverride(ThumbnailSizeAsFloat)
-		.HeightOverride(ThumbnailSizeAsFloat + (ThumbnailVerticalPadding * 2.f))
-		.Padding(0.f, ThumbnailVerticalPadding);
+		.Padding(WidgetRowBinder.BindData(&FWidgetPaddingColumn_Experimental::Padding, FMargin(0.f)))
+		.WidthOverride(ParentWidgetBinder.BindData(&FSizeValueColumn_Experimental::SizeValue, [DefaultThumbnailSize, ThumbnailSizeOffset] (float InSizeValue)
+		{
+			const float Size = FMath::IsNearlyZero(InSizeValue) ? DefaultThumbnailSize : InSizeValue;
+			return FOptionalSize(Size + ThumbnailSizeOffset);
+		}))
+		.HeightOverride(ParentWidgetBinder.BindData(&FSizeValueColumn_Experimental::SizeValue, [DefaultThumbnailSize, ThumbnailSizeOffset] (float InSizeValue)
+		{
+			const float Size = FMath::IsNearlyZero(InSizeValue) ? DefaultThumbnailSize : InSizeValue;
+			return FOptionalSize(Size + ThumbnailSizeOffset);
+		}));
 
 	// For assets, grab the color from the asset definition
 	if (bIsAsset)
@@ -64,33 +311,78 @@ TSharedPtr<SWidget> FAssetThumbnailWidgetConstructor::CreateWidget(UE::Editor::D
 
 		// Thumbnail Configuration Args
 		FAssetThumbnailConfig ThumbnailConfig;
-		ThumbnailConfig.bAllowFadeIn = true;
-		ThumbnailConfig.bAllowHintText = false;
-		ThumbnailConfig.bAllowRealTimeOnHovered = false; // we use our own OnMouseEnter/Leave for logical asset item
+		UE::Editor::DataStorage::FMetaDataEntryView bAllowFadeInMeta = Arguments.FindGeneric(TedsAssetDataHelper::MetaDataNames::GetThumbnailFadeInMetaDataName());
+		if (bAllowFadeInMeta.IsSet() && bAllowFadeInMeta.IsType<bool>())
+		{
+			ThumbnailConfig.bAllowFadeIn = *bAllowFadeInMeta.TryGetExact<bool>();
+		}
+
+		UE::Editor::DataStorage::FMetaDataEntryView bAllowHintTextMeta = Arguments.FindGeneric(TedsAssetDataHelper::MetaDataNames::GetThumbnailHintTextMetaDataName());
+		if (bAllowHintTextMeta.IsSet() && bAllowHintTextMeta.IsType<bool>())
+		{
+			ThumbnailConfig.bAllowHintText = *bAllowHintTextMeta.TryGetExact<bool>();
+		}
+
+		UE::Editor::DataStorage::FMetaDataEntryView bAllowRealTimeOnHoveredMeta = Arguments.FindGeneric(TedsAssetDataHelper::MetaDataNames::GetThumbnailRealTimeOnHoveredMetaDataName());
+		if (bAllowRealTimeOnHoveredMeta.IsSet() && bAllowRealTimeOnHoveredMeta.IsType<bool>())
+		{
+			ThumbnailConfig.bAllowRealTimeOnHovered = *bAllowRealTimeOnHoveredMeta.TryGetExact<bool>();
+		}
+
+		if (DataStorage->HasColumns<FOnGetWidgetSlateBrushColumn_Experimental>(WidgetRow))
+		{
+			ThumbnailConfig.AssetChipBorderImageOverride = WidgetRowBinder.BindData(&FOnGetWidgetSlateBrushColumn_Experimental::OnGetWidgetSlateBrush, [] (FOnGetWidgetSlateBrush InOnGetWidgetSlateBrush)
+			{
+				const FSlateBrush* ReturnBrush = InOnGetWidgetSlateBrush.IsBound() ? InOnGetWidgetSlateBrush.Execute() : FAppStyle::GetNoBrush();
+				return ReturnBrush;
+			});
+		}
 		// TODO: AssetItem->GetItem().GetItemTemporaryReason() == EContentBrowserItemFlags::Temporary_Creation;
 		// This was the previous check to assign bForceGenericThumbnail, for now it will use the Generic if the AssetData is not valid
 		ThumbnailConfig.bForceGenericThumbnail = !AssetDataToUse.IsValid();
 		ThumbnailConfig.bAllowAssetSpecificThumbnailOverlay = !ThumbnailConfig.bForceGenericThumbnail;
 		ThumbnailConfig.ThumbnailLabel = EThumbnailLabel::ClassName;
-		// TODO: If Teds View items will be bound to the CB size and can change their size, this would need to be bound to that as well
-		// ThumbnailConfig.GenericThumbnailSize = 64;
+
+		if (DataStorage->HasColumns<FSizeValueColumn_Experimental>(WidgetRow))
+		{
+			ThumbnailConfig.GenericThumbnailSize = ParentWidgetBinder.BindData(&FSizeValueColumn_Experimental::SizeValue, [DefaultThumbnailSize] (float InThumbnailSize)
+			{
+				return FMath::IsNearlyZero(InThumbnailSize) ? DefaultThumbnailSize : static_cast<int32>(InThumbnailSize);
+			});
+		}
+
 		// TODO: This data will need to be integrated in TEDS and later on remove this/update it to use TEDS instead to retrieve the data needed
 		// ThumbnailConfig.AssetSystemInfoProvider = nullptr;
+		UE::Editor::DataStorage::FMetaDataEntryView bAllowAssetStatusThumbnailOverlayMeta = Arguments.FindGeneric(TedsAssetDataHelper::MetaDataNames::GetThumbnailStatusMetaDataName());
+		if (bAllowAssetStatusThumbnailOverlayMeta.IsSet() && bAllowAssetStatusThumbnailOverlayMeta.IsType<bool>())
+		{
+			ThumbnailConfig.bAllowAssetStatusThumbnailOverlay = *bAllowAssetStatusThumbnailOverlayMeta.TryGetExact<bool>();
+		}
 		ThumbnailConfig.bShowAssetChip = true;
 		// TODO: If we are able to change the thumbnail in Teds view we would need to bind this to make it work
 		// ThumbnailConfig.bCanDisplayEditModePrimitiveTools = false;
-		// ThumbnailConfig.IsEditModeVisible = false;
+		ThumbnailConfig.IsEditModeVisible = ParentWidgetBinder.BindData(&FThumbnailEditModeColumn_Experimental::IsEditModeToggled, [] (bool IsEditModeToggled)
+		{
+			return IsEditModeToggled ? EVisibility::Visible : EVisibility::Collapsed;
+		});
 
 		// TODO: Consider caching them instead, to avoid creating them everytime
-		const TSharedRef<FAssetThumbnail> AssetThumbnail = MakeShared<FAssetThumbnail>(AssetDataToUse, ThumbnailSize, ThumbnailSize, UThumbnailManager::Get().GetSharedThumbnailPool());
+		constexpr uint32 ThumbnailResolution = 256;
+		const TSharedRef<FAssetThumbnail> AssetThumbnail = MakeShared<FAssetThumbnail>(AssetDataToUse, ThumbnailResolution, ThumbnailResolution, UThumbnailManager::Get().GetSharedThumbnailPool());
 
-		ThumbnailBox->SetContent(AssetThumbnail->MakeThumbnailWidget(ThumbnailConfig));
+		TSharedRef<SWidget> ThumbnailWidget = AssetThumbnail->MakeThumbnailWidget(ThumbnailConfig);
+		ThumbnailBox->SetContent(ThumbnailWidget);
+
+		// If the widget creating this is interested in the ThumbnailTooltip it has to add the column itself to avoid adding them if not used
+		if (FLocalWidgetTooltipColumn_Experimental* WidgetTooltipColumn = DataStorage->GetColumn<FLocalWidgetTooltipColumn_Experimental>(WidgetRow))
+		{
+			WidgetTooltipColumn->Tooltip = ThumbnailWidget->GetToolTip();
+		}
 	}
 	// For folders, use the color and folder type column directly
 	else
 	{
 		TAttribute<const FSlateBrush*> FolderImage;
-		TAttribute<FSlateColor> ColorAndOpacityAttribute;
 
 		FolderImage = Binder.BindData(&FFolderTypeColumn_Experimental::FolderType, [] (EFolderType InFolderType)
 		{
@@ -116,13 +408,55 @@ TSharedPtr<SWidget> FAssetThumbnailWidgetConstructor::CreateWidget(UE::Editor::D
 			}
 		});
 
-		ColorAndOpacityAttribute = Binder.BindData(&FSlateColorColumn::Color, FSlateColor::UseForeground());
-		ThumbnailBox->SetContent(SNew(SImage)
-						.Image(FolderImage)
-						.ColorAndOpacity(ColorAndOpacityAttribute));
+		TAttribute<const FSlateBrush*> ShadowFolderImage;
+
+		ShadowFolderImage = Binder.BindData(&FFolderTypeColumn_Experimental::FolderType, [] (EFolderType InFolderType)
+		{
+			// Default values
+			switch (InFolderType)
+			{
+				// TODO: Cpp and Virtual are not currently checked, see TedsAssetData PopulatePathDataTableRow for more info
+				case EFolderType::CustomVirtual:
+					return FAppStyle::GetBrush(FName(TEXT("ContentBrowser.ListViewVirtualFolderShadow")));
+
+				case EFolderType::Developer:
+				case EFolderType::PluginRoot:
+				case EFolderType::Code:
+				case EFolderType::Normal:
+				default:
+					return FAppStyle::GetBrush(FName(TEXT("ContentBrowser.FolderItem.DropShadow")));
+			}
+		});
+
+		ThumbnailBox->SetContent(
+			SNew(SBorder)
+			.BorderImage(ShadowFolderImage)
+			.Padding(FMargin(0,0,2.0f,2.0f))
+			[
+				SNew(SImage)
+				.Image(FolderImage)
+				.ColorAndOpacity(Binder.BindData(&FSlateColorColumn::Color, FSlateColor::UseForeground()))
+			]);
+
+		ThumbnailBox->SetToolTip(SNew(Private::SFolderItemToolTip)
+				.TargetRow(TargetRow)
+				.WidgetRow(WidgetRow)
+				.DataStorage(DataStorage)
+				.FolderImage(FolderImage));
 	}
 
 	return ThumbnailBox;
 }
 
+TConstArrayView<const UScriptStruct*> FAssetThumbnailWidgetConstructor::GetAdditionalColumnsList() const
+{
+	using namespace UE::Editor::DataStorage;
+
+	static const TTypedElementColumnTypeList<
+		FSizeValueColumn_Experimental,
+		FThumbnailEditModeColumn_Experimental> Columns;
+
+	return Columns;
+}
+
 #undef LOCTEXT_NAMESPACE
diff --git a/Engine/Plugins/Experimental/EditorDataStorageFeatures/Source/TedsAssetData/Private/Widgets/AssetDataThumbnailWidget.h b/Engine/Plugins/Experimental/EditorDataStorageFeatures/Source/TedsAssetData/Private/Widgets/AssetDataThumbnailWidget.h
--- a/Engine/Plugins/Experimental/EditorDataStorageFeatures/Source/TedsAssetData/Private/Widgets/AssetDataThumbnailWidget.h
+++ b/Engine/Plugins/Experimental/EditorDataStorageFeatures/Source/TedsAssetData/Private/Widgets/AssetDataThumbnailWidget.h
@@ -40,4 +40,7 @@ public:
 		UE::Editor::DataStorage::RowHandle TargetRow,
 		UE::Editor::DataStorage::RowHandle WidgetRow, 
 		const UE::Editor::DataStorage::FMetaDataView& Arguments) override;
+
+protected:
+	virtual TConstArrayView<const UScriptStruct*> GetAdditionalColumnsList() const override;
 };
diff --git a/Engine/Plugins/Experimental/EditorDataStorageFeatures/Source/TedsAssetData/Private/Widgets/AssetLabelWidget.cpp b/Engine/Plugins/Experimental/EditorDataStorageFeatures/Source/TedsAssetData/Private/Widgets/AssetLabelWidget.cpp
--- a/Engine/Plugins/Experimental/EditorDataStorageFeatures/Source/TedsAssetData/Private/Widgets/AssetLabelWidget.cpp
+++ b/Engine/Plugins/Experimental/EditorDataStorageFeatures/Source/TedsAssetData/Private/Widgets/AssetLabelWidget.cpp
@@ -15,7 +15,7 @@ void UAssetLabelWidgetFactory::RegisterWidgetConstructors(UE::Editor::DataStorag
 {
 	using namespace UE::Editor::DataStorage::Queries;
 	DataStorageUi.RegisterWidgetFactory<FAssetLabelWidgetConstructor>(
-		DataStorageUi.FindPurpose(DataStorageUi.GetGeneralWidgetPurposeID()),
+		DataStorageUi.FindPurpose(IUiProvider::FPurposeInfo("ContentBrowser", "RowLabel", NAME_None).GeneratePurposeID()),
 		TColumn<FNameColumn>());
 }
 
diff --git a/Engine/Plugins/Experimental/EditorDataStorageFeatures/Source/TedsAssetData/Private/Widgets/ContentBrowserListViewNameWidget.cpp b/Engine/Plugins/Experimental/EditorDataStorageFeatures/Source/TedsAssetData/Private/Widgets/ContentBrowserListViewNameWidget.cpp
--- a/Engine/Plugins/Experimental/EditorDataStorageFeatures/Source/TedsAssetData/Private/Widgets/ContentBrowserListViewNameWidget.cpp
+++ b/Engine/Plugins/Experimental/EditorDataStorageFeatures/Source/TedsAssetData/Private/Widgets/ContentBrowserListViewNameWidget.cpp
@@ -6,7 +6,11 @@
 #include "AssetDefinition.h"
 #include "Elements/Framework/TypedElementAttributeBinding.h"
 #include "TedsAssetDataColumns.h"
+#include "TedsAssetDataHelper.h"
+#include "TedsAssetDataWidgetColumns.h"
+#include "Columns/SlateDelegateColumns.h"
 #include "Elements/Columns/TypedElementFolderColumns.h"
+#include "Elements/Columns/TypedElementHiearchyColumns.h"
 #include "Elements/Columns/TypedElementMiscColumns.h"
 #include "Widgets/Images/SImage.h"
 #include "Widgets/Layout/SBox.h"
@@ -57,6 +61,12 @@ TSharedPtr<SWidget> FContentBrowserListViewNameWidgetConstructor::CreateWidget(U
 
 	if (DataStorage->IsRowAvailable(TargetRow))
 	{
+		UE::Editor::DataStorage::RowHandle ParentWidgetRowHandle = UE::Editor::DataStorage::InvalidRowHandle;
+		if (FTableRowParentColumn* ParentWidgetRow = DataStorage->GetColumn<FTableRowParentColumn>(WidgetRow))
+		{
+			ParentWidgetRowHandle = ParentWidgetRow->Parent;
+		}
+
 		TSharedPtr<FTypedElementWidgetConstructor> OutThumbnailWidgetConstructorPtr;
 		TSharedPtr<FTypedElementWidgetConstructor> OutLabelWidgetConstructorPtr;
 
@@ -74,27 +84,60 @@ TSharedPtr<SWidget> FContentBrowserListViewNameWidgetConstructor::CreateWidget(U
 
 		// TODO: Arguments need to be developed further to customize the Thumbnail Config if needed
 		TArray<TWeakObjectPtr<const UScriptStruct>> ThumbnailColumns = GetThumbnailColumns();
-		DataStorageUi->CreateWidgetConstructors(DataStorageUi->FindPurpose(DataStorageUi->GetGeneralWidgetPurposeID()),
-			IUiProvider::EMatchApproach::ExactMatch, ThumbnailColumns, Arguments, AssignThumbnailWidgetToColumn);
+		UE::Editor::DataStorage::FMetaData ThumbnailMeta;
+		ThumbnailMeta.AddOrSetMutableData(TedsAssetDataHelper::MetaDataNames::GetThumbnailStatusMetaDataName(), true);
+		ThumbnailMeta.AddOrSetMutableData(TedsAssetDataHelper::MetaDataNames::GetThumbnailFadeInMetaDataName(), true);
+		ThumbnailMeta.AddOrSetMutableData(TedsAssetDataHelper::MetaDataNames::GetThumbnailHintTextMetaDataName(), false);
+		// TODO: use our own OnMouseEnter/Leave for setting the realtime flag, also check CB Setting
+		ThumbnailMeta.AddOrSetMutableData(TedsAssetDataHelper::MetaDataNames::GetThumbnailRealTimeOnHoveredMetaDataName(), false);
+		DataStorageUi->CreateWidgetConstructors(DataStorageUi->FindPurpose(DataStorageUi->GetGeneralWidgetPurposeID()), IUiProvider::EMatchApproach::ExactMatch, ThumbnailColumns, Arguments, AssignThumbnailWidgetToColumn);
 
 		TArray<TWeakObjectPtr<const UScriptStruct>> LabelColumns = GetLabelColumns();
-		DataStorageUi->CreateWidgetConstructors(DataStorageUi->FindPurpose(DataStorageUi->GetGeneralWidgetPurposeID()),
+
+		const UE::Editor::DataStorage::IUiProvider::FPurposeID ContentBrowserNameWidgetPurposeId = IUiProvider::FPurposeInfo("ContentBrowser", "RowLabel", NAME_None).GeneratePurposeID();
+		DataStorageUi->CreateWidgetConstructors(DataStorageUi->FindPurpose(ContentBrowserNameWidgetPurposeId),
 			IUiProvider::EMatchApproach::ExactMatch, LabelColumns, Arguments, AssignLabelWidgetToColumn);
 
-		static const FMargin ColumnItemPadding(8, 0);
 		constexpr float ThumbnailNameHorizontalPadding = 8.f;
 
 		TSharedPtr<SWidget> ThumbnailWidget = SNullWidget::NullWidget;
 		TSharedPtr<SWidget> LabelWidget = SNullWidget::NullWidget;
 
+		UE::Editor::DataStorage::RowHandle ThumbnailWidgetRowHandle = UE::Editor::DataStorage::InvalidRowHandle;
 		if (OutThumbnailWidgetConstructorPtr)
 		{
-			ThumbnailWidget = DataStorageUi->ConstructWidget(WidgetRow, *OutThumbnailWidgetConstructorPtr, Arguments);
+			if (ThumbnailWidgetRowHandle = DataStorage->AddRow(DataStorage->FindTable(TedsAssetDataHelper::TableView::GetWidgetTableName())); ThumbnailWidgetRowHandle != UE::Editor::DataStorage::InvalidRowHandle)
+			{
+				// Referenced Data Row
+				DataStorage->AddColumn(ThumbnailWidgetRowHandle, FTypedElementRowReferenceColumn{ .Row = TargetRow });
+
+				// Parent widget row
+				DataStorage->AddColumn(ThumbnailWidgetRowHandle, FTableRowParentColumn{ .Parent = ParentWidgetRowHandle });
+
+				// TODO: ask, should the thumbnail tooltip be used for the entire widget or just the Thumbnail?
+				// Used to retrieve the Thumbnail tooltip to use on the whole TileItem
+				// DataStorage->AddColumn(ThumbnailWidgetRowHandle, FWidgetTooltipColumn::StaticStruct());
+
+				// Used to decide on the actual Thumbnail size
+				if (FSizeValueColumn_Experimental* TileItemSizeColumn = DataStorage->GetColumn<FSizeValueColumn_Experimental>(WidgetRow))
+				{
+					DataStorage->AddColumn(ThumbnailWidgetRowHandle, FSizeValueColumn_Experimental{ .SizeValue = TileItemSizeColumn->SizeValue });
+				}
+
+				ThumbnailWidget = DataStorageUi->ConstructWidget(ThumbnailWidgetRowHandle, *OutThumbnailWidgetConstructorPtr, Arguments);
+			}
 		}
 
+		UE::Editor::DataStorage::RowHandle LabelWidgetRowHandle = UE::Editor::DataStorage::InvalidRowHandle;
 		if (OutLabelWidgetConstructorPtr)
 		{
-			LabelWidget = DataStorageUi->ConstructWidget(WidgetRow, *OutLabelWidgetConstructorPtr, Arguments);
+			if (LabelWidgetRowHandle = DataStorage->AddRow(DataStorage->FindTable(TedsAssetDataHelper::TableView::GetWidgetTableName())); ThumbnailWidgetRowHandle != UE::Editor::DataStorage::InvalidRowHandle)
+			{
+				// Referenced Data Row
+				DataStorage->AddColumn(LabelWidgetRowHandle, FTypedElementRowReferenceColumn{ .Row = TargetRow });
+
+				LabelWidget = DataStorageUi->ConstructWidget(LabelWidgetRowHandle, *OutLabelWidgetConstructorPtr, Arguments);
+			}
 		}
 
 		if (!ThumbnailWidget.IsValid())
@@ -107,23 +150,19 @@ TSharedPtr<SWidget> FContentBrowserListViewNameWidgetConstructor::CreateWidget(U
 			LabelWidget = SNullWidget::NullWidget;
 		}
 
-		return SNew(SBox)
-				.Padding(ColumnItemPadding)
+		return SNew(SHorizontalBox)
+				+SHorizontalBox::Slot()
+				.AutoWidth()
+				[
+					ThumbnailWidget.ToSharedRef()
+				]
+
+				+SHorizontalBox::Slot()
+				.FillWidth(1.0f)
+				.Padding(ThumbnailNameHorizontalPadding, 0.f, 0.f, 0.f)
+				.VAlign(VAlign_Center)
 				[
-					SNew(SHorizontalBox)
-					+SHorizontalBox::Slot()
-					.AutoWidth()
-					[
-						ThumbnailWidget.ToSharedRef()
-					]
-
-					+SHorizontalBox::Slot()
-					.FillWidth(1.0f)
-					.Padding(ThumbnailNameHorizontalPadding, 0.f, 0.f, 0.f)
-					.VAlign(VAlign_Center)
-					[
-						LabelWidget.ToSharedRef()
-					]
+					LabelWidget.ToSharedRef()
 				];
 	}
 
diff --git a/Engine/Plugins/Experimental/EditorDataStorageFeatures/Source/TedsAssetData/Public/TedsAssetDataHelper.h b/Engine/Plugins/Experimental/EditorDataStorageFeatures/Source/TedsAssetData/Public/TedsAssetDataHelper.h
new file mode 100644
--- /dev/null
+++ b/Engine/Plugins/Experimental/EditorDataStorageFeatures/Source/TedsAssetData/Public/TedsAssetDataHelper.h
@@ -0,0 +1,29 @@
+// Copyright Epic Games, Inc. All Rights Reserved.
+
+#pragma once
+
+#include "Containers/UnrealString.h"
+#include "UObject/NameTypes.h"
+
+namespace TedsAssetDataHelper
+{
+	namespace MetaDataNames
+	{
+		TEDSASSETDATA_API FName GetThumbnailStatusMetaDataName();
+		TEDSASSETDATA_API FName GetThumbnailFadeInMetaDataName();
+		TEDSASSETDATA_API FName GetThumbnailHintTextMetaDataName();
+		TEDSASSETDATA_API FName GetThumbnailRealTimeOnHoveredMetaDataName();
+		TEDSASSETDATA_API FName GetThumbnailSizeOffsetMetaDataName();
+	}
+
+	namespace TableView
+	{
+		TEDSASSETDATA_API FName GetWidgetTableName();
+	}
+
+	/** Remove all / from the start of InStringToModify */
+	TEDSASSETDATA_API FString RemoveSlashFromStart(const FString& InStringToModify);
+
+	/** Remove Last / and successive char from InStringToModify */
+	TEDSASSETDATA_API FString RemoveAllFromLastSlash(const FString& InStringToModify);
+}
diff --git a/Engine/Plugins/Experimental/EditorDataStorageFeatures/Source/TedsAssetData/Public/TedsAssetDataWidgetColumns.h b/Engine/Plugins/Experimental/EditorDataStorageFeatures/Source/TedsAssetData/Public/TedsAssetDataWidgetColumns.h
new file mode 100644
--- /dev/null
+++ b/Engine/Plugins/Experimental/EditorDataStorageFeatures/Source/TedsAssetData/Public/TedsAssetDataWidgetColumns.h
@@ -0,0 +1,118 @@
+﻿// Copyright Epic Games, Inc. All Rights Reserved.
+
+#pragma once
+
+#include "TedsAssetDataStructs.h"
+
+#include "AssetThumbnail.h"
+#include "Elements/Common/TypedElementCommonTypes.h"
+#include "Framework/SlateDelegates.h"
+
+#include "TedsAssetDataWidgetColumns.generated.h"
+
+
+// Column to store the thumbnail size enum value of an asset
+USTRUCT(meta = (DisplayName = "Thumbnail Size"))
+struct FThumbnailSizeColumn_Experimental final : public FEditorDataStorageColumn
+{
+	GENERATED_BODY()
+
+	// EThumbnailSize
+	EThumbnailSize ThumbnailSize;
+};
+
+// Column to store the size value of a widget
+USTRUCT(meta = (DisplayName = "Size Value"))
+struct FSizeValueColumn_Experimental final : public FEditorDataStorageColumn
+{
+	GENERATED_BODY()
+
+	UPROPERTY()
+	float SizeValue;
+};
+
+// Column to store the padding
+USTRUCT(meta = (DisplayName = "Padding"))
+struct FWidgetPaddingColumn_Experimental final : public FEditorDataStorageColumn
+{
+	GENERATED_BODY()
+
+	UPROPERTY()
+	FMargin Padding;
+};
+
+/**
+ * Column added to a widget row to allow systems to provide an external EditMode trigger
+ */
+USTRUCT(meta = (DisplayName = "Content Browser Settings"))
+struct FThumbnailEditModeColumn_Experimental final : public FEditorDataStorageColumn
+{
+	GENERATED_BODY()
+
+	UPROPERTY()
+	bool IsEditModeToggled;
+};
+
+/**
+ * Column added to a widget row to allow systems to provide an external Overflow Policy
+ */
+USTRUCT(meta = (DisplayName = "Widget with externally managed OverflowPolicy"))
+struct FTextOverflowPolicyColumn_Experimental : public FEditorDataStorageColumn
+{
+	GENERATED_BODY()
+	ETextOverflowPolicy OverflowPolicy;
+};
+
+/**
+ * Column added to a widget row to allow systems to provide an external Visibility
+ */
+USTRUCT(meta = (DisplayName = "Widget with externally managed Visibility"))
+struct FWidgetVisibilityColumn_Experimental : public FEditorDataStorageColumn
+{
+	GENERATED_BODY()
+	EVisibility Visibility;
+};
+
+/**
+ * Column added to a widget row to allow systems to provide an external FontStyle
+ */
+USTRUCT(meta = (DisplayName = "Widget with externally managed FontStyle"))
+struct FFontStyleColumn_Experimental : public FEditorDataStorageColumn
+{
+	GENERATED_BODY()
+	FSlateFontInfo FontInfo;
+};
+
+/**
+ * Column added to a widget row to allow systems to store the widget tooltip
+ */
+USTRUCT(meta = (DisplayName = "Widget Tooltip"))
+struct FLocalWidgetTooltipColumn_Experimental : public FEditorDataStorageColumn
+{
+	GENERATED_BODY()
+	TSharedPtr<class IToolTip> Tooltip;
+};
+
+DECLARE_DELEGATE_RetVal(const FSlateBrush*, FOnGetWidgetSlateBrush)
+/**
+ * Column added to a widget row to allow systems to provide an external brush callback for the widget
+ * TODO: This need to be later on removed and replaced by updating the value only when needed (Hover/Selection etc...)
+ */
+USTRUCT(meta = (DisplayName = "Widget with externally managed Brush"))
+struct FOnGetWidgetSlateBrushColumn_Experimental : public FEditorDataStorageColumn
+{
+	GENERATED_BODY()
+	FOnGetWidgetSlateBrush OnGetWidgetSlateBrush;
+};
+
+DECLARE_DELEGATE_RetVal(FSlateColor, FOnGetWidgetColorAndOpacity)
+/**
+ * Column added to a widget row to allow systems to provide an external ColorAndOpacity callback for the widget
+ * TODO: This need to be later on removed and replaced by updating the value only when needed (Hover/Selection etc...)
+ */
+USTRUCT(meta = (DisplayName = "Widget with externally managed ColorAndOpacity"))
+struct FOnGetWidgetColorAndOpacityColumn_Experimental : public FEditorDataStorageColumn
+{
+	GENERATED_BODY()
+	FOnGetWidgetColorAndOpacity OnGetWidgetColorAndOpacity;
+};
diff --git a/Engine/Plugins/Experimental/EditorDataStorageFeatures/Source/TedsAssetData/TedsAssetData.build.cs b/Engine/Plugins/Experimental/EditorDataStorageFeatures/Source/TedsAssetData/TedsAssetData.build.cs
--- a/Engine/Plugins/Experimental/EditorDataStorageFeatures/Source/TedsAssetData/TedsAssetData.build.cs
+++ b/Engine/Plugins/Experimental/EditorDataStorageFeatures/Source/TedsAssetData/TedsAssetData.build.cs
@@ -32,6 +32,7 @@ public class TedsAssetData : ModuleRules
 					"Slate",
 					"SlateCore",
 					"TedsAlerts",
+					"ToolWidgets",
 					"TypedElementFramework",
 					"UnrealEd",
 					"UnrealEd",
diff --git a/Engine/Plugins/Experimental/EditorDataStorageFeatures/Source/TedsContentBrowser/Private/TedsContentBrowserModule.cpp b/Engine/Plugins/Experimental/EditorDataStorageFeatures/Source/TedsContentBrowser/Private/TedsContentBrowserModule.cpp
--- a/Engine/Plugins/Experimental/EditorDataStorageFeatures/Source/TedsContentBrowser/Private/TedsContentBrowserModule.cpp
+++ b/Engine/Plugins/Experimental/EditorDataStorageFeatures/Source/TedsContentBrowser/Private/TedsContentBrowserModule.cpp
@@ -15,6 +15,7 @@
 #include "Modules/ModuleManager.h"
 #include "TedsAlertColumns.h"
 #include "TedsAssetDataColumns.h"
+#include "TedsAssetDataWidgetColumns.h"
 #include "TedsRowViewNode.h"
 #include "Elements/Columns/TypedElementLabelColumns.h"
 #include "Elements/Columns/TypedElementTypeInfoColumns.h"
@@ -217,8 +218,12 @@ namespace UE::Editor::ContentBrowser
 							.Identifier = "ShadowPhysicsAsset"
 						});
 
+		CustomViewType = ETableViewMode::List;
+
 		// Create the table viewer widget
 		TableViewer = SNew(STedsTableViewer)
+					.ItemHeight_Lambda([this] () { return GetListViewItemHeight(); })
+					.ItemPadding_Lambda([this] () { return GetListViewItemPadding(); })
 					.QueryStack(RowQueryStack)
 					.CellWidgetPurpose(IUiProvider::FPurposeInfo("ContentBrowser", "RowLabel", NAME_None).GeneratePurposeID())
 					.HeaderWidgetPurpose(IUiProvider::FPurposeInfo("General", "Header", NAME_None).GeneratePurposeID())
@@ -237,17 +242,26 @@ namespace UE::Editor::ContentBrowser
 							OnSelectionChangedDelegate.Execute(AssetViewItem, ESelectInfo::Direct);
 						}
 					});
-		
+
+		const RowHandle TileViewRowHandle = TableViewer->GetWidgetRowHandle();
+		DataStorage->AddColumn(TileViewRowHandle, FSizeValueColumn_Experimental{ .SizeValue = GetThumbnailSizeValue() });
+
 		BindViewColumns();
 	}
 
 	void FTedsContentBrowserViewExtender::CreateTileView()
 	{
+		CustomViewType = ETableViewMode::Tile;
+
 		// Create the table viewer widget
 		TableViewer = SNew(STedsTileViewer)
+					.ItemAlignment(EListItemAlignment::LeftAligned)
+					.TileStyle(&FAppStyle::Get().GetWidgetStyle<FTableRowStyle>("ContentBrowser.AssetListView.TileTableRow"))
 					.QueryStack(RowQueryStack)
-					.WidgetPurpose(IUiProvider::FPurposeInfo("ContentBrowser", "RowLabel", NAME_None).GeneratePurposeID())
+					.WidgetPurpose(IUiProvider::FPurposeInfo("ContentBrowser", "TileLabel", NAME_None).GeneratePurposeID())
 					.SelectionMode(ESelectionMode::Multi)
+					.ItemWidth_Lambda([this] () { return GetTileViewItemWidth(); })
+					.ItemHeight_Lambda([this] () { return GetTileViewItemHeight(); })
 					.OnSelectionChanged_Lambda([this](RowHandle Row)
 					{
 						if(TSharedPtr<FAssetViewItem> AssetViewItem = GetAssetViewItemFromRow(Row))
@@ -256,7 +270,12 @@ namespace UE::Editor::ContentBrowser
 							OnSelectionChangedDelegate.Execute(AssetViewItem, ESelectInfo::Direct);
 						}
 					});
-		
+
+		const RowHandle TileViewRowHandle = TableViewer->GetWidgetRowHandle();
+		DataStorage->AddColumn(TileViewRowHandle, FThumbnailSizeColumn_Experimental{ .ThumbnailSize = GetThumbnailSize() });
+		DataStorage->AddColumn(TileViewRowHandle, FSizeValueColumn_Experimental{ .SizeValue = GetThumbnailSizeValue() });
+		DataStorage->AddColumn(TileViewRowHandle, FThumbnailEditModeColumn_Experimental{ .IsEditModeToggled = IsThumbnailEditMode });
+
 		BindViewColumns();
 	}
 
@@ -300,6 +319,78 @@ namespace UE::Editor::ContentBrowser
 		}
 	}
 
+	float FTedsContentBrowserViewExtender::GetTileViewTypeNameHeight() const
+	{
+		if (CurrentThumbnailSize == EThumbnailSize::Tiny)
+		{
+			return 0.f;
+		}
+		return TedsTileViewTypeNameHeight;
+	}
+
+	float FTedsContentBrowserViewExtender::GetThumbnailSizeValue() const
+	{
+		return ThumbnailSizeValue;
+	}
+
+	EThumbnailSize FTedsContentBrowserViewExtender::GetThumbnailSize() const
+	{
+		return CurrentThumbnailSize;
+	}
+
+	float FTedsContentBrowserViewExtender::GetTileViewItemWidth() const
+	{
+		return GetThumbnailSizeValue() + TedsTileViewWidthPadding;
+	}
+
+	float FTedsContentBrowserViewExtender::GetTileViewItemHeight() const
+	{
+		return GetThumbnailSizeValue() + GetTileViewTypeNameHeight() + TedsTileViewHeightPadding;
+	}
+
+	float FTedsContentBrowserViewExtender::GetListViewItemHeight() const
+	{
+		return GetThumbnailSizeValue();
+	}
+
+	FMargin FTedsContentBrowserViewExtender::GetListViewItemPadding() const
+	{
+		return GetThumbnailSize() >= EThumbnailSize::Large? FMargin(0.f, TedsListViewHeightPadding) : FMargin(0.f);
+	}
+
+	void FTedsContentBrowserViewExtender::UpdateThumbnailSize() const
+	{
+		if (TableViewer.IsValid())
+		{
+			if (FThumbnailSizeColumn_Experimental* ThumbnailSizeColumn = DataStorage->GetColumn<FThumbnailSizeColumn_Experimental>(TableViewer->GetWidgetRowHandle()))
+			{
+				ThumbnailSizeColumn->ThumbnailSize = GetThumbnailSize();
+			}
+		}
+	}
+
+	void FTedsContentBrowserViewExtender::UpdateSizeValue() const
+	{
+		if (TableViewer.IsValid())
+		{
+			if (FSizeValueColumn_Experimental* SizeValueColumn = DataStorage->GetColumn<FSizeValueColumn_Experimental>(TableViewer->GetWidgetRowHandle()))
+			{
+				SizeValueColumn->SizeValue = GetThumbnailSizeValue();
+			}
+		}
+	}
+
+	void FTedsContentBrowserViewExtender::UpdateEditMode() const
+	{
+		if (TableViewer.IsValid())
+		{
+			if (FThumbnailEditModeColumn_Experimental* ContentBrowserSettingsColumn = DataStorage->GetColumn<FThumbnailEditModeColumn_Experimental>(TableViewer->GetWidgetRowHandle()))
+			{
+				ContentBrowserSettingsColumn->IsEditModeToggled = IsThumbnailEditMode;
+			}
+		}
+	}
+
 	FName FTestContentSource::GetName()
 	{
 		return FName("TestContentSource");
@@ -372,6 +463,22 @@ namespace UE::Editor::ContentBrowser
 		RefreshRows(InItemsSource);
 	}
 
+	bool FTedsContentBrowserViewExtender::SupportScaling() const
+	{
+		return true;
+	}
+
+	bool FTedsContentBrowserViewExtender::SupportZooming() const
+	{
+		// Zooming is for now supported only in the custom TileView
+		return CustomViewType == ETableViewMode::Tile;
+	}
+
+	bool FTedsContentBrowserViewExtender::SupportThumbnailEditMode() const
+	{
+		// Thumbnail edit mode is for now supported only in the custom TileView
+		return CustomViewType == ETableViewMode::Tile;
+	}
 
 	TArray<TSharedPtr<FAssetViewItem>> FTedsContentBrowserViewExtender::GetSelectedItems()
 	{
diff --git a/Engine/Plugins/Experimental/EditorDataStorageFeatures/Source/TedsContentBrowser/Private/TedsContentBrowserModule.h b/Engine/Plugins/Experimental/EditorDataStorageFeatures/Source/TedsContentBrowser/Private/TedsContentBrowserModule.h
--- a/Engine/Plugins/Experimental/EditorDataStorageFeatures/Source/TedsContentBrowser/Private/TedsContentBrowserModule.h
+++ b/Engine/Plugins/Experimental/EditorDataStorageFeatures/Source/TedsContentBrowser/Private/TedsContentBrowserModule.h
@@ -59,6 +59,9 @@ namespace UE::Editor::ContentBrowser
 		virtual void ClearSelection() override;
 		virtual bool IsRightClickScrolling() override;
 		virtual void OnItemListChanged(TArray<TSharedPtr<FAssetViewItem>>* InItemsSource) override;
+		virtual bool SupportScaling() const override;
+		virtual bool SupportZooming() const override;
+		virtual bool SupportThumbnailEditMode() const override;
 		// ~IContentBrowserViewExtender interface
 
 		// Refresh the rows in the current view by syncing the the items source
@@ -73,7 +76,6 @@ namespace UE::Editor::ContentBrowser
 		DataStorage::RowHandle GetRowFromAssetViewItem(const TSharedPtr<FAssetViewItem>& Item);
 
 	protected:
-
 		// Update the table viewer used by the integration to use a list view
 		void CreateListView();
 
@@ -84,6 +86,60 @@ namespace UE::Editor::ContentBrowser
 		void BindViewColumns();
 
 	private:
+		/** Get the name area total height */
+		float GetTileViewTypeNameHeight() const;
+
+		/** Get the thumbnail size value */
+		float GetThumbnailSizeValue() const;
+
+		/** Get the thumbnail size */
+		EThumbnailSize GetThumbnailSize() const;
+
+		/** Get the tile item width */
+		float GetTileViewItemWidth() const;
+
+		/** Get the tile item height */
+		float GetTileViewItemHeight() const;
+
+		/** Get the list item height */
+		float GetListViewItemHeight() const;
+
+		/** Get the list item padding */
+		FMargin GetListViewItemPadding() const;
+
+		/** Update the thumbnail size of the Tile/Column view and broadcast the event */
+		void UpdateThumbnailSize() const;
+
+		/** Update the size value of the Tile/Column view and broadcast the event */
+		void UpdateSizeValue() const;
+
+		/** Update the EditMode value of the Tile view and broadcast the event */
+		void UpdateEditMode() const;
+
+	private:
+		/** Custom view type */
+		ETableViewMode::Type CustomViewType = ETableViewMode::List;
+
+		/** The current Thumbnail size */
+		EThumbnailSize CurrentThumbnailSize = EThumbnailSize::Medium;
+
+		/** The current Thumbnail size value */
+		float ThumbnailSizeValue = 80.f;
+
+		/** Whether the thumbnail edit mode is enabled */
+		bool IsThumbnailEditMode = false;
+
+		/** Name area total height */
+		static constexpr float TedsTileViewTypeNameHeight = 67.f;
+
+		/** Vertical padding for the TileViewItem */
+		static constexpr int32 TedsTileViewHeightPadding = 9;
+
+		/** Horizontal padding for the TileViewItem */
+		static constexpr int32 TedsTileViewWidthPadding = 8;
+
+		/** Vertical padding for the ListViewItem */
+		static constexpr int32 TedsListViewHeightPadding = 2;
 
 		// Ptr to the data storage interface
 		DataStorage::ICoreProvider* DataStorage;
diff --git a/Engine/Plugins/Experimental/EditorDataStorageFeatures/Source/TedsContentBrowser/Private/Widgets/ContentBrowserTileViewWidget.cpp b/Engine/Plugins/Experimental/EditorDataStorageFeatures/Source/TedsContentBrowser/Private/Widgets/ContentBrowserTileViewWidget.cpp
new file mode 100644
--- /dev/null
+++ b/Engine/Plugins/Experimental/EditorDataStorageFeatures/Source/TedsContentBrowser/Private/Widgets/ContentBrowserTileViewWidget.cpp
@@ -0,0 +1,622 @@
+// Copyright Epic Games, Inc. All Rights Reserved.
+
+#include "ContentBrowserTileViewWidget.h"
+
+#include "AssetThumbnail.h"
+#include "AssetDefinition.h"
+#include "Elements/Framework/TypedElementAttributeBinding.h"
+#include "TedsAssetDataColumns.h"
+#include "TedsAssetDataHelper.h"
+#include "TedsAssetDataWidgetColumns.h"
+#include "TedsTableViewerWidgetColumns.h"
+#include "Columns/SlateDelegateColumns.h"
+#include "Elements/Columns/TypedElementFolderColumns.h"
+#include "Elements/Columns/TypedElementHiearchyColumns.h"
+#include "Elements/Columns/TypedElementMiscColumns.h"
+#include "Internationalization/BreakIterator.h"
+#include "Styling/AppStyle.h"
+#include "Styling/SlateBrush.h"
+#include "Styling/StyleColors.h"
+#include "Widgets/Images/SImage.h"
+#include "Widgets/Layout/SBox.h"
+#include "Widgets/Layout/SSpacer.h"
+#include "Widgets/Text/STextBlock.h"
+#include "Widgets/SBoxPanel.h"
+#include "Widgets/Text/SInlineEditableTextBlock.h"
+
+#define LOCTEXT_NAMESPACE "ContentBrowserTileViewWidget"
+
+namespace Purpose
+{
+	static const FLazyName WidgetPurposeNamespace(TEXT("ContentBrowser"));
+	static const FLazyName WidgetPurposeName(TEXT("TileLabel"));
+
+	FName GetPurposeNamespace()
+	{
+		return WidgetPurposeNamespace;
+	}
+
+	FName GetPurposeName()
+	{
+		return WidgetPurposeName;
+	}
+}
+
+void UContentBrowserTileViewWidgetFactory::RegisterWidgetPurposes(UE::Editor::DataStorage::IUiProvider& DataStorageUi) const
+{
+	using namespace UE::Editor::DataStorage::Queries;
+	DataStorageUi.RegisterWidgetPurpose(
+		IUiProvider::FPurposeInfo(Purpose::GetPurposeNamespace(), Purpose::GetPurposeName(), NAME_None,
+		IUiProvider::EPurposeType::UniqueByNameAndColumn,
+		LOCTEXT("ContentBrowserTileViewWidget_PurposeDescription", "Widget that display a Tile in the Content Browser")));
+}
+
+void UContentBrowserTileViewWidgetFactory::RegisterWidgetConstructors(UE::Editor::DataStorage::ICoreProvider& DataStorage,
+	UE::Editor::DataStorage::IUiProvider& DataStorageUi) const
+{
+	using namespace UE::Editor::DataStorage::Queries;
+	DataStorageUi.RegisterWidgetFactory<FContentBrowserTileViewWidgetConstructor>(
+		DataStorageUi.FindPurpose(IUiProvider::FPurposeInfo(Purpose::GetPurposeNamespace(), Purpose::GetPurposeName(), NAME_None).GeneratePurposeID()), 
+		TColumn<FAssetTag>() || TColumn<FFolderTag>());
+}
+
+FContentBrowserTileViewWidgetConstructor::FContentBrowserTileViewWidgetConstructor()
+	: FSimpleWidgetConstructor(StaticStruct())
+{
+}
+
+TSharedPtr<SWidget> FContentBrowserTileViewWidgetConstructor::CreateWidget(UE::Editor::DataStorage::ICoreProvider* DataStorage,
+	UE::Editor::DataStorage::IUiProvider* DataStorageUi,
+	UE::Editor::DataStorage::RowHandle TargetRow, UE::Editor::DataStorage::RowHandle WidgetRow,
+	const UE::Editor::DataStorage::FMetaDataView& Arguments)
+{
+	using namespace UE::Editor::DataStorage;
+
+	if (DataStorage->IsRowAvailable(TargetRow))
+	{
+		FAttributeBinder WidgetRowBinder(WidgetRow, DataStorage);
+		FAttributeBinder RowBinder(TargetRow, DataStorage);
+		RowHandle ParentWidgetRowHandle = InvalidRowHandle;
+		if (FTableRowParentColumn* ParentWidgetRow = DataStorage->GetColumn<FTableRowParentColumn>(WidgetRow))
+		{
+			ParentWidgetRowHandle = ParentWidgetRow->Parent;
+		}
+		FAttributeBinder ParentWidgetRowBinder(ParentWidgetRowHandle, DataStorage);
+
+		constexpr float BorderPadding = 1.f;
+		constexpr float ShadowLeftTopPadding = 3.f;
+		constexpr float ShadowRightBotPadding = 4.f;
+		constexpr float ThumbnailBorderPadding = 0.f;
+		constexpr float NameAreaBoxLeftRightBotPadding = 4.f;
+		constexpr float NameAreaBoxTopPadding = 6.f;
+		constexpr float ClassNameMaxHeight = 14.f;
+		const FName ItemShadowBorderName = FName(TEXT("ContentBrowser.AssetTileItem.DropShadow"));
+		bool bIsAsset = DataStorage->HasColumns<FAssetTag>(TargetRow);
+
+		// Maybe divide them into Folder and Asset
+		TSharedPtr<SBox> TileItem = SNew(SBox)
+			.Padding(FMargin(BorderPadding, BorderPadding, 0.f, 1.f));
+
+		FTedsTileViewCommonArguments TedsTileViewCommonArguments
+		{
+			.DataStorage = DataStorage,
+			.DataStorageUi = DataStorageUi,
+			.TargetRow = TargetRow,
+			.WidgetRow = WidgetRow,
+			.ParentWidgetRowHandle = ParentWidgetRowHandle,
+			.bIsAsset = bIsAsset,
+			.ThumbnailWidget = SNullWidget::NullWidget,
+			.TileItem = TileItem
+		};
+
+		// Create the thumbnail widget through teds
+		RowHandle OutThumbnailWidgetRowHandle;
+		TSharedPtr<SWidget> ThumbnailWidget = CreateThumbnailWidget(TedsTileViewCommonArguments, OutThumbnailWidgetRowHandle);
+		if (!ThumbnailWidget.IsValid())
+		{
+			ThumbnailWidget = SNullWidget::NullWidget;
+		}
+		else
+		{
+			FAttributeBinder ThumbnailWidgetBinder(OutThumbnailWidgetRowHandle, DataStorage);
+			TileItem->SetToolTip(ThumbnailWidgetBinder.BindData(&FLocalWidgetTooltipColumn_Experimental::Tooltip));
+		}
+
+		// Create the item type widget through teds
+		RowHandle OutItemTypeWidgetRowHandle;
+		TSharedPtr<SWidget> ItemTypeWidget = CreateItemTypeWidget(TedsTileViewCommonArguments, OutItemTypeWidgetRowHandle);
+		if (!ItemTypeWidget.IsValid())
+		{
+			ItemTypeWidget = SNullWidget::NullWidget;
+		}
+
+		TedsTileViewCommonArguments.ThumbnailWidget = ThumbnailWidget;
+
+		TileItem->SetContent(
+				SNew(SBorder)
+				.Padding(FMargin(ShadowLeftTopPadding, ShadowLeftTopPadding, ShadowRightBotPadding, ShadowRightBotPadding))
+				.BorderImage_Lambda([this, ItemShadowBorderName, TedsTileViewCommonArguments] ()
+				{
+					if (TedsTileViewCommonArguments.bIsAsset)
+					{
+						return FAppStyle::GetBrush(ItemShadowBorderName);
+					}
+					return GetFolderBackgroundShadowImage(TedsTileViewCommonArguments);
+				})
+				[
+					SNew(SBorder)
+					.Padding(ThumbnailBorderPadding)
+					.BorderImage_Lambda([this, TedsTileViewCommonArguments] ()
+						{
+							if (TedsTileViewCommonArguments.bIsAsset)
+							{
+								return GetNameAreaBackgroundImage(TedsTileViewCommonArguments);
+							}
+							return GetFolderBackgroundImage(TedsTileViewCommonArguments);
+						})
+					[
+						SNew(SVerticalBox)
+
+						+ SVerticalBox::Slot()
+						.AutoHeight()
+						[
+							SNew(SOverlay)
+
+							//Thumbnail
+							+ SOverlay::Slot()
+							.Padding(bIsAsset ? 0.f : FMargin(2.f, 2.f, 2.f, 2.f))
+							[
+								SNew(SBorder)
+								.Padding(0.f)
+								.BorderImage_Lambda([this, TedsTileViewCommonArguments] (){ return GetFolderSlotBorder(TedsTileViewCommonArguments); })
+								[
+									ThumbnailWidget.ToSharedRef()
+								]
+							]
+
+							// Overlay name for tiny folders
+							+ SOverlay::Slot()
+							.HAlign(HAlign_Center)
+							.VAlign(VAlign_Center)
+							.Padding(10.f, 0.f)
+							[
+								SNew(SBorder)
+								.Padding(4.f, 2.f)
+								.BorderImage(FAppStyle::GetBrush("ContentBrowser.AssetTileItem.TinyFolderTextBorder"))
+								.Visibility(ParentWidgetRowBinder.BindData(&FThumbnailSizeColumn_Experimental::ThumbnailSize, [bIsAsset] (EThumbnailSize InThumbnailSize)
+								{
+									return !bIsAsset && InThumbnailSize == EThumbnailSize::Tiny ? EVisibility::Visible : EVisibility::Collapsed;
+								}))
+								[
+									SNew(STextBlock)
+									.ColorAndOpacity(FStyleColors::White)
+									.OverflowPolicy(ETextOverflowPolicy::Ellipsis)
+									.Text(RowBinder.BindData(&FNameColumn::Name, [] (FName InName)
+									{
+										return FText::FromString(TedsAssetDataHelper::RemoveSlashFromStart(InName.ToString()));
+									}))
+								]
+							]
+						]
+
+						+ SVerticalBox::Slot()
+						.FillHeight(1.f)
+						[
+							SNew(SBox)
+							.Padding(FMargin(NameAreaBoxLeftRightBotPadding, NameAreaBoxTopPadding, NameAreaBoxLeftRightBotPadding, NameAreaBoxLeftRightBotPadding))
+							.Visibility(ParentWidgetRowBinder.BindData(&FThumbnailSizeColumn_Experimental::ThumbnailSize, [] (EThumbnailSize InThumbnailSize)
+							{
+								return InThumbnailSize == EThumbnailSize::Tiny ? EVisibility::Collapsed : EVisibility::Visible;
+							}))
+							[
+								SNew(SVerticalBox)
+								+ SVerticalBox::Slot()
+								.HAlign(!bIsAsset ? HAlign_Center : HAlign_Fill)
+								[
+									SNew(SBox)
+									.VAlign(VAlign_Top)
+									.HeightOverride(GetNameAreaMaxDesiredHeight())
+									[
+										SNew(SInlineEditableTextBlock)
+											.Font(GetThumbnailFont())
+											.Text(RowBinder.BindData(&FNameColumn::Name, [] (FName InName)
+											{
+												return FText::FromString(TedsAssetDataHelper::RemoveSlashFromStart(InName.ToString()));
+											}))
+											// TODO: Renaming logic for later, this should be synced between the actual AssetData and TEDS
+											// .OnBeginTextEdit()
+											// .OnTextCommitted()
+											// .OnVerifyTextChanged()
+											// .HighlightText(InArgs._HighlightText)
+											.IsSelected(WidgetRowBinder.BindEvent(&FExternalWidgetExclusiveSelectionColumn::IsSelectedExclusively))
+											// TODO: need to check if in EditMode if valid if temporary and if can be renamed, for now true by default need to be changed when the rename logic is ready
+											.IsReadOnly(true)
+											// TODO: CB TileView instead of setting this to true is updating this on the tick only when needed, see: SAssetViewItem::Tick
+											.AutoWrapNonEditText(true)
+											.LineBreakPolicy(FBreakIterator::CreateCamelCaseBreakIterator())
+											.OverflowPolicy(ETextOverflowPolicy::MultilineEllipsis)
+											.ColorAndOpacity_Lambda([this, TedsTileViewCommonArguments] () { return GetNameAreaTextColor(TedsTileViewCommonArguments); })
+									]
+								]
+
+								+ SVerticalBox::Slot()
+								.AutoHeight()
+								[
+									SNew(SBox)
+									.HeightOverride(ClassNameMaxHeight)
+									.VAlign(VAlign_Bottom)
+									[
+										ItemTypeWidget.ToSharedRef()
+									]
+								]
+							]
+						]
+					]
+				]
+			);
+
+		return TileItem;
+	}
+
+	return SNullWidget::NullWidget;;
+}
+
+TConstArrayView<const UScriptStruct*> FContentBrowserTileViewWidgetConstructor::GetAdditionalColumnsList() const
+{
+	static UE::Editor::DataStorage::TTypedElementColumnTypeList<FSizeValueColumn_Experimental, FThumbnailSizeColumn_Experimental> Columns;
+	return Columns;
+}
+
+const FSlateBrush* FContentBrowserTileViewWidgetConstructor::GetFolderBackgroundShadowImage(FTedsTileViewCommonArguments InTedsTileViewCommonArguments) const
+{
+	if (!InTedsTileViewCommonArguments.DataStorage)
+	{
+		return FStyleDefaults::GetNoBrush();
+	}
+
+	UE::Editor::DataStorage::FAttributeBinder WidgetRowBinder(InTedsTileViewCommonArguments.WidgetRow, InTedsTileViewCommonArguments.DataStorage);
+
+	FIsSelected IsSelected = WidgetRowBinder.BindEvent(&FExternalWidgetSelectionColumn::IsSelected);
+
+	const bool bIsSelected = IsSelected.IsBound() ? IsSelected.Execute() : false;
+	const bool bIsHoveredOrDraggedOver = (InTedsTileViewCommonArguments.TileItem.IsValid() && InTedsTileViewCommonArguments.TileItem->IsHovered()) || false /*bDraggedOver*/;
+
+	if (bIsSelected || bIsHoveredOrDraggedOver)
+	{
+		static const FName DropShadowName("ContentBrowser.AssetTileItem.DropShadow");
+		return FAppStyle::Get().GetBrush(DropShadowName);
+	}
+
+	return FStyleDefaults::GetNoBrush();
+}
+
+const FSlateBrush* FContentBrowserTileViewWidgetConstructor::GetFolderSlotBorder(FTedsTileViewCommonArguments InTedsTileViewCommonArguments) const
+{
+	if (!InTedsTileViewCommonArguments.DataStorage)
+	{
+		return FStyleDefaults::GetNoBrush();
+	}
+
+	UE::Editor::DataStorage::FAttributeBinder WidgetRowBinder(InTedsTileViewCommonArguments.WidgetRow, InTedsTileViewCommonArguments.DataStorage);
+
+	FIsSelected IsSelected = WidgetRowBinder.BindEvent(&FExternalWidgetSelectionColumn::IsSelected);
+
+	const bool bIsSelected = IsSelected.IsBound() ? IsSelected.Execute() : false;
+	const bool bIsHoveredOrDraggedOver = (InTedsTileViewCommonArguments.TileItem.IsValid() && InTedsTileViewCommonArguments.TileItem->IsHovered()) || false /*bDraggedOver*/;
+
+	if (bIsSelected || bIsHoveredOrDraggedOver)
+	{
+		static const FLazyName SelectedOrHovered("ContentBrowser.AssetTileItem.FolderAreaBackground"); // Panel
+		return FAppStyle::Get().GetBrush(SelectedOrHovered);
+	}
+
+	return FStyleDefaults::GetNoBrush();
+}
+
+const FSlateBrush* FContentBrowserTileViewWidgetConstructor::GetNameAreaBackgroundImage(FTedsTileViewCommonArguments InTedsTileViewCommonArguments) const
+{
+	if (!InTedsTileViewCommonArguments.DataStorage)
+	{
+		return FStyleDefaults::GetNoBrush();
+	}
+
+	if (const FThumbnailSizeColumn_Experimental* ThumbnailSizeColumn = InTedsTileViewCommonArguments.DataStorage->GetColumn<FThumbnailSizeColumn_Experimental>(InTedsTileViewCommonArguments.ParentWidgetRowHandle))
+	{
+		if (ThumbnailSizeColumn->ThumbnailSize == EThumbnailSize::Tiny)
+		{
+			return FStyleDefaults::GetNoBrush();
+		}
+	}
+
+	UE::Editor::DataStorage::FAttributeBinder WidgetRowBinder(InTedsTileViewCommonArguments.WidgetRow, InTedsTileViewCommonArguments.DataStorage);
+
+	FIsSelected IsSelected = WidgetRowBinder.BindEvent(&FExternalWidgetSelectionColumn::IsSelected);
+
+	const static FLazyName SelectedHover = "ContentBrowser.AssetTileItem.AssetContentSelectedHoverBackground";
+	const static FLazyName Selected = "ContentBrowser.AssetTileItem.AssetContentSelectedBackground";
+	const static FLazyName Hovered = "ContentBrowser.AssetTileItem.AssetContentHoverBackground";
+	const static FLazyName Normal = "ContentBrowser.AssetTileItem.AssetContent";
+
+	const bool bIsSelected = IsSelected.IsBound() ? IsSelected.Execute() : false;
+	// TODO: ThumbnailWidget is the TedsWidget and not the actual Thumbnail widget, need to check the hovered state of the actual thumbnail widget, needed for the ThumbnailEditMode to keep it hovered
+	const bool bIsHoveredOrDraggedOver = (InTedsTileViewCommonArguments.TileItem.IsValid() && InTedsTileViewCommonArguments.TileItem->IsHovered()) || false /*bDraggedOver*/ || false /*(InThumbnailWidget.IsValid() && InThumbnailWidget->IsHovered())*/;
+	if (bIsSelected && bIsHoveredOrDraggedOver)
+	{
+		return FAppStyle::Get().GetBrush(SelectedHover);
+	}
+	else if (bIsSelected)
+	{
+		return FAppStyle::Get().GetBrush(Selected);
+	}
+	else if (bIsHoveredOrDraggedOver && InTedsTileViewCommonArguments.bIsAsset)
+	{
+		return FAppStyle::Get().GetBrush(Hovered);
+	}
+	else if (InTedsTileViewCommonArguments.bIsAsset)
+	{
+		return FAppStyle::Get().GetBrush(Normal);
+	}
+	return FStyleDefaults::GetNoBrush();
+}
+
+FSlateFontInfo FContentBrowserTileViewWidgetConstructor::GetThumbnailFont() const
+{
+	const static FLazyName RegularFont("ContentBrowser.AssetTileViewNameFont");
+    return FAppStyle::GetFontStyle(RegularFont);
+}
+
+FSlateColor FContentBrowserTileViewWidgetConstructor::GetNameAreaTextColor(FTedsTileViewCommonArguments InTedsTileViewCommonArguments) const
+{
+	if (!InTedsTileViewCommonArguments.DataStorage)
+	{
+		return FSlateColor::UseForeground();
+	}
+
+	UE::Editor::DataStorage::FAttributeBinder WidgetRowBinder(InTedsTileViewCommonArguments.WidgetRow, InTedsTileViewCommonArguments.DataStorage);
+
+	FIsSelected IsSelected = WidgetRowBinder.BindEvent(&FExternalWidgetSelectionColumn::IsSelected);
+
+	const bool bIsSelected = IsSelected.IsBound() ? IsSelected.Execute() : false;
+	// TODO: ThumbnailWidget is the TedsWidget and not the actual Thumbnail widget, need to check the hovered state of the actual thumbnail widget, needed for the ThumbnailEditMode to keep it hovered
+	const bool bIsHoveredOrDraggedOver = (InTedsTileViewCommonArguments.TileItem.IsValid() && InTedsTileViewCommonArguments.TileItem->IsHovered()) || false /*bDraggedOver*/ || false /*(InThumbnailWidget.IsValid() && InThumbnailWidget->IsHovered())*/;
+	if (bIsSelected || bIsHoveredOrDraggedOver)
+	{
+		return FStyleColors::White;
+	}
+
+	return FSlateColor::UseForeground();
+}
+
+const FSlateBrush* FContentBrowserTileViewWidgetConstructor::GetAssetThumbnailBorderOverride(FTedsTileViewCommonArguments InTedsTileViewCommonArguments) const
+{
+	if (!InTedsTileViewCommonArguments.DataStorage)
+	{
+		return FAppStyle::GetNoBrush();
+	}
+
+	UE::Editor::DataStorage::FAttributeBinder WidgetRowBinder(InTedsTileViewCommonArguments.WidgetRow, InTedsTileViewCommonArguments.DataStorage);
+
+	FIsSelected IsSelected = WidgetRowBinder.BindEvent(&FExternalWidgetSelectionColumn::IsSelected);
+
+	const bool bIsSelected = IsSelected.IsBound() ? IsSelected.Execute() : false;
+	// TODO: ThumbnailWidget is the TedsWidget and not the actual Thumbnail widget, need to check the hovered state of the actual thumbnail widget, needed for the ThumbnailEditMode to keep it hovered
+	const bool bIsHoveredOrDraggedOver = (InTedsTileViewCommonArguments.TileItem.IsValid() && InTedsTileViewCommonArguments.TileItem->IsHovered()) || false /*bDraggedOver*/ || false /*(InThumbnailWidget.IsValid() && InThumbnailWidget->IsHovered())*/;
+	if (bIsSelected && bIsHoveredOrDraggedOver)
+	{
+		static const FLazyName SelectedHover("ContentBrowser.AssetTileItem.AssetBorderSelectedHoverBackground");
+		return FAppStyle::Get().GetBrush(SelectedHover);
+	}
+	else if (bIsSelected)
+	{
+		static const FLazyName Selected("ContentBrowser.AssetTileItem.AssetBorderSelectedBackground");
+		return FAppStyle::Get().GetBrush(Selected);
+	}
+	else if (bIsHoveredOrDraggedOver && InTedsTileViewCommonArguments.bIsAsset)
+	{
+		static const FLazyName Hovered("ContentBrowser.AssetTileItem.AssetBorderHoverBackground");
+		return FAppStyle::Get().GetBrush(Hovered);
+	}
+	else if (InTedsTileViewCommonArguments.bIsAsset)
+	{
+		static const FLazyName Normal("AssetThumbnail.AssetBorder");
+		return FAppStyle::Get().GetBrush(Normal);
+	}
+
+	return FStyleDefaults::GetNoBrush();
+}
+
+FOptionalSize FContentBrowserTileViewWidgetConstructor::GetNameAreaMaxDesiredHeight() const
+{
+	constexpr int32 MaxHeightNameArea = 42;
+	return MaxHeightNameArea;
+}
+
+EVisibility FContentBrowserTileViewWidgetConstructor::GetAssetClassLabelVisibility(bool bIsAsset) const
+{
+	return !bIsAsset ? EVisibility::Collapsed : EVisibility::Visible;
+}
+
+FSlateColor FContentBrowserTileViewWidgetConstructor::GetAssetClassLabelTextColor(FTedsTileViewCommonArguments InTedsTileViewCommonArguments) const
+{
+	if (!InTedsTileViewCommonArguments.DataStorage)
+	{
+		return FStyleColors::Hover2;
+	}
+
+	UE::Editor::DataStorage::FAttributeBinder WidgetRowBinder(InTedsTileViewCommonArguments.WidgetRow, InTedsTileViewCommonArguments.DataStorage);
+
+	FIsSelected IsSelected = WidgetRowBinder.BindEvent(&FExternalWidgetSelectionColumn::IsSelected);
+
+	const bool bIsSelected = IsSelected.IsBound() ? IsSelected.Execute() : false;
+	// TODO: ThumbnailWidget is the TedsWidget and not the actual Thumbnail widget, need to check the hovered state of the actual thumbnail widget, needed for the ThumbnailEditMode to keep it hovered
+	const bool bIsHoveredOrDraggedOver = (InTedsTileViewCommonArguments.TileItem.IsValid() && InTedsTileViewCommonArguments.TileItem->IsHovered()) || false /*bDraggedOver*/ || false /*(InThumbnailWidget.IsValid() && InThumbnailWidget->IsHovered())*/;
+	if (bIsSelected || bIsHoveredOrDraggedOver)
+	{
+		return FStyleColors::White;
+	}
+
+	return FStyleColors::Hover2;
+}
+
+const FSlateBrush* FContentBrowserTileViewWidgetConstructor::GetFolderBackgroundImage(FTedsTileViewCommonArguments InTedsTileViewCommonArguments) const
+{
+	if (!InTedsTileViewCommonArguments.DataStorage)
+	{
+		return FStyleDefaults::GetNoBrush();
+	}
+
+	UE::Editor::DataStorage::FAttributeBinder WidgetRowBinder(InTedsTileViewCommonArguments.WidgetRow, InTedsTileViewCommonArguments.DataStorage);
+
+	FIsSelected IsSelected = WidgetRowBinder.BindEvent(&FExternalWidgetSelectionColumn::IsSelected);
+
+	const bool bIsSelected = IsSelected.IsBound() ? IsSelected.Execute() : false;
+	const bool bIsHoveredOrDraggedOver = (InTedsTileViewCommonArguments.TileItem.IsValid() && InTedsTileViewCommonArguments.TileItem->IsHovered()) || false /*bDraggedOver*/;
+
+	if (bIsSelected && bIsHoveredOrDraggedOver)
+	{
+		static const FName SelectedHoverBackground("ContentBrowser.AssetTileItem.FolderAreaSelectedHoverBackground");
+		return FAppStyle::Get().GetBrush(SelectedHoverBackground);
+	}
+	else if (bIsSelected)
+	{
+		static const FName SelectedBackground("ContentBrowser.AssetTileItem.FolderAreaSelectedBackground");
+		return FAppStyle::Get().GetBrush(SelectedBackground);
+	}
+	else if (bIsHoveredOrDraggedOver)
+	{
+		static const FName HoveredBackground("ContentBrowser.AssetTileItem.FolderAreaHoveredBackground");
+		return FAppStyle::Get().GetBrush(HoveredBackground);
+	}
+
+	return FStyleDefaults::GetNoBrush();
+}
+
+TSharedPtr<SWidget> FContentBrowserTileViewWidgetConstructor::CreateThumbnailWidget(FTedsTileViewCommonArguments InTedsTileViewCommonArguments, UE::Editor::DataStorage::RowHandle& OutThumbnailWidgetRowHandle) const
+{
+	using namespace UE::Editor::DataStorage;
+
+	TSharedPtr<FTypedElementWidgetConstructor> OutThumbnailWidgetConstructorPtr;
+
+	auto AssignThumbnailWidgetToColumn = [&OutThumbnailWidgetConstructorPtr] (TUniquePtr<FTypedElementWidgetConstructor> Constructor, TConstArrayView<TWeakObjectPtr<const UScriptStruct>>)
+	{
+		OutThumbnailWidgetConstructorPtr = TSharedPtr<FTypedElementWidgetConstructor>(Constructor.Release());
+		return false;
+	};
+
+	TArray<TWeakObjectPtr<const UScriptStruct>> ThumbnailColumns = GetThumbnailColumns();
+	UE::Editor::DataStorage::FMetaData ThumbnailMeta;
+	ThumbnailMeta.AddOrSetMutableData(TedsAssetDataHelper::MetaDataNames::GetThumbnailStatusMetaDataName(), true);
+	ThumbnailMeta.AddOrSetMutableData(TedsAssetDataHelper::MetaDataNames::GetThumbnailFadeInMetaDataName(), true);
+	ThumbnailMeta.AddOrSetMutableData(TedsAssetDataHelper::MetaDataNames::GetThumbnailHintTextMetaDataName(), false);
+	// TODO: use our own OnMouseEnter/Leave for setting the realtime flag, also check CB Setting
+	ThumbnailMeta.AddOrSetMutableData(TedsAssetDataHelper::MetaDataNames::GetThumbnailRealTimeOnHoveredMetaDataName(), false);
+
+	// Folder has more padding since this widget has to emulate the border that AssetThumbnails has
+	if (!InTedsTileViewCommonArguments.bIsAsset)
+	{
+		ThumbnailMeta.AddOrSetMutableData(TedsAssetDataHelper::MetaDataNames::GetThumbnailSizeOffsetMetaDataName(), -4.f);
+	}
+	const FGenericMetaDataView ThumbnailMetaView = FGenericMetaDataView(ThumbnailMeta);
+
+	const RowHandle DefaultPurposeRowHandle = InTedsTileViewCommonArguments.DataStorageUi->FindPurpose(InTedsTileViewCommonArguments.DataStorageUi->GetGeneralWidgetPurposeID());
+	InTedsTileViewCommonArguments.DataStorageUi->CreateWidgetConstructors(DefaultPurposeRowHandle, IUiProvider::EMatchApproach::ExactMatch, ThumbnailColumns, ThumbnailMetaView, AssignThumbnailWidgetToColumn);
+
+	TSharedPtr<SWidget> ThumbnailWidget = SNullWidget::NullWidget;
+	if (OutThumbnailWidgetConstructorPtr)
+	{
+		OutThumbnailWidgetRowHandle = InTedsTileViewCommonArguments.DataStorage->AddRow(InTedsTileViewCommonArguments.DataStorage->FindTable(TedsAssetDataHelper::TableView::GetWidgetTableName()));
+		if (OutThumbnailWidgetRowHandle != InvalidRowHandle)
+		{
+			// Referenced Data Row
+			InTedsTileViewCommonArguments.DataStorage->AddColumn(OutThumbnailWidgetRowHandle, FTypedElementRowReferenceColumn{ .Row = InTedsTileViewCommonArguments.TargetRow });
+
+			// Parent widget row
+			InTedsTileViewCommonArguments.DataStorage->AddColumn(OutThumbnailWidgetRowHandle, FTableRowParentColumn{ .Parent = InTedsTileViewCommonArguments.ParentWidgetRowHandle });
+
+			// Padding
+			InTedsTileViewCommonArguments.DataStorage->AddColumn(OutThumbnailWidgetRowHandle, FWidgetPaddingColumn_Experimental{ .Padding = InTedsTileViewCommonArguments.bIsAsset ? FMargin(0.f) : FMargin(5.f) });
+
+			// TODO: this need to be updated when the widget is Selected/Hovered/Dragged or keep it the same if the ThumbnailWidget is kept hovered
+			// Used to override the Thumbnail border image
+			InTedsTileViewCommonArguments.DataStorage->AddColumn(OutThumbnailWidgetRowHandle, FOnGetWidgetSlateBrushColumn_Experimental
+			{
+				.OnGetWidgetSlateBrush = FOnGetWidgetSlateBrush::CreateLambda([this, InTedsTileViewCommonArguments] () { return GetAssetThumbnailBorderOverride(InTedsTileViewCommonArguments); })
+			});
+
+			// Used to retrieve the Thumbnail tooltip to use on the whole TileItem
+			InTedsTileViewCommonArguments.DataStorage->AddColumn(OutThumbnailWidgetRowHandle, FLocalWidgetTooltipColumn_Experimental::StaticStruct());
+
+			// Used to decide on the actual Thumbnail size
+			if (FSizeValueColumn_Experimental* TileItemSizeColumn = InTedsTileViewCommonArguments.DataStorage->GetColumn<FSizeValueColumn_Experimental>(InTedsTileViewCommonArguments.WidgetRow))
+			{
+				InTedsTileViewCommonArguments.DataStorage->AddColumn(OutThumbnailWidgetRowHandle, FSizeValueColumn_Experimental{ .SizeValue = TileItemSizeColumn->SizeValue });
+			}
+
+			ThumbnailWidget = InTedsTileViewCommonArguments.DataStorageUi->ConstructWidget(OutThumbnailWidgetRowHandle, *OutThumbnailWidgetConstructorPtr, ThumbnailMetaView);
+		}
+	}
+
+	return ThumbnailWidget;
+}
+
+TSharedPtr<SWidget> FContentBrowserTileViewWidgetConstructor::CreateItemTypeWidget(FTedsTileViewCommonArguments InTedsTileViewCommonArguments, UE::Editor::DataStorage::RowHandle& OutItemTypeWidgetRowHandle) const
+{
+	using namespace UE::Editor::DataStorage;
+
+	TSharedPtr<FTypedElementWidgetConstructor> OutItemTypeWidgetConstructorPtr;
+
+	auto AssignItemTypeWidgetToColumn = [&OutItemTypeWidgetConstructorPtr] (TUniquePtr<FTypedElementWidgetConstructor> Constructor, TConstArrayView<TWeakObjectPtr<const UScriptStruct>>)
+	{
+		OutItemTypeWidgetConstructorPtr = TSharedPtr<FTypedElementWidgetConstructor>(Constructor.Release());
+		return false;
+	};
+
+	TArray<TWeakObjectPtr<const UScriptStruct>> ItemTypeColumns = GetItemTypeColumns();
+
+	const RowHandle DefaultPurposeRowHandle = InTedsTileViewCommonArguments.DataStorageUi->FindPurpose(InTedsTileViewCommonArguments.DataStorageUi->GetGeneralWidgetPurposeID());
+	InTedsTileViewCommonArguments.DataStorageUi->CreateWidgetConstructors(DefaultPurposeRowHandle, IUiProvider::EMatchApproach::ExactMatch, ItemTypeColumns, FMetaDataView(), AssignItemTypeWidgetToColumn);
+
+	TSharedPtr<SWidget> ItemTypeWidget = SNullWidget::NullWidget;
+	if (OutItemTypeWidgetConstructorPtr)
+	{
+		OutItemTypeWidgetRowHandle = InTedsTileViewCommonArguments.DataStorage->AddRow(InTedsTileViewCommonArguments.DataStorage->FindTable(TedsAssetDataHelper::TableView::GetWidgetTableName()));
+		if (OutItemTypeWidgetRowHandle != InvalidRowHandle)
+		{
+			// Referenced Data Row
+			InTedsTileViewCommonArguments.DataStorage->AddColumn(OutItemTypeWidgetRowHandle, FTypedElementRowReferenceColumn{ .Row = InTedsTileViewCommonArguments.TargetRow });
+
+			// Adding the FontStyle to use
+			InTedsTileViewCommonArguments.DataStorage->AddColumn(OutItemTypeWidgetRowHandle, FFontStyleColumn_Experimental{ .FontInfo = FAppStyle::GetFontStyle("ContentBrowser.AssetTileViewClassNameFont") });
+
+			// Adding the visibility to use
+			InTedsTileViewCommonArguments.DataStorage->AddColumn(OutItemTypeWidgetRowHandle, FWidgetVisibilityColumn_Experimental{ .Visibility = GetAssetClassLabelVisibility(InTedsTileViewCommonArguments.bIsAsset) });
+
+			// Adding the OverflowPolicy to use
+			InTedsTileViewCommonArguments.DataStorage->AddColumn(OutItemTypeWidgetRowHandle, FTextOverflowPolicyColumn_Experimental{ .OverflowPolicy = ETextOverflowPolicy::Ellipsis });
+
+			// TODO: this need to be updated when the widget is Selected/Hovered/Dragged or keep it the same if the ThumbnailWidget is kept hovered
+			// Adding the ColorAndOpacity to use
+			InTedsTileViewCommonArguments.DataStorage->AddColumn(OutItemTypeWidgetRowHandle, FOnGetWidgetColorAndOpacityColumn_Experimental
+				{
+					.OnGetWidgetColorAndOpacity = FOnGetWidgetColorAndOpacity::CreateLambda([this, InTedsTileViewCommonArguments] { return GetAssetClassLabelTextColor(InTedsTileViewCommonArguments); })
+				});
+
+			ItemTypeWidget = InTedsTileViewCommonArguments.DataStorageUi->ConstructWidget(OutItemTypeWidgetRowHandle, *OutItemTypeWidgetConstructorPtr, FMetaDataView());
+		}
+	}
+	return ItemTypeWidget;
+}
+
+TArray<TWeakObjectPtr<const UScriptStruct>> FContentBrowserTileViewWidgetConstructor::GetThumbnailColumns() const
+{
+	static TArray<TWeakObjectPtr<const UScriptStruct>> ThumbnailColumns({ TWeakObjectPtr(FAssetTag::StaticStruct()), TWeakObjectPtr(FFolderTag::StaticStruct()) });
+	return ThumbnailColumns;
+}
+
+TArray<TWeakObjectPtr<const UScriptStruct>> FContentBrowserTileViewWidgetConstructor::GetItemTypeColumns() const
+{
+	static TArray<TWeakObjectPtr<const UScriptStruct>> ItemTypeColumns({ TWeakObjectPtr(FAssetClassColumn::StaticStruct()), TWeakObjectPtr(FFolderTag::StaticStruct()) });
+	return ItemTypeColumns;
+}
+
+#undef LOCTEXT_NAMESPACE
diff --git a/Engine/Plugins/Experimental/EditorDataStorageFeatures/Source/TedsContentBrowser/Private/Widgets/ContentBrowserTileViewWidget.h b/Engine/Plugins/Experimental/EditorDataStorageFeatures/Source/TedsContentBrowser/Private/Widgets/ContentBrowserTileViewWidget.h
new file mode 100644
--- /dev/null
+++ b/Engine/Plugins/Experimental/EditorDataStorageFeatures/Source/TedsContentBrowser/Private/Widgets/ContentBrowserTileViewWidget.h
@@ -0,0 +1,114 @@
+// Copyright Epic Games, Inc. All Rights Reserved.
+
+#pragma once
+
+#include "Elements/Interfaces/TypedElementDataStorageFactory.h"
+#include "Elements/Interfaces/TypedElementDataStorageUiInterface.h"
+#include "Fonts/SlateFontInfo.h"
+#include "Internationalization/Text.h"
+#include "Layout/Visibility.h"
+#include "Styling/SlateColor.h"
+#include "Types/SlateStructs.h"
+#include "UObject/ObjectMacros.h"
+
+#include "ContentBrowserTileViewWidget.generated.h"
+
+class IEditorDataStorageProvider;
+class UScriptStruct;
+struct FSlateBrush;
+
+namespace Purpose
+{
+	FName GetPurposeNamespace();
+	FName GetPurposeName();
+}
+
+UCLASS()
+class UContentBrowserTileViewWidgetFactory : public UEditorDataStorageFactory
+{
+	GENERATED_BODY()
+
+public:
+	~UContentBrowserTileViewWidgetFactory() override = default;
+
+	TEDSCONTENTBROWSER_API virtual void RegisterWidgetPurposes(UE::Editor::DataStorage::IUiProvider& DataStorageUi) const override;
+
+	TEDSCONTENTBROWSER_API void RegisterWidgetConstructors(UE::Editor::DataStorage::ICoreProvider& DataStorage,
+		UE::Editor::DataStorage::IUiProvider& DataStorageUi) const override;
+};
+
+// Content Browser Label + Thumbnail widget
+USTRUCT()
+struct FContentBrowserTileViewWidgetConstructor : public FSimpleWidgetConstructor
+{
+	GENERATED_BODY()
+
+public:
+	TEDSCONTENTBROWSER_API FContentBrowserTileViewWidgetConstructor();
+
+	virtual TSharedPtr<SWidget> CreateWidget(
+		UE::Editor::DataStorage::ICoreProvider* DataStorage,
+		UE::Editor::DataStorage::IUiProvider* DataStorageUi,
+		UE::Editor::DataStorage::RowHandle TargetRow,
+		UE::Editor::DataStorage::RowHandle WidgetRow, 
+		const UE::Editor::DataStorage::FMetaDataView& Arguments) override;
+
+	virtual TConstArrayView<const UScriptStruct*> GetAdditionalColumnsList() const override;
+
+private:
+
+	/** Common data needed for various function */
+	struct FTedsTileViewCommonArguments
+	{
+		UE::Editor::DataStorage::ICoreProvider* DataStorage;
+		UE::Editor::DataStorage::IUiProvider* DataStorageUi;
+		UE::Editor::DataStorage::RowHandle TargetRow;
+		UE::Editor::DataStorage::RowHandle WidgetRow;
+		UE::Editor::DataStorage::RowHandle ParentWidgetRowHandle;
+		bool bIsAsset;
+		TSharedPtr<SWidget> ThumbnailWidget;
+		TSharedPtr<SWidget> TileItem;
+	};
+
+	/** Get the folder shadow image */
+	const FSlateBrush* GetFolderBackgroundShadowImage(FTedsTileViewCommonArguments InTedsTileViewCommonArguments) const;
+
+	/** Get the folder image slot border */
+	const FSlateBrush* GetFolderSlotBorder(FTedsTileViewCommonArguments InTedsTileViewCommonArguments) const;
+
+	/** Get the asset thumbnail border to use */
+	const FSlateBrush* GetAssetThumbnailBorderOverride(FTedsTileViewCommonArguments InTedsTileViewCommonArguments) const;
+
+	/** Get the name are background image */
+	const FSlateBrush* GetNameAreaBackgroundImage(FTedsTileViewCommonArguments InTedsTileViewCommonArguments) const;
+
+	/** Get the font to use for the thumbnail label */
+	FSlateFontInfo GetThumbnailFont() const;
+
+	/** Get the name area text color */
+	FSlateColor GetNameAreaTextColor(FTedsTileViewCommonArguments InTedsTileViewCommonArguments) const;
+
+	/** Get the maximum height for the name area */
+	FOptionalSize GetNameAreaMaxDesiredHeight() const;
+
+	/** Gets the visibility of the asset class label in thumbnails */
+	EVisibility GetAssetClassLabelVisibility(bool bIsAsset) const;
+
+	/** Gets the color and opacity of the asset class type */
+	FSlateColor GetAssetClassLabelTextColor(FTedsTileViewCommonArguments InTedsTileViewCommonArguments) const;
+
+	/** Gets the background image for folders */
+	const FSlateBrush* GetFolderBackgroundImage(FTedsTileViewCommonArguments InTedsTileViewCommonArguments) const;
+
+	/** Create the thumbnail widget through TEDS */
+	TSharedPtr<SWidget> CreateThumbnailWidget(FTedsTileViewCommonArguments InTedsTileViewCommonArguments, UE::Editor::DataStorage::RowHandle& OutThumbnailWidgetRowHandle) const;
+
+	/** Create the item type widget through TEDS */
+	TSharedPtr<SWidget> CreateItemTypeWidget(FTedsTileViewCommonArguments InTedsTileViewCommonArguments, UE::Editor::DataStorage::RowHandle& OutItemTypeWidgetRowHandle) const;
+
+	/** Get the Thumbnail columns */
+	TArray<TWeakObjectPtr<const UScriptStruct>> GetThumbnailColumns() const;
+	
+	/** Get the ItemType columns */
+	TArray<TWeakObjectPtr<const UScriptStruct>> GetItemTypeColumns() const;
+};
diff --git a/Engine/Plugins/Experimental/EditorDataStorageFeatures/Source/TedsTableViewer/Private/Widgets/STedsTableViewer.cpp b/Engine/Plugins/Experimental/EditorDataStorageFeatures/Source/TedsTableViewer/Private/Widgets/STedsTableViewer.cpp
--- a/Engine/Plugins/Experimental/EditorDataStorageFeatures/Source/TedsTableViewer/Private/Widgets/STedsTableViewer.cpp
+++ b/Engine/Plugins/Experimental/EditorDataStorageFeatures/Source/TedsTableViewer/Private/Widgets/STedsTableViewer.cpp
@@ -20,7 +20,9 @@ namespace UE::Editor::DataStorage
 	{
 		OnSelectionChanged = InArgs._OnSelectionChanged;
 		EmptyRowsMessage = InArgs._EmptyRowsMessage;
-		
+		ItemHeight = InArgs._ItemHeight;
+		ItemPadding = InArgs._ItemPadding;
+
 		IUiProvider::FPurposeID CellWigetPurpose = InArgs._CellWidgetPurpose;
 
 		if (!CellWigetPurpose.IsSet())
@@ -208,6 +210,17 @@ namespace UE::Editor::DataStorage
 		return AsShared();
 	}
 
+	bool STedsTableViewer::IsSelected(RowHandle InRow) const
+	{
+		FTedsRowHandle TedsRowHandle{ .RowHandle = InRow };
+		return ListView->IsItemSelected(TedsRowHandle);
+	}
+
+	bool STedsTableViewer::IsSelectedExclusively(RowHandle InRow) const
+	{
+		return IsSelected(InRow) && ListView->GetNumItemsSelected() == 1;
+	}
+
 	bool STedsTableViewer::IsItemVisible(TableViewerItemPtr InItem) const
 	{
 		return ListView->IsItemVisible(InItem);
@@ -216,6 +229,9 @@ namespace UE::Editor::DataStorage
 	TSharedRef<ITableRow> STedsTableViewer::MakeTableRowWidget(TableViewerItemPtr InItem, const TSharedRef<STableViewBase>& OwnerTable) const
 	{
 		return SNew(STedsTableViewerRow, OwnerTable, Model.ToSharedRef())
+				.ItemHeight(ItemHeight)
+				.Padding(ItemPadding)
+				.ParentWidgetRowHandle(GetWidgetRowHandle())
 				.Item(InItem);
 	}
 } // namespace UE::Editor::DataStorage
diff --git a/Engine/Plugins/Experimental/EditorDataStorageFeatures/Source/TedsTableViewer/Private/Widgets/STedsTableViewerRow.cpp b/Engine/Plugins/Experimental/EditorDataStorageFeatures/Source/TedsTableViewer/Private/Widgets/STedsTableViewerRow.cpp
--- a/Engine/Plugins/Experimental/EditorDataStorageFeatures/Source/TedsTableViewer/Private/Widgets/STedsTableViewerRow.cpp
+++ b/Engine/Plugins/Experimental/EditorDataStorageFeatures/Source/TedsTableViewer/Private/Widgets/STedsTableViewerRow.cpp
@@ -4,6 +4,7 @@
 #include "Widgets/STedsTableViewerRow.h"
 
 #include "TedsTableViewerColumn.h"
+#include "Elements/Columns/TypedElementHiearchyColumns.h"
 #include "Widgets/STedsTableViewer.h"
 
 namespace UE::Editor::DataStorage
@@ -12,8 +13,11 @@ namespace UE::Editor::DataStorage
 	{
 		Item = InArgs._Item;
 		TableViewerModel = InTableViewerModel;
+		ParentWidgetRowHandle = InArgs._ParentWidgetRowHandle;
+		ItemHeight = InArgs._ItemHeight;
 
 		const auto Args = FSuperRowType::FArguments()
+		.Padding(InArgs._Padding)
 		.Style(&FAppStyle::Get().GetWidgetStyle<FTableRowStyle>("SceneOutliner.TableViewRow"));
 
 		SMultiColumnTableRow<TableViewerItemPtr>::Construct(Args, OwnerTableView);
@@ -23,10 +27,15 @@ namespace UE::Editor::DataStorage
 	{
 		if(TSharedPtr<FTedsTableViewerColumn> Column = TableViewerModel->GetColumn(ColumnName))
 		{
-			if(TSharedPtr<SWidget> RowWidget = Column->ConstructRowWidget(Item))
+			auto WidgetRowSetupDelegate = [this] (ICoreProvider& InStorage, const RowHandle& InUIRowHandle)
+			{
+				InStorage.AddColumn(InUIRowHandle, FTableRowParentColumn{ .Parent = ParentWidgetRowHandle });
+			};
+
+			if (TSharedPtr<SWidget> RowWidget = Column->ConstructRowWidget(Item, WidgetRowSetupDelegate))
 			{
 				return SNew(SBox)
-						.MinDesiredHeight(20.0f)
+						.HeightOverride(this, &STedsTableViewerRow::GetCurrentItemHeight)
 						.VAlign(VAlign_Center)
 						[
 							RowWidget.ToSharedRef()
@@ -37,6 +46,12 @@ namespace UE::Editor::DataStorage
 		return SNullWidget::NullWidget;
 	}
 
+	FOptionalSize STedsTableViewerRow::GetCurrentItemHeight() const
+	{
+		static constexpr float ItemHeightDefaultValue = 48.f;
+		return ItemHeight.IsSet() ? ItemHeight.Get() : ItemHeightDefaultValue;
+	}
+
 	TSharedRef<SWidget> SHierarchyViewerRow::GenerateWidgetForColumn(const FName& ColumnName)
 	{
 		TSharedRef<SWidget> ActualWidget = STedsTableViewerRow::GenerateWidgetForColumn(ColumnName);
diff --git a/Engine/Plugins/Experimental/EditorDataStorageFeatures/Source/TedsTableViewer/Private/Widgets/STedsTileViewer.cpp b/Engine/Plugins/Experimental/EditorDataStorageFeatures/Source/TedsTableViewer/Private/Widgets/STedsTileViewer.cpp
--- a/Engine/Plugins/Experimental/EditorDataStorageFeatures/Source/TedsTableViewer/Private/Widgets/STedsTileViewer.cpp
+++ b/Engine/Plugins/Experimental/EditorDataStorageFeatures/Source/TedsTableViewer/Private/Widgets/STedsTileViewer.cpp
@@ -10,6 +10,8 @@
 #include "Elements/Framework/TypedElementAttributeBinding.h"
 #include "Elements/Framework/TypedElementDataStorageWidget.h"
 #include "TedsTableViewerColumn.h"
+#include "TedsTableViewerWidgetColumns.h"
+#include "Elements/Columns/TypedElementHiearchyColumns.h"
 #include "Styling/StyleColors.h"
 #include "Widgets/STedsTableViewerRow.h"
 
@@ -21,6 +23,7 @@ namespace UE::Editor::DataStorage
 		WidgetPurpose = InArgs._WidgetPurpose;
 		Columns = InArgs._Columns;
 		TableRowStyle = InArgs._TileStyle;
+		TilePadding = InArgs._TilePadding;
 
 		// Use the default purpose if the user didn't specify any
 		// TODO: Have a better "default" widget for tiles instead of the default cell widget purpose
@@ -49,13 +52,14 @@ namespace UE::Editor::DataStorage
 				TedsWidget->AsWidget()
 			];
 		}
-		
+
 		AddWidgetColumns();
 		
 		// Attribute binder to bind widget columns to attributes on the ListView
 		FAttributeBinder Binder(TedsWidget->GetRowHandle());
 
 		TileView = SNew(STileView<TableViewerItemPtr>)
+			.ItemAlignment(InArgs._ItemAlignment)
 			.ListItemsSource(&Model->GetItems())
 			.OnGenerateTile(this, &STedsTileViewer::MakeTileWidget)
 			.OnSelectionChanged(this, &STedsTileViewer::OnListSelectionChanged)
@@ -150,19 +154,37 @@ namespace UE::Editor::DataStorage
 		return AsShared();
 	}
 
+	bool STedsTileViewer::IsSelected(RowHandle InRow) const
+	{
+		FTedsRowHandle TedsRowHandle{ .RowHandle = InRow };
+		return TileView->IsItemSelected(TedsRowHandle);
+	}
+
+	bool STedsTileViewer::IsSelectedExclusively(RowHandle InRow) const
+	{
+		return IsSelected(InRow) && TileView->GetNumItemsSelected() == 1;
+	}
+
 	TSharedRef<ITableRow> STedsTileViewer::MakeTileWidget(TableViewerItemPtr InItem,
 	                                                      const TSharedRef<STableViewBase>& OwnerTable) const
 	{
 		TSharedPtr<SWidget> RowWidget = SNullWidget::NullWidget;
 		
-		if(TSharedPtr<FTedsTableViewerColumn> Column = Model->GetColumn("TileView"))
+		if (TSharedPtr<FTedsTableViewerColumn> Column = Model->GetColumn("TileView"))
 		{
-			RowWidget = Column->ConstructRowWidget(InItem);
+			RowWidget = Column->ConstructRowWidget(InItem,
+				[&](ICoreProvider& DataStorage, const RowHandle& WidgetRow)
+				{
+					DataStorage.AddColumn(WidgetRow, FTableRowParentColumn{ .Parent = GetWidgetRowHandle() });
+					DataStorage.AddColumn(WidgetRow, FExternalWidgetSelectionColumn{ .IsSelected = FIsSelected::CreateSP(this, &STedsTileViewer::IsSelected, InItem.RowHandle) });
+					DataStorage.AddColumn(WidgetRow, FExternalWidgetExclusiveSelectionColumn{ .IsSelectedExclusively = FIsSelected::CreateSP(this, &STedsTileViewer::IsSelectedExclusively, InItem.RowHandle) });
+				});
 		}
 
 		return SNew(STableRow<TableViewerItemPtr>, OwnerTable)
-			.Padding(4.f) // TODO: Make modular
+			.Padding(TilePadding)
 			.Style(TableRowStyle)
+			.Cursor(bAllowDragging ? EMouseCursor::GrabHand : EMouseCursor::Default)
 			[
 				RowWidget.ToSharedRef()
 			];
@@ -186,8 +208,8 @@ namespace UE::Editor::DataStorage
 		if(ICoreProvider* DataStorage = Model->GetDataStorageInterface())
 		{
 			const RowHandle WidgetRowHandle = TedsWidget->GetRowHandle();
-		
-			if(DataStorage->IsRowAvailable(WidgetRowHandle))
+
+			if (DataStorage->IsRowAvailable(WidgetRowHandle))
 			{
 				// The table viewer should not show up as a row in a table viewer because that will cause all sorts of recursion issues
 				DataStorage->AddColumn(WidgetRowHandle, FHideRowFromUITag::StaticStruct());
diff --git a/Engine/Plugins/Experimental/EditorDataStorageFeatures/Source/TedsTableViewer/Public/ITedsTableViewer.h b/Engine/Plugins/Experimental/EditorDataStorageFeatures/Source/TedsTableViewer/Public/ITedsTableViewer.h
--- a/Engine/Plugins/Experimental/EditorDataStorageFeatures/Source/TedsTableViewer/Public/ITedsTableViewer.h
+++ b/Engine/Plugins/Experimental/EditorDataStorageFeatures/Source/TedsTableViewer/Public/ITedsTableViewer.h
@@ -31,6 +31,12 @@ namespace UE::Editor::DataStorage
     	// Get the table viewer as an SWidget
     	virtual TSharedRef<SWidget> AsWidget() = 0;
 
+    	// Get if a specific row is selected
+    	virtual bool IsSelected(RowHandle InRow) const = 0;
+
+    	// Get if a specific row is selected exclusively
+    	virtual bool IsSelectedExclusively(RowHandle InRow) const = 0;
+   
     	virtual ~ITableViewer() = default;
     };
 }
\ No newline at end of file
diff --git a/Engine/Plugins/Experimental/EditorDataStorageFeatures/Source/TedsTableViewer/Public/TedsTableViewerWidgetColumns.h b/Engine/Plugins/Experimental/EditorDataStorageFeatures/Source/TedsTableViewer/Public/TedsTableViewerWidgetColumns.h
new file mode 100644
--- /dev/null
+++ b/Engine/Plugins/Experimental/EditorDataStorageFeatures/Source/TedsTableViewer/Public/TedsTableViewerWidgetColumns.h
@@ -0,0 +1,26 @@
+﻿// Copyright Epic Games, Inc. All Rights Reserved.
+
+#pragma once
+
+#include "Elements/Common/TypedElementCommonTypes.h"
+#include "Framework/SlateDelegates.h"
+
+#include "TedsTableViewerWidgetColumns.generated.h"
+
+/**
+* Column added to a widget row when an external widget manages exclusive selection for the widget referenced by the row,
+ * such as an owning SListView or STreeView
+ */
+USTRUCT(meta = (DisplayName = "Widget with externally managed exclusive selection"))
+struct FExternalWidgetExclusiveSelectionColumn : public FEditorDataStorageColumn
+{
+	GENERATED_BODY()
+	/** Delegate to execute to check the status of if the widget is the only selected between multiple widgets
+	 * Only needs to be hooked up if an external widget is managing selection on multiple widgets, such
+	 * as a SListView or STreeView.
+	 *
+	 * E.g usage: Enables renaming by clicking the label twice if the widget is editable.
+	 * @see SInlineEditableTextBlock::IsSelected
+	 */
+	FIsSelected IsSelectedExclusively;
+};
diff --git a/Engine/Plugins/Experimental/EditorDataStorageFeatures/Source/TedsTableViewer/Public/Widgets/STedsTableViewer.h b/Engine/Plugins/Experimental/EditorDataStorageFeatures/Source/TedsTableViewer/Public/Widgets/STedsTableViewer.h
--- a/Engine/Plugins/Experimental/EditorDataStorageFeatures/Source/TedsTableViewer/Public/Widgets/STedsTableViewer.h
+++ b/Engine/Plugins/Experimental/EditorDataStorageFeatures/Source/TedsTableViewer/Public/Widgets/STedsTableViewer.h
@@ -40,6 +40,8 @@ namespace UE::Editor::DataStorage
 
 		SLATE_BEGIN_ARGS(STedsTableViewer)
 			: _ListSelectionMode(ESelectionMode::Type::Single)
+			, _ItemHeight(48.f)
+			, _ItemPadding(FMargin(0.f))
 		{
 			
 		}
@@ -68,6 +70,12 @@ namespace UE::Editor::DataStorage
 		// Empty = simply show the column headers instead of a message
 		SLATE_ATTRIBUTE(FText, EmptyRowsMessage)
 
+		/** The height of the list item */
+		SLATE_ATTRIBUTE(float, ItemHeight)
+
+		/** The padding of the list item */
+		SLATE_ATTRIBUTE(FMargin, ItemPadding)
+
 		SLATE_END_ARGS()
 
 	public:
@@ -96,6 +104,9 @@ namespace UE::Editor::DataStorage
 		
 		TEDSTABLEVIEWER_API virtual TSharedRef<SWidget> AsWidget() override;
 
+		TEDSTABLEVIEWER_API virtual bool IsSelected(RowHandle InRow) const override;
+
+		TEDSTABLEVIEWER_API virtual bool IsSelectedExclusively(RowHandle InRow) const override;
 	protected:
 		
 		TSharedRef<ITableRow> MakeTableRowWidget(TableViewerItemPtr InItem, const TSharedRef<STableViewBase>& OwnerTable) const;
@@ -111,6 +122,11 @@ namespace UE::Editor::DataStorage
 		void AddWidgetColumns();
 
 	private:
+		// The actual height of each item
+		TAttribute<float> ItemHeight;
+
+		// The actual padding between each item
+		TAttribute<FMargin> ItemPadding;
 
 		// The actual ListView widget that displays the rows
 		TSharedPtr<SListView<TableViewerItemPtr>> ListView;
diff --git a/Engine/Plugins/Experimental/EditorDataStorageFeatures/Source/TedsTableViewer/Public/Widgets/STedsTableViewerRow.h b/Engine/Plugins/Experimental/EditorDataStorageFeatures/Source/TedsTableViewer/Public/Widgets/STedsTableViewerRow.h
--- a/Engine/Plugins/Experimental/EditorDataStorageFeatures/Source/TedsTableViewer/Public/Widgets/STedsTableViewerRow.h
+++ b/Engine/Plugins/Experimental/EditorDataStorageFeatures/Source/TedsTableViewer/Public/Widgets/STedsTableViewerRow.h
@@ -17,11 +17,22 @@ namespace UE::Editor::DataStorage
 
 	public:
 
-		SLATE_BEGIN_ARGS( STedsTableViewerRow ) {}
+		SLATE_BEGIN_ARGS( STedsTableViewerRow )
+			: _ParentWidgetRowHandle(InvalidRowHandle)
+			, _Padding(FMargin(0.f, 4.f))
+			, _ItemHeight(48.f)
+		{}
 
 		/** The list item for this row */
 		SLATE_ARGUMENT( TableViewerItemPtr, Item )
 
+		SLATE_ARGUMENT( RowHandle, ParentWidgetRowHandle )
+
+		SLATE_ATTRIBUTE( FMargin, Padding )
+			
+		/** The height of the list item */
+		SLATE_ATTRIBUTE( float, ItemHeight )
+
 		SLATE_END_ARGS()
 		
 		/** Construct function for this widget */
@@ -30,8 +41,13 @@ namespace UE::Editor::DataStorage
 		/** Overridden from SMultiColumnTableRow.  Generates a widget for this column of the tree row. */
 		virtual TSharedRef<SWidget> GenerateWidgetForColumn( const FName& ColumnName ) override;
 
-	protected:
+	private:
+		/** Get the current ItemHeight as an OptionalSize for the Box HeightOverride */
+		FOptionalSize GetCurrentItemHeight() const;
 
+	protected:
+		TAttribute<float> ItemHeight;
+		RowHandle ParentWidgetRowHandle = InvalidRowHandle;
 		TSharedPtr<FTedsTableViewerModel> TableViewerModel;
 		TableViewerItemPtr Item;
 	};
diff --git a/Engine/Plugins/Experimental/EditorDataStorageFeatures/Source/TedsTableViewer/Public/Widgets/STedsTileViewer.h b/Engine/Plugins/Experimental/EditorDataStorageFeatures/Source/TedsTableViewer/Public/Widgets/STedsTileViewer.h
--- a/Engine/Plugins/Experimental/EditorDataStorageFeatures/Source/TedsTableViewer/Public/Widgets/STedsTileViewer.h
+++ b/Engine/Plugins/Experimental/EditorDataStorageFeatures/Source/TedsTableViewer/Public/Widgets/STedsTileViewer.h
@@ -37,6 +37,8 @@ namespace UE::Editor::DataStorage
 			, _ItemHeight(128)
 			, _ItemWidth(128)
 			, _TileStyle(&FAppStyle::Get().GetWidgetStyle<FTableRowStyle>("TableView.Row"))
+			, _TilePadding(FMargin(0.f))
+			, _ItemAlignment(EListItemAlignment::EvenlyDistributed)
 		{
 			
 		}
@@ -55,13 +57,22 @@ namespace UE::Editor::DataStorage
 			
 		// The selection mode for the table viewer (single/multi etc)
 		SLATE_ARGUMENT(ESelectionMode::Type, SelectionMode)
-			
+
+		// Height of each tile
 		SLATE_ATTRIBUTE( float, ItemHeight )
 
+		// Width of each tile
 		SLATE_ATTRIBUTE( float, ItemWidth )
-			
+
+		// Style of each tile
 		SLATE_STYLE_ARGUMENT( FTableRowStyle, TileStyle )
 
+		// Padding between each tile
+		SLATE_ARGUMENT( FMargin, TilePadding )
+
+		// Tile View alignment
+		SLATE_ATTRIBUTE( EListItemAlignment, ItemAlignment )
+
 		SLATE_END_ARGS()
 
 	public:
@@ -86,6 +97,10 @@ namespace UE::Editor::DataStorage
 		// Get the table viewer as an SWidget
 		TEDSTABLEVIEWER_API virtual TSharedRef<SWidget> AsWidget() override;
 
+		TEDSTABLEVIEWER_API virtual bool IsSelected(RowHandle InRow) const override;
+
+		TEDSTABLEVIEWER_API virtual bool IsSelectedExclusively(RowHandle InRow) const override;
+
 	protected:
 		
 		TSharedRef<ITableRow> MakeTileWidget(TableViewerItemPtr InItem, const TSharedRef<STableViewBase>& OwnerTable) const;
@@ -100,6 +115,9 @@ namespace UE::Editor::DataStorage
 
 	private:
 
+		// Whether dragging is allowed
+		bool bAllowDragging = true;
+
 		// The actual ListView widget that displays the rows
 		TSharedPtr<STileView<TableViewerItemPtr>> TileView;
 
@@ -120,5 +138,8 @@ namespace UE::Editor::DataStorage
 
 		// Style name to use for the rows (tiles)
 		const FTableRowStyle* TableRowStyle = nullptr;
+
+		// Padding between each tile
+		FMargin TilePadding;
 	};
 }
diff --git a/Engine/Source/Editor/ContentBrowser/Public/Experimental/ContentBrowserViewExtender.h b/Engine/Source/Editor/ContentBrowser/Public/Experimental/ContentBrowserViewExtender.h
--- a/Engine/Source/Editor/ContentBrowser/Public/Experimental/ContentBrowserViewExtender.h
+++ b/Engine/Source/Editor/ContentBrowser/Public/Experimental/ContentBrowserViewExtender.h
@@ -2,6 +2,7 @@
 
 #pragma once
 
+#include "AssetThumbnail.h"
 #include "AssetViewTypes.h"
 #include "Delegates/DelegateCombinations.h"
 #include "Framework/SlateDelegates.h"
@@ -62,4 +63,12 @@ public:
 	// Get a tooltip for the custom view
 	virtual FText GetViewTooltipText() = 0;
 
+	// True if zooming and size changing can be done in the custom view, set it to false if not
+	virtual bool SupportScaling() const = 0;
+
+	// True if zooming and size changing can be done in the custom view, set it to false if not
+	virtual bool SupportZooming() const = 0;
+
+	// True if thumbnail edit mode should be supported, set it to false if not
+	virtual bool SupportThumbnailEditMode() const = 0;
 };
\ No newline at end of file
