[EditorDataStorageFeatures] Issue 857364: PVS Studio warnings

diff --git a/Engine/Plugins/Experimental/EditorDataStorageFeatures/Source/TedsContentBrowser/Private/Widgets/ContentBrowserTileViewWidget.cpp b/Engine/Plugins/Experimental/EditorDataStorageFeatures/Source/TedsContentBrowser/Private/Widgets/ContentBrowserTileViewWidget.cpp
--- a/Engine/Plugins/Experimental/EditorDataStorageFeatures/Source/TedsContentBrowser/Private/Widgets/ContentBrowserTileViewWidget.cpp
+++ b/Engine/Plugins/Experimental/EditorDataStorageFeatures/Source/TedsContentBrowser/Private/Widgets/ContentBrowserTileViewWidget.cpp
@@ -279,7 +279,7 @@ const FSlateBrush* FContentBrowserTileViewWidgetConstructor::GetFolderBackground
 	FIsSelected IsSelected = WidgetRowBinder.BindEvent(&FExternalWidgetSelectionColumn::IsSelected);
 
 	const bool bIsSelected = IsSelected.IsBound() ? IsSelected.Execute() : false;
-	const bool bIsHoveredOrDraggedOver = (InTedsTileViewCommonArguments.TileItem.IsValid() && InTedsTileViewCommonArguments.TileItem->IsHovered()) || false /*bDraggedOver*/;
+	const bool bIsHoveredOrDraggedOver = (InTedsTileViewCommonArguments.TileItem.IsValid() && InTedsTileViewCommonArguments.TileItem->IsHovered()) /*|| bDraggedOver*/;
 
 	if (bIsSelected || bIsHoveredOrDraggedOver)
 	{
@@ -302,7 +302,7 @@ const FSlateBrush* FContentBrowserTileViewWidgetConstructor::GetFolderSlotBorder
 	FIsSelected IsSelected = WidgetRowBinder.BindEvent(&FExternalWidgetSelectionColumn::IsSelected);
 
 	const bool bIsSelected = IsSelected.IsBound() ? IsSelected.Execute() : false;
-	const bool bIsHoveredOrDraggedOver = (InTedsTileViewCommonArguments.TileItem.IsValid() && InTedsTileViewCommonArguments.TileItem->IsHovered()) || false /*bDraggedOver*/;
+	const bool bIsHoveredOrDraggedOver = (InTedsTileViewCommonArguments.TileItem.IsValid() && InTedsTileViewCommonArguments.TileItem->IsHovered()) /*|| bDraggedOver*/;
 
 	if (bIsSelected || bIsHoveredOrDraggedOver)
 	{
@@ -339,7 +339,7 @@ const FSlateBrush* FContentBrowserTileViewWidgetConstructor::GetNameAreaBackgrou
 
 	const bool bIsSelected = IsSelected.IsBound() ? IsSelected.Execute() : false;
 	// TODO: ThumbnailWidget is the TedsWidget and not the actual Thumbnail widget, need to check the hovered state of the actual thumbnail widget, needed for the ThumbnailEditMode to keep it hovered
-	const bool bIsHoveredOrDraggedOver = (InTedsTileViewCommonArguments.TileItem.IsValid() && InTedsTileViewCommonArguments.TileItem->IsHovered()) || false /*bDraggedOver*/ || false /*(InThumbnailWidget.IsValid() && InThumbnailWidget->IsHovered())*/;
+	const bool bIsHoveredOrDraggedOver = (InTedsTileViewCommonArguments.TileItem.IsValid() && InTedsTileViewCommonArguments.TileItem->IsHovered()) /*|| bDraggedOver || (InThumbnailWidget.IsValid() && InThumbnailWidget->IsHovered())*/;
 	if (bIsSelected && bIsHoveredOrDraggedOver)
 	{
 		return FAppStyle::Get().GetBrush(SelectedHover);
@@ -378,7 +378,7 @@ FSlateColor FContentBrowserTileViewWidgetConstructor::GetNameAreaTextColor(FTeds
 
 	const bool bIsSelected = IsSelected.IsBound() ? IsSelected.Execute() : false;
 	// TODO: ThumbnailWidget is the TedsWidget and not the actual Thumbnail widget, need to check the hovered state of the actual thumbnail widget, needed for the ThumbnailEditMode to keep it hovered
-	const bool bIsHoveredOrDraggedOver = (InTedsTileViewCommonArguments.TileItem.IsValid() && InTedsTileViewCommonArguments.TileItem->IsHovered()) || false /*bDraggedOver*/ || false /*(InThumbnailWidget.IsValid() && InThumbnailWidget->IsHovered())*/;
+	const bool bIsHoveredOrDraggedOver = (InTedsTileViewCommonArguments.TileItem.IsValid() && InTedsTileViewCommonArguments.TileItem->IsHovered()) /*|| bDraggedOver || (InThumbnailWidget.IsValid() && InThumbnailWidget->IsHovered())*/;
 	if (bIsSelected || bIsHoveredOrDraggedOver)
 	{
 		return FStyleColors::White;
@@ -400,7 +400,7 @@ const FSlateBrush* FContentBrowserTileViewWidgetConstructor::GetAssetThumbnailBo
 
 	const bool bIsSelected = IsSelected.IsBound() ? IsSelected.Execute() : false;
 	// TODO: ThumbnailWidget is the TedsWidget and not the actual Thumbnail widget, need to check the hovered state of the actual thumbnail widget, needed for the ThumbnailEditMode to keep it hovered
-	const bool bIsHoveredOrDraggedOver = (InTedsTileViewCommonArguments.TileItem.IsValid() && InTedsTileViewCommonArguments.TileItem->IsHovered()) || false /*bDraggedOver*/ || false /*(InThumbnailWidget.IsValid() && InThumbnailWidget->IsHovered())*/;
+	const bool bIsHoveredOrDraggedOver = (InTedsTileViewCommonArguments.TileItem.IsValid() && InTedsTileViewCommonArguments.TileItem->IsHovered()) /*|| bDraggedOver || (InThumbnailWidget.IsValid() && InThumbnailWidget->IsHovered())*/;
 	if (bIsSelected && bIsHoveredOrDraggedOver)
 	{
 		static const FLazyName SelectedHover("ContentBrowser.AssetTileItem.AssetBorderSelectedHoverBackground");
@@ -449,7 +449,7 @@ FSlateColor FContentBrowserTileViewWidgetConstructor::GetAssetClassLabelTextColo
 
 	const bool bIsSelected = IsSelected.IsBound() ? IsSelected.Execute() : false;
 	// TODO: ThumbnailWidget is the TedsWidget and not the actual Thumbnail widget, need to check the hovered state of the actual thumbnail widget, needed for the ThumbnailEditMode to keep it hovered
-	const bool bIsHoveredOrDraggedOver = (InTedsTileViewCommonArguments.TileItem.IsValid() && InTedsTileViewCommonArguments.TileItem->IsHovered()) || false /*bDraggedOver*/ || false /*(InThumbnailWidget.IsValid() && InThumbnailWidget->IsHovered())*/;
+	const bool bIsHoveredOrDraggedOver = (InTedsTileViewCommonArguments.TileItem.IsValid() && InTedsTileViewCommonArguments.TileItem->IsHovered()) /*|| bDraggedOver || (InThumbnailWidget.IsValid() && InThumbnailWidget->IsHovered())*/;
 	if (bIsSelected || bIsHoveredOrDraggedOver)
 	{
 		return FStyleColors::White;
@@ -470,7 +470,7 @@ const FSlateBrush* FContentBrowserTileViewWidgetConstructor::GetFolderBackground
 	FIsSelected IsSelected = WidgetRowBinder.BindEvent(&FExternalWidgetSelectionColumn::IsSelected);
 
 	const bool bIsSelected = IsSelected.IsBound() ? IsSelected.Execute() : false;
-	const bool bIsHoveredOrDraggedOver = (InTedsTileViewCommonArguments.TileItem.IsValid() && InTedsTileViewCommonArguments.TileItem->IsHovered()) || false /*bDraggedOver*/;
+	const bool bIsHoveredOrDraggedOver = (InTedsTileViewCommonArguments.TileItem.IsValid() && InTedsTileViewCommonArguments.TileItem->IsHovered()) /*|| bDraggedOver*/;
 
 	if (bIsSelected && bIsHoveredOrDraggedOver)
 	{
