SHierarchyViewer: Implement some new pure virtual functions

diff --git a/Engine/Plugins/Experimental/EditorDataStorageFeatures/Source/TedsTableViewer/Private/Widgets/STedsHierarchyViewer.cpp b/Engine/Plugins/Experimental/EditorDataStorageFeatures/Source/TedsTableViewer/Private/Widgets/STedsHierarchyViewer.cpp
--- a/Engine/Plugins/Experimental/EditorDataStorageFeatures/Source/TedsTableViewer/Private/Widgets/STedsHierarchyViewer.cpp
+++ b/Engine/Plugins/Experimental/EditorDataStorageFeatures/Source/TedsTableViewer/Private/Widgets/STedsHierarchyViewer.cpp
@@ -231,6 +231,17 @@ namespace UE::Editor::DataStorage
 		return AsShared();
 	}
 
+	bool SHierarchyViewer::IsSelected(RowHandle InRow) const
+	{
+		FTedsRowHandle TedsRowHandle{ .RowHandle = InRow };
+		return TreeView->IsItemSelected(TedsRowHandle);
+	}
+
+	bool SHierarchyViewer::IsSelectedExclusively(RowHandle InRow) const
+	{
+		return IsSelected(InRow) && TreeView->GetNumItemsSelected() == 1;
+	}
+
 	bool SHierarchyViewer::IsItemVisible(TableViewerItemPtr InItem) const
 	{
 		return TreeView->IsItemVisible(InItem);
diff --git a/Engine/Plugins/Experimental/EditorDataStorageFeatures/Source/TedsTableViewer/Public/Widgets/STedsHierarchyViewer.h b/Engine/Plugins/Experimental/EditorDataStorageFeatures/Source/TedsTableViewer/Public/Widgets/STedsHierarchyViewer.h
--- a/Engine/Plugins/Experimental/EditorDataStorageFeatures/Source/TedsTableViewer/Public/Widgets/STedsHierarchyViewer.h
+++ b/Engine/Plugins/Experimental/EditorDataStorageFeatures/Source/TedsTableViewer/Public/Widgets/STedsHierarchyViewer.h
@@ -99,6 +99,13 @@ namespace UE::Editor::DataStorage
 		
 		TEDSTABLEVIEWER_API virtual TSharedRef<SWidget> AsWidget() override;
 
+		// Get if a specific row is selected
+		TEDSTABLEVIEWER_API virtual bool IsSelected(RowHandle InRow) const override;
+
+		// Get if a specific row is selected exclusively
+		TEDSTABLEVIEWER_API virtual bool IsSelectedExclusively(RowHandle InRow) const override;
+
+
 	protected:
 		
 		TSharedRef<ITableRow> MakeTableRowWidget(TableViewerItemPtr InItem, const TSharedRef<STableViewBase>& OwnerTable) const;
