[EditorDataStorageFeatures] Change STedsTableViewer to not consider the ItemHeight if not given

diff --git a/Engine/Plugins/Experimental/EditorDataStorageFeatures/Source/TedsTableViewer/Private/Widgets/STedsTableViewerRow.cpp b/Engine/Plugins/Experimental/EditorDataStorageFeatures/Source/TedsTableViewer/Private/Widgets/STedsTableViewerRow.cpp
--- a/Engine/Plugins/Experimental/EditorDataStorageFeatures/Source/TedsTableViewer/Private/Widgets/STedsTableViewerRow.cpp
+++ b/Engine/Plugins/Experimental/EditorDataStorageFeatures/Source/TedsTableViewer/Private/Widgets/STedsTableViewerRow.cpp
@@ -36,6 +36,7 @@ namespace UE::Editor::DataStorage
 			{
 				return SNew(SBox)
 						.HeightOverride(this, &STedsTableViewerRow::GetCurrentItemHeight)
+						.MinDesiredHeight(20.f)
 						.VAlign(VAlign_Center)
 						[
 							RowWidget.ToSharedRef()
@@ -48,8 +49,7 @@ namespace UE::Editor::DataStorage
 
 	FOptionalSize STedsTableViewerRow::GetCurrentItemHeight() const
 	{
-		static constexpr float ItemHeightDefaultValue = 48.f;
-		return ItemHeight.IsSet() ? ItemHeight.Get() : ItemHeightDefaultValue;
+		return ItemHeight.IsSet() ? ItemHeight.Get() : FOptionalSize();
 	}
 
 	TSharedRef<SWidget> SHierarchyViewerRow::GenerateWidgetForColumn(const FName& ColumnName)
diff --git a/Engine/Plugins/Experimental/EditorDataStorageFeatures/Source/TedsTableViewer/Public/Widgets/STedsTableViewer.h b/Engine/Plugins/Experimental/EditorDataStorageFeatures/Source/TedsTableViewer/Public/Widgets/STedsTableViewer.h
--- a/Engine/Plugins/Experimental/EditorDataStorageFeatures/Source/TedsTableViewer/Public/Widgets/STedsTableViewer.h
+++ b/Engine/Plugins/Experimental/EditorDataStorageFeatures/Source/TedsTableViewer/Public/Widgets/STedsTableViewer.h
@@ -40,7 +40,7 @@ namespace UE::Editor::DataStorage
 
 		SLATE_BEGIN_ARGS(STedsTableViewer)
 			: _ListSelectionMode(ESelectionMode::Type::Single)
-			, _ItemHeight(48.f)
+			, _ItemHeight()
 			, _ItemPadding(FMargin(0.f))
 		{
 			
diff --git a/Engine/Plugins/Experimental/EditorDataStorageFeatures/Source/TedsTableViewer/Public/Widgets/STedsTableViewerRow.h b/Engine/Plugins/Experimental/EditorDataStorageFeatures/Source/TedsTableViewer/Public/Widgets/STedsTableViewerRow.h
--- a/Engine/Plugins/Experimental/EditorDataStorageFeatures/Source/TedsTableViewer/Public/Widgets/STedsTableViewerRow.h
+++ b/Engine/Plugins/Experimental/EditorDataStorageFeatures/Source/TedsTableViewer/Public/Widgets/STedsTableViewerRow.h
@@ -20,7 +20,7 @@ namespace UE::Editor::DataStorage
 		SLATE_BEGIN_ARGS( STedsTableViewerRow )
 			: _ParentWidgetRowHandle(InvalidRowHandle)
 			, _Padding(FMargin(0.f, 4.f))
-			, _ItemHeight(48.f)
+			, _ItemHeight()
 		{}
 
 		/** The list item for this row */
