Made sure CDO are not editable and being named correctly in the TEDs text boxes (removing Default__ and _C on their display name)

#rnx

diff --git a/Engine/Plugins/Experimental/EditorDataStorageFeatures/Source/TedsOutliner/Private/Widgets/OutlinerCDOTextWidget.cpp b/Engine/Plugins/Experimental/EditorDataStorageFeatures/Source/TedsOutliner/Private/Widgets/OutlinerCDOTextWidget.cpp
new file mode 100644
--- /dev/null
+++ b/Engine/Plugins/Experimental/EditorDataStorageFeatures/Source/TedsOutliner/Private/Widgets/OutlinerCDOTextWidget.cpp
@@ -0,0 +1,64 @@
+// Copyright Epic Games, Inc. All Rights Reserved.
+
+#include "Widgets/OutlinerCDOTextWidget.h"
+
+#include "ActorEditorUtils.h"
+#include "Columns/SlateDelegateColumns.h"
+#include "Elements/Columns/TypedElementCompatibilityColumns.h"
+#include "Elements/Columns/TypedElementLabelColumns.h"
+#include "Elements/Framework/TypedElementAttributeBinding.h"
+#include "Widgets/Text/SInlineEditableTextBlock.h"
+#include "Widgets/Text/STextBlock.h"
+
+#define LOCTEXT_NAMESPACE "OutlinerCDOTextWidget"
+
+void UOutlinerCDOTextWidgetFactory::RegisterWidgetConstructors(UE::Editor::DataStorage::ICoreProvider& DataStorage,
+	UE::Editor::DataStorage::IUiProvider& DataStorageUi) const
+{
+	using namespace UE::Editor::DataStorage::Queries;
+
+	DataStorageUi.RegisterWidgetFactory<FOutlinerCDOTextWidgetConstructor>(DataStorageUi.FindPurpose(
+		IUiProvider::FPurposeInfo("SceneOutliner", "RowLabel", "Text").GeneratePurposeID()),
+		TColumn<FTypedElementLabelColumn>() && TColumn<FTypedElementClassDefaultObjectTag>());
+}
+
+FOutlinerCDOTextWidgetConstructor::FOutlinerCDOTextWidgetConstructor()
+	: Super(FOutlinerCDOTextWidgetConstructor::StaticStruct())
+{
+}
+
+FOutlinerCDOTextWidgetConstructor::FOutlinerCDOTextWidgetConstructor(const UScriptStruct* InTypeInfo)
+	: Super(InTypeInfo)
+{
+}
+
+TSharedPtr<SWidget> FOutlinerCDOTextWidgetConstructor::CreateWidget(UE::Editor::DataStorage::ICoreProvider* DataStorage,
+	UE::Editor::DataStorage::IUiProvider* DataStorageUi,
+	UE::Editor::DataStorage::RowHandle TargetRow,
+	UE::Editor::DataStorage::RowHandle WidgetRow, const UE::Editor::DataStorage::FMetaDataView& Arguments)
+{
+	using namespace UE::Editor::DataStorage;
+	if (!DataStorage->IsRowAvailable(TargetRow))
+	{
+		return SNew(STextBlock)
+			.Text(LOCTEXT("MissingRowReferenceColumn", "Unable to retrieve row reference."));
+	}
+
+	using namespace UE::Editor::DataStorage;
+	FAttributeBinder TargetRowBinder(TargetRow, DataStorage);
+
+	TSharedPtr<STextBlock> TextBlock = SNew(STextBlock)
+		.IsEnabled(false)
+		.Text(TargetRowBinder.BindData((&FTypedElementLabelColumn::Label), [](const FString& InString)
+		{
+			FString Name = InString;
+			Name.RemoveFromStart(TEXT("Default__"));
+			Name.RemoveFromEnd(TEXT("_C"));
+			return FText::FromString(Name);
+		}, FString()))
+		.ToolTipText(TargetRowBinder.BindText(&FTypedElementLabelColumn::Label));
+
+	return TextBlock;
+}
+
+#undef LOCTEXT_NAMESPACE
diff --git a/Engine/Plugins/Experimental/EditorDataStorageFeatures/Source/TedsOutliner/Public/Widgets/OutlinerCDOTextWidget.h b/Engine/Plugins/Experimental/EditorDataStorageFeatures/Source/TedsOutliner/Public/Widgets/OutlinerCDOTextWidget.h
new file mode 100644
--- /dev/null
+++ b/Engine/Plugins/Experimental/EditorDataStorageFeatures/Source/TedsOutliner/Public/Widgets/OutlinerCDOTextWidget.h
@@ -0,0 +1,44 @@
+// Copyright Epic Games, Inc. All Rights Reserved.
+
+#pragma once
+
+#include "Elements/Interfaces/TypedElementDataStorageFactory.h"
+#include "Elements/Interfaces/TypedElementDataStorageUiInterface.h"
+#include "Types/SlateEnums.h"
+
+#include "OutlinerCDOTextWidget.generated.h"
+
+struct FSlateBrush;
+struct FTypedElementClassTypeInfoColumn;
+class  SInlineEditableTextBlock;
+class  STextBlock;
+
+UCLASS()
+class UOutlinerCDOTextWidgetFactory : public UEditorDataStorageFactory
+{
+	GENERATED_BODY()
+
+public:
+	virtual ~UOutlinerCDOTextWidgetFactory() override = default;
+
+	void RegisterWidgetConstructors(UE::Editor::DataStorage::ICoreProvider& DataStorage,
+		UE::Editor::DataStorage::IUiProvider& DataStorageUi) const override;
+};
+
+USTRUCT()
+struct FOutlinerCDOTextWidgetConstructor : public FSimpleWidgetConstructor
+{
+	GENERATED_BODY()
+
+public:
+	TEDSOUTLINER_API FOutlinerCDOTextWidgetConstructor();
+	TEDSOUTLINER_API FOutlinerCDOTextWidgetConstructor(const UScriptStruct* InTypeInfo);
+	TEDSOUTLINER_API virtual ~FOutlinerCDOTextWidgetConstructor() override = default;
+
+	TEDSOUTLINER_API virtual TSharedPtr<SWidget> CreateWidget(
+		UE::Editor::DataStorage::ICoreProvider* DataStorage,
+		UE::Editor::DataStorage::IUiProvider* DataStorageUi,
+		UE::Editor::DataStorage::RowHandle TargetRow,
+		UE::Editor::DataStorage::RowHandle WidgetRow,
+		const UE::Editor::DataStorage::FMetaDataView& Arguments) override;
+};
