[EditorDataStorageFeatures] Remove leftover API
Removing API not needed

diff --git a/Engine/Plugins/Experimental/EditorDataStorageFeatures/Source/TedsContentBrowser/Private/TedsContentBrowserModule.cpp b/Engine/Plugins/Experimental/EditorDataStorageFeatures/Source/TedsContentBrowser/Private/TedsContentBrowserModule.cpp
--- a/Engine/Plugins/Experimental/EditorDataStorageFeatures/Source/TedsContentBrowser/Private/TedsContentBrowserModule.cpp
+++ b/Engine/Plugins/Experimental/EditorDataStorageFeatures/Source/TedsContentBrowser/Private/TedsContentBrowserModule.cpp
@@ -463,23 +463,6 @@ namespace UE::Editor::ContentBrowser
 		RefreshRows(InItemsSource);
 	}
 
-	bool FTedsContentBrowserViewExtender::SupportScaling() const
-	{
-		return true;
-	}
-
-	bool FTedsContentBrowserViewExtender::SupportZooming() const
-	{
-		// Zooming is for now supported only in the custom TileView
-		return CustomViewType == ETableViewMode::Tile;
-	}
-
-	bool FTedsContentBrowserViewExtender::SupportThumbnailEditMode() const
-	{
-		// Thumbnail edit mode is for now supported only in the custom TileView
-		return CustomViewType == ETableViewMode::Tile;
-	}
-
 	TArray<TSharedPtr<FAssetViewItem>> FTedsContentBrowserViewExtender::GetSelectedItems()
 	{
 		// CB 2.0 TODO: Figure out selection
diff --git a/Engine/Plugins/Experimental/EditorDataStorageFeatures/Source/TedsContentBrowser/Private/TedsContentBrowserModule.h b/Engine/Plugins/Experimental/EditorDataStorageFeatures/Source/TedsContentBrowser/Private/TedsContentBrowserModule.h
--- a/Engine/Plugins/Experimental/EditorDataStorageFeatures/Source/TedsContentBrowser/Private/TedsContentBrowserModule.h
+++ b/Engine/Plugins/Experimental/EditorDataStorageFeatures/Source/TedsContentBrowser/Private/TedsContentBrowserModule.h
@@ -59,9 +59,6 @@ namespace UE::Editor::ContentBrowser
 		virtual void ClearSelection() override;
 		virtual bool IsRightClickScrolling() override;
 		virtual void OnItemListChanged(TArray<TSharedPtr<FAssetViewItem>>* InItemsSource) override;
-		virtual bool SupportScaling() const override;
-		virtual bool SupportZooming() const override;
-		virtual bool SupportThumbnailEditMode() const override;
 		// ~IContentBrowserViewExtender interface
 
 		// Refresh the rows in the current view by syncing the the items source
diff --git a/Engine/Source/Editor/ContentBrowser/Public/Experimental/ContentBrowserViewExtender.h b/Engine/Source/Editor/ContentBrowser/Public/Experimental/ContentBrowserViewExtender.h
--- a/Engine/Source/Editor/ContentBrowser/Public/Experimental/ContentBrowserViewExtender.h
+++ b/Engine/Source/Editor/ContentBrowser/Public/Experimental/ContentBrowserViewExtender.h
@@ -62,13 +62,4 @@ public:
 	
 	// Get a tooltip for the custom view
 	virtual FText GetViewTooltipText() = 0;
-
-	// True if zooming and size changing can be done in the custom view, set it to false if not
-	virtual bool SupportScaling() const = 0;
-
-	// True if zooming and size changing can be done in the custom view, set it to false if not
-	virtual bool SupportZooming() const = 0;
-
-	// True if thumbnail edit mode should be supported, set it to false if not
-	virtual bool SupportThumbnailEditMode() const = 0;
 };
\ No newline at end of file
