TEDS: Register level folders in TEDS
Add a Cvar to populate folders in the Outliner directly through TEDS

diff --git a/Engine/Plugins/Experimental/EditorDataStorage/Source/TedsUI/Private/Widgets/WorldWidget.cpp b/Engine/Plugins/Experimental/EditorDataStorage/Source/TedsUI/Private/Widgets/WorldWidget.cpp
new file mode 100644
--- /dev/null
+++ b/Engine/Plugins/Experimental/EditorDataStorage/Source/TedsUI/Private/Widgets/WorldWidget.cpp
@@ -0,0 +1,47 @@
+// Copyright Epic Games, Inc. All Rights Reserved.
+
+#include "Widgets/WorldWidget.h"
+
+#include "Elements/Columns/TypedElementCompatibilityColumns.h"
+#include "Elements/Framework/TypedElementAttributeBinding.h"
+#include "Engine/World.h"
+#include "Internationalization/Text.h"
+#include "Widgets/Text/STextBlock.h"
+
+#define LOCTEXT_NAMESPACE "UWorldWidgetFactory"
+
+void UWorldWidgetFactory::RegisterWidgetConstructors(UE::Editor::DataStorage::ICoreProvider& DataStorage,
+	                                                     UE::Editor::DataStorage::IUiProvider& DataStorageUi) const
+{
+	using namespace UE::Editor::DataStorage::Queries;
+	
+	DataStorageUi.RegisterWidgetFactory<FWorldWidgetConstructor>(DataStorageUi.FindPurpose(DataStorageUi.GetGeneralWidgetPurposeID()),
+		TColumn<FTypedElementWorldColumn>());
+}
+
+FWorldWidgetConstructor::FWorldWidgetConstructor()
+	: FSimpleWidgetConstructor(StaticStruct())
+{
+}
+
+TSharedPtr<SWidget> FWorldWidgetConstructor::CreateWidget(UE::Editor::DataStorage::ICoreProvider* DataStorage,
+	UE::Editor::DataStorage::IUiProvider* DataStorageUi, UE::Editor::DataStorage::RowHandle TargetRow, UE::Editor::DataStorage::RowHandle WidgetRow,
+	const UE::Editor::DataStorage::FMetaDataView& Arguments)
+{
+	UE::Editor::DataStorage::FAttributeBinder Binder(TargetRow, DataStorage);
+	
+	return SNew(STextBlock)
+			.Text(Binder.BindData(&FTypedElementWorldColumn::World, [](const TWeakObjectPtr<UWorld>& WorldPtr)
+				{
+					if (UWorld* World = WorldPtr.Get())
+					{
+						FString WorldTypeString = LexToString(World->WorldType);
+						
+						return FText::Format(LOCTEXT("WorldName", "{0} ({1})"), FText::FromName(World->GetFName()),
+							FText::FromString(WorldTypeString));
+					}
+					return FText::GetEmpty();
+				}));
+}
+
+#undef LOCTEXT_NAMESPACE // "UWorldWidgetFactory"
diff --git a/Engine/Plugins/Experimental/EditorDataStorage/Source/TedsUI/Private/Widgets/WorldWidget.h b/Engine/Plugins/Experimental/EditorDataStorage/Source/TedsUI/Private/Widgets/WorldWidget.h
new file mode 100644
--- /dev/null
+++ b/Engine/Plugins/Experimental/EditorDataStorage/Source/TedsUI/Private/Widgets/WorldWidget.h
@@ -0,0 +1,40 @@
+// Copyright Epic Games, Inc. All Rights Reserved.
+
+#pragma once
+
+#include "Elements/Interfaces/TypedElementDataStorageFactory.h"
+#include "Elements/Interfaces/TypedElementDataStorageUiInterface.h"
+
+#include "WorldWidget.generated.h"
+
+class UScriptStruct;
+
+UCLASS()
+class UWorldWidgetFactory : public UEditorDataStorageFactory
+{
+	GENERATED_BODY()
+
+public:
+	~UWorldWidgetFactory() override = default;
+
+	void RegisterWidgetConstructors(UE::Editor::DataStorage::ICoreProvider& DataStorage,
+		UE::Editor::DataStorage::IUiProvider& DataStorageUi) const override;
+};
+
+// Widget to display the name and type of a UWorld
+USTRUCT()
+struct FWorldWidgetConstructor : public FSimpleWidgetConstructor
+{
+	GENERATED_BODY()
+
+public:
+	FWorldWidgetConstructor();
+	~FWorldWidgetConstructor() override = default;
+
+	virtual TSharedPtr<SWidget> CreateWidget(
+		UE::Editor::DataStorage::ICoreProvider* DataStorage,
+		UE::Editor::DataStorage::IUiProvider* DataStorageUi,
+		UE::Editor::DataStorage::RowHandle TargetRow,
+		UE::Editor::DataStorage::RowHandle WidgetRow, 
+		const UE::Editor::DataStorage::FMetaDataView& Arguments) override;
+};
\ No newline at end of file
diff --git a/Engine/Plugins/Experimental/EditorDataStorageFeatures/Source/TedsOutliner/Private/ActorFolderFactory.cpp b/Engine/Plugins/Experimental/EditorDataStorageFeatures/Source/TedsOutliner/Private/ActorFolderFactory.cpp
new file mode 100644
--- /dev/null
+++ b/Engine/Plugins/Experimental/EditorDataStorageFeatures/Source/TedsOutliner/Private/ActorFolderFactory.cpp
@@ -0,0 +1,335 @@
+// Copyright Epic Games, Inc. All Rights Reserved.
+
+#include "ActorFolderFactory.h"
+
+#include "Columns/ActorFolderColumns.h"
+#include "Editor.h"
+#include "EditorActorFolders.h"
+#include "Elements/Columns/TypedElementCompatibilityColumns.h"
+#include "Elements/Columns/TypedElementFolderColumns.h"
+#include "Elements/Columns/TypedElementHiearchyColumns.h"
+#include "Elements/Columns/TypedElementLabelColumns.h"
+#include "Elements/Columns/TypedElementMiscColumns.h"
+#include "Elements/Columns/TypedElementSlateWidgetColumns.h"
+#include "Elements/Framework/TypedElementRegistry.h"
+#include "Elements/Interfaces/TypedElementDataStorageInterface.h"
+#include "EngineUtils.h"
+#include "Engine/World.h"
+#include "SceneOutlinerHelpers.h"
+
+namespace UE::TedsActorFolderFactory::Local
+{
+	static const FName TableName("Editor_ActorFolderTable");
+}
+
+void UTedsActorFolderFactory::RegisterTables(UE::Editor::DataStorage::ICoreProvider& InDataStorage)
+{
+	InDataStorage.RegisterTable(
+		UE::Editor::DataStorage::TTypedElementColumnTypeList<
+		FFolderTag, FTypedElementLabelColumn, FTypedElementWorldColumn, FSlateColorColumn>(),
+	UE::TedsActorFolderFactory::Local::TableName);
+}
+
+void UTedsActorFolderFactory::PreRegister(UE::Editor::DataStorage::ICoreProvider& InDataStorage)
+{
+	Super::PreRegister(InDataStorage);
+	
+	FActorFolders::Get().OnFolderCreated.AddUObject(this, &UTedsActorFolderFactory::OnFolderCreated);
+	FActorFolders::Get().OnFolderDeleted.AddUObject(this, &UTedsActorFolderFactory::OnFolderDeleted);
+	FActorFolders::Get().OnFolderMoved.AddUObject(this, &UTedsActorFolderFactory::OnFolderMoved);
+
+	FEditorDelegates::MapChange.AddUObject(this, &UTedsActorFolderFactory::OnMapChange);
+	FEditorDelegates::PostPIEStarted.AddUObject(this, &UTedsActorFolderFactory::OnPIEStarted);
+	
+	FWorldDelegates::LevelAddedToWorld.AddUObject(this, &UTedsActorFolderFactory::OnLevelAdded);
+
+	GEngine->OnActorFolderAdded().AddUObject(this, &UTedsActorFolderFactory::OnActorFolderAdded);
+	GEngine->OnActorFolderRemoved().AddUObject(this, &UTedsActorFolderFactory::OnActorFolderAdded);
+	
+	InDataStorage.OnUpdateCompleted().AddUObject(this, &UTedsActorFolderFactory::Tick);
+	
+	// Store a pointer to data storage so we don't have to grab it from the global instance of the registry every time
+	DataStorage = &InDataStorage;
+}
+
+void UTedsActorFolderFactory::RegisterQueries(UE::Editor::DataStorage::ICoreProvider& InDataStorage)
+{
+	using namespace UE::Editor::DataStorage::Queries;
+	using namespace UE::Editor;
+
+	InDataStorage.RegisterQuery(
+		Select(
+			TEXT("Delete folder with deleted world"),
+			FProcessor(EQueryTickPhase::PostPhysics, InDataStorage.GetQueryTickGroupName(EQueryTickGroups::SyncExternalToDataStorage))
+				.SetExecutionMode(EExecutionMode::GameThread),
+			[](IQueryContext& Context, RowHandle Row, const FTypedElementWorldColumn& World)
+			{
+				if (!World.World.IsValid(false, true))
+				{
+					Context.RemoveRow(Row);
+				}
+			}
+		)
+		.Where()
+			.All<FFolderTag>()
+		.Compile()
+		);
+
+	// We can't currently detect changes to folders since they aren't UObjects so we run this query for all folders every frame.
+	// If this becomes an issue we can add an event on folder expansion changed that the factory can subscribe to instead (or add the sync tag manually)
+	InDataStorage.RegisterQuery(
+		Select(
+			TEXT("Set Folder expansion state"),
+			FProcessor(EQueryTickPhase::PostPhysics, InDataStorage.GetQueryTickGroupName(EQueryTickGroups::SyncExternalToDataStorage))
+				.SetExecutionMode(EExecutionMode::GameThread),
+			[](IQueryContext& Context, DataStorage::RowHandle Row,
+				const FFolderCompatibilityColumn& FolderCompatibilityColumn, const FTypedElementWorldColumn& WorldColumn)
+			{
+				TRACE_CPUPROFILER_EVENT_SCOPE(UTedsActorFolderFactory::SetFolderExpansionState);
+				if (UWorld* World = WorldColumn.World.Get())
+				{
+					// PIE folders are always treated as expanded
+					const bool bIsExpanded = !World->IsGameWorld() ?
+						FActorFolders::Get().IsFolderExpanded(*World, FolderCompatibilityColumn.Folder) : true;
+					
+					if (bIsExpanded)
+					{
+						Context.AddColumns<FFolderExpandedTag>(Row);
+					}
+					else
+					{
+						Context.RemoveColumns<FFolderExpandedTag>(Row);
+					}
+				}
+			}
+		)
+		.Where()
+			.All<FFolderTag>()
+		.Compile()
+		);
+
+	InDataStorage.RegisterQuery(
+		Select(
+			TEXT("Sync label column to folder"),
+			FProcessor(EQueryTickPhase::FrameEnd, InDataStorage.GetQueryTickGroupName(EQueryTickGroups::SyncDataStorageToExternal))
+				.SetExecutionMode(EExecutionMode::GameThread),
+			[this](RowHandle RowHandle)
+			{
+				// We have to defer the rename because it deletes and re-creates the folder which re-registers them in TEDS which cannot be done while
+				// we are in this query callback
+				FoldersToRename.Add(RowHandle);
+			}
+		)
+		.Where()
+			.All<FFolderTag, FFolderCompatibilityColumn, FTypedElementWorldColumn, FTypedElementSyncBackToWorldTag>()
+		.Compile()
+	);
+}
+
+void UTedsActorFolderFactory::PreShutdown(UE::Editor::DataStorage::ICoreProvider& InDataStorage)
+{
+	FActorFolders::Get().OnFolderCreated.RemoveAll(this);
+	FActorFolders::Get().OnFolderDeleted.RemoveAll(this);
+	FActorFolders::Get().OnFolderMoved.RemoveAll(this);
+
+	FEditorDelegates::MapChange.RemoveAll(this);
+	FEditorDelegates::PostPIEStarted.RemoveAll(this);
+	
+	FWorldDelegates::LevelAddedToWorld.RemoveAll(this);
+
+	GEngine->OnActorFolderAdded().RemoveAll(this);
+	GEngine->OnActorFolderRemoved().RemoveAll(this);
+	
+	InDataStorage.OnUpdateCompleted().RemoveAll(this);
+
+	Super::PreShutdown(InDataStorage);
+}
+
+void UTedsActorFolderFactory::OnFolderCreated(UWorld& World, const FFolder& Folder)
+{
+	RegisterFolder(World, Folder);
+}
+
+void UTedsActorFolderFactory::OnFolderDeleted(UWorld&, const FFolder& Folder)
+{
+	UnregisterFolder(Folder);
+}
+
+void UTedsActorFolderFactory::OnMapChange(uint32 MapChangeFlags)
+{
+	TRACE_CPUPROFILER_EVENT_SCOPE(UTedsActorFolderFactory::OnMapChange);
+	
+	// Iterate through editor world folders and add them to TEDS
+	if (UWorld* World = GEditor ? GEditor->GetEditorWorldContext().World() : nullptr)
+	{
+		FActorFolders::Get().ForEachFolder(*World, [this, World](const FFolder& Folder)
+		{
+			OnFolderCreated(*World, Folder);
+			return true;
+		});
+	}
+}
+
+void UTedsActorFolderFactory::OnPIEStarted(bool bIsSimulating)
+{
+	TRACE_CPUPROFILER_EVENT_SCOPE(UTedsActorFolderFactory::OnPIEStarted);
+	
+	// Iterate through PIE world actors and add their folders to TEDS
+	if (FWorldContext* PIEWorldContext = GEditor->GetPIEWorldContext())
+	{
+		if (UWorld* PieWorld = PIEWorldContext->World())
+		{
+			for (FActorIterator ActorIt(PieWorld); ActorIt; ++ActorIt)
+			{
+				if(AActor* Actor = *ActorIt)
+				{
+					OnFolderCreated(*PieWorld, Actor->GetFolder());
+				}
+			}
+		}
+	}
+}
+
+void UTedsActorFolderFactory::OnActorFolderAdded(UActorFolder* InActorFolder)
+{
+	if(ULevel* Level = InActorFolder->GetOuterULevel())
+	{
+		if(UWorld* World = Level->GetWorld())
+		{
+			RegisterFolder(*World, InActorFolder->GetFolder());
+		}
+	}
+	
+}
+
+void UTedsActorFolderFactory::OnActorFolderRemoved(UActorFolder* InActorFolder)
+{
+	UnregisterFolder(InActorFolder->GetFolder());
+}
+
+void UTedsActorFolderFactory::OnLevelAdded(ULevel* InLevel, UWorld* InWorld)
+{
+	TRACE_CPUPROFILER_EVENT_SCOPE(UTedsActorFolderFactory::OnLevelAdded);
+
+	if(InLevel && InWorld)
+	{
+		// Iterate through all actors in the level and add any folders containing them
+		for (AActor* Actor : InLevel->Actors)
+		{
+			if (Actor)
+			{
+				// Parent Folder
+				FFolder Folder = Actor->GetFolder();
+
+				if(!Folder.IsNone())
+				{
+					RegisterFolder(*InWorld, Folder);
+				}
+			}
+		}
+	}
+}
+
+void UTedsActorFolderFactory::OnFolderMoved(UWorld& InWorld, const FFolder& InOldFolder, const FFolder& InNewFolder)
+{
+	UE::Editor::DataStorage::FMapKey OldIndex(static_cast<uint64>(GetTypeHash(InOldFolder)));
+	UE::Editor::DataStorage::RowHandle Row = DataStorage->LookupMappedRow(OldIndex);
+
+	// Remove the mapping to the old folder
+	DataStorage->RemoveRowMapping(OldIndex);
+	
+	// Map this row to the new folder information
+	UE::Editor::DataStorage::FMapKey NewIndex(static_cast<uint64>(GetTypeHash(InNewFolder)));
+	DataStorage->MapRow(NewIndex, Row);
+
+	// Update the default columns in TEDS as the data could have changed after the folder moved
+	SetFolderColumns(Row, InWorld, InNewFolder);
+}
+
+// Register the folder and fill in the columns with data (World, FFolder, Parent, ActorFolder etc)
+UE::Editor::DataStorage::RowHandle UTedsActorFolderFactory::RegisterFolder(UWorld& World, const FFolder& Folder)
+{
+	// Don't register invalid folders
+	if(Folder.IsNone())
+	{
+		return UE::Editor::DataStorage::InvalidRowHandle;
+	}
+
+	UE::Editor::DataStorage::FMapKey Index(static_cast<uint64>(GetTypeHash(Folder)));
+	UE::Editor::DataStorage::RowHandle Row = DataStorage->LookupMappedRow(Index);
+	
+	if (!DataStorage->IsRowAvailable(Row))
+	{
+		// Add and index the row
+		static UE::Editor::DataStorage::RowHandle Table = DataStorage->FindTable(UE::TedsActorFolderFactory::Local::TableName);
+
+		if(Table == UE::Editor::DataStorage::InvalidTableHandle)
+		{
+			return UE::Editor::DataStorage::InvalidRowHandle;
+		}
+	
+		Row = DataStorage->AddRow(Table);
+		DataStorage->MapRow(Index, Row);
+
+		SetFolderColumns(Row, World,  Folder);
+	}
+	
+	return Row;
+}
+
+void UTedsActorFolderFactory::SetFolderColumns(UE::Editor::DataStorage::RowHandle Row, UWorld& World, const FFolder& Folder)
+{
+	DataStorage->GetColumn<FTypedElementWorldColumn>(Row)->World = TWeakObjectPtr<UWorld>(&World);
+	DataStorage->GetColumn<FTypedElementLabelColumn>(Row)->Label = Folder.GetLeafName().ToString();
+	DataStorage->AddColumn(Row, FFolderCompatibilityColumn{.Folder = Folder});
+
+	if(UActorFolder* ActorFolder = Folder.GetActorFolder())
+	{
+		DataStorage->AddColumn(Row, FTypedElementUObjectColumn{.Object = ActorFolder});
+	}
+	else
+	{
+		DataStorage->RemoveColumn<FTypedElementUObjectColumn>(Row);
+	}
+		
+	// Find or register the parent folder
+	UE::Editor::DataStorage::RowHandle ParentRow = RegisterFolder(World, Folder.GetParent());
+
+	if(DataStorage->IsRowAvailable(ParentRow))
+	{
+		DataStorage->AddColumn(Row, FTableRowParentColumn{.Parent = ParentRow});
+	}
+	else
+	{
+		DataStorage->RemoveColumn<FTableRowParentColumn>(Row);
+	}
+}
+
+void UTedsActorFolderFactory::UnregisterFolder(const FFolder& Folder)
+{
+	UE::Editor::DataStorage::FMapKey Index(static_cast<uint64>(GetTypeHash(Folder)));
+	UE::Editor::DataStorage::RowHandle Row = DataStorage->LookupMappedRow(Index);
+
+	if (DataStorage->IsRowAvailable(Row))
+	{
+		DataStorage->RemoveRow(Row);
+	}
+}
+
+void UTedsActorFolderFactory::Tick()
+{
+	for(UE::Editor::DataStorage::RowHandle RowHandle : FoldersToRename)
+	{
+		FFolderCompatibilityColumn* FolderCompatibilityColumn = DataStorage->GetColumn<FFolderCompatibilityColumn>(RowHandle);
+		FTypedElementWorldColumn* WorldColumn = DataStorage->GetColumn<FTypedElementWorldColumn>(RowHandle);
+		FTypedElementLabelColumn* LabelColumn = DataStorage->GetColumn<FTypedElementLabelColumn>(RowHandle);
+
+		if(FolderCompatibilityColumn && WorldColumn && LabelColumn)
+		{
+			SceneOutliner::FSceneOutlinerHelpers::RenameFolder(FolderCompatibilityColumn->Folder,
+				FText::FromString(LabelColumn->Label), WorldColumn->World.Get());
+		}
+	}
+	
+	FoldersToRename.Empty();
+}
diff --git a/Engine/Plugins/Experimental/EditorDataStorageFeatures/Source/TedsOutliner/Private/ActorFolderFactory.h b/Engine/Plugins/Experimental/EditorDataStorageFeatures/Source/TedsOutliner/Private/ActorFolderFactory.h
new file mode 100644
--- /dev/null
+++ b/Engine/Plugins/Experimental/EditorDataStorageFeatures/Source/TedsOutliner/Private/ActorFolderFactory.h
@@ -0,0 +1,64 @@
+// Copyright Epic Games, Inc. All Rights Reserved.
+
+#pragma once
+
+#include "Elements/Common/TypedElementHandles.h"
+#include "Elements/Interfaces/TypedElementDataStorageFactory.h"
+#include "Containers/Ticker.h"
+
+#include "ActorFolderFactory.generated.h"
+
+namespace UE::Editor::DataStorage
+{
+	struct IQueryContext;
+}
+
+class IEditorDataStorageProvider;
+class UWorld;
+struct FFolder;
+class UActorFolder;
+
+UCLASS()
+class UTedsActorFolderFactory : public UEditorDataStorageFactory
+{
+	GENERATED_BODY()
+
+public:
+	~UTedsActorFolderFactory() override = default;
+
+	virtual void RegisterTables(UE::Editor::DataStorage::ICoreProvider& DataStorage) override;
+	virtual void PreRegister(UE::Editor::DataStorage::ICoreProvider& DataStorage) override;
+	virtual void RegisterQueries(UE::Editor::DataStorage::ICoreProvider& DataStorage) override;
+	virtual void PreShutdown(UE::Editor::DataStorage::ICoreProvider& DataStorage) override;
+
+private:
+
+	void OnFolderCreated(UWorld& World, const FFolder& Folder);
+	void OnFolderDeleted(UWorld& World, const FFolder& Folder);
+	void OnFolderMoved(UWorld& InWorld, const FFolder& InOldFolder, const FFolder& InNewFolder);
+
+	void OnActorFolderAdded(UActorFolder* InActorFolder);
+	void OnActorFolderRemoved(UActorFolder* InActorFolder);
+
+	void OnMapChange(uint32 MapChangeFlags);
+	void OnLevelAdded(ULevel* InLevel, UWorld* InWorld);
+
+	void OnPIEStarted(bool bIsSimulating);
+
+	// Register a folder in TEDS (If not already registered) and return the row handle
+	UE::Editor::DataStorage::RowHandle RegisterFolder(UWorld& World, const FFolder& Folder);
+	void SetFolderColumns(UE::Editor::DataStorage::RowHandle Row, UWorld& World, const FFolder& Folder);
+
+
+	// Unregister a folder in TEDS
+	void UnregisterFolder(const FFolder& Folder);
+
+	void Tick();
+	
+private:
+	
+	UE::Editor::DataStorage::ICoreProvider* DataStorage;
+
+	// List of folders to process rename in the next tick
+	TArray<UE::Editor::DataStorage::RowHandle> FoldersToRename;
+};
diff --git a/Engine/Plugins/Experimental/EditorDataStorageFeatures/Source/TedsOutliner/Private/Compatibility/SceneOutlinerTedsBridge.cpp b/Engine/Plugins/Experimental/EditorDataStorageFeatures/Source/TedsOutliner/Private/Compatibility/SceneOutlinerTedsBridge.cpp
--- a/Engine/Plugins/Experimental/EditorDataStorageFeatures/Source/TedsOutliner/Private/Compatibility/SceneOutlinerTedsBridge.cpp
+++ b/Engine/Plugins/Experimental/EditorDataStorageFeatures/Source/TedsOutliner/Private/Compatibility/SceneOutlinerTedsBridge.cpp
@@ -29,6 +29,7 @@
 #include "Elements/Interfaces/Capabilities/TypedElementUiTextCapability.h"
 #include "TedsOutlinerItem.h"
 #include "Modules/ModuleManager.h"
+#include "Widgets/Docking/SDockTab.h"
 
 #define LOCTEXT_NAMESPACE "SceneOutlinerTedsBridge"
 
@@ -198,7 +199,7 @@ public:
 	{
 		TSharedPtr<ISceneOutliner> OutlinerPinned = OwningOutliner.Pin();
 
-		if(!OutlinerPinned)
+		if (!OutlinerPinned)
 		{
 			return false;
 		}
@@ -207,15 +208,15 @@ public:
 		FSceneOutlinerTreeItemPtr Item = OutlinerPinned->GetTreeItem(InRowHandle);
 
 		// If it doesn't exist, this could be a legacy item that uses something other than the row id as the ID, so check if we have a dealiaser
-		if(!Item)
+		if (!Item)
 		{
-			if(Dealiaser.IsBound())
+			if (Dealiaser.IsBound())
 			{
 				Item = OutlinerPinned->GetTreeItem(Dealiaser.Execute(InRowHandle));
 			}
 		}
 
-		if(!Item)
+		if (!Item)
 		{
 			return false;
 		}
@@ -232,7 +233,7 @@ public:
 	// TODO: Sorting is currently handled through the fallback column if it exists because we have no way to sort columns through TEDS
 	virtual void SortItems(TArray<FSceneOutlinerTreeItemPtr>& RootItems, const EColumnSortMode::Type SortMode) const override
 	{
-		if(FallbackColumn)
+		if (FallbackColumn)
 		{
 			FallbackColumn->SortItems(RootItems, SortMode);
 		}
@@ -247,7 +248,7 @@ public:
 	{
 		TSharedPtr<ISceneOutliner> OutlinerPinned = OwningOutliner.Pin();
 
-		if(!OutlinerPinned)
+		if (!OutlinerPinned)
 		{
 			return;
 		}
@@ -275,7 +276,7 @@ public:
 
 		TSharedPtr<SWidget> RowWidget;
 
-		if(const FTedsOutlinerTreeItem* TedsItem = TreeItem->CastTo<FTedsOutlinerTreeItem>())
+		if (const FTedsOutlinerTreeItem* TedsItem = TreeItem->CastTo<FTedsOutlinerTreeItem>())
 		{
 			TargetRowHandle = TedsItem->GetRowHandle();
 			
@@ -287,7 +288,7 @@ public:
 				TargetRowHandle = StorageCompatibility.FindRowWithCompatibleObject(Actor);
 			}
 		}
-		else if(FallbackColumn)
+		else if (FallbackColumn)
 		{
 			RowWidget = FallbackColumn->ConstructRowWidget(TreeItem, Row);
 		}
@@ -301,7 +302,7 @@ public:
 				});
 		}
 
-		if(RowWidget)
+		if (RowWidget)
 		{
 			SetHighlightText(*RowWidget);
 			return RowWidget.ToSharedRef();
@@ -313,7 +314,7 @@ public:
 	virtual void PopulateSearchStrings( const ISceneOutlinerTreeItem& Item, TArray< FString >& OutSearchStrings ) const override
 	{
 		// TODO: We don't currently have a way to convert TEDS widgets into searchable strings, but we can rely on the fallback column if it exists
-		if(FallbackColumn)
+		if (FallbackColumn)
 		{
 			FallbackColumn->PopulateSearchStrings(Item, OutSearchStrings);
 		}
@@ -401,7 +402,7 @@ FName FSceneOutlinerTedsQueryBinder::FindOutlinerColumnFromTEDSColumns(TConstArr
 	// Currently, the algorithm naively looks through the mapping and returns the first match
 	for(const TWeakObjectPtr<const UScriptStruct>& Column : TEDSColumns)
 	{
-		if(const FName* FoundDefaultColumn = TEDSToOutlinerDefaultColumnMapping.Find(Column))
+		if (const FName* FoundDefaultColumn = TEDSToOutlinerDefaultColumnMapping.Find(Column))
 		{
 			return *FoundDefaultColumn;
 		}
@@ -410,6 +411,23 @@ FName FSceneOutlinerTedsQueryBinder::FindOutlinerColumnFromTEDSColumns(TConstArr
 	return FName();
 }
 
+void FSceneOutlinerTedsQueryBinder::RefreshLevelEditorOutliners() const
+{
+	FLevelEditorModule& LevelEditorModule = FModuleManager::GetModuleChecked<FLevelEditorModule>("LevelEditor");
+	TSharedPtr<FTabManager> LevelEditorTabManager = LevelEditorModule.GetLevelEditorTabManager();
+
+	static const FName TabIDS[] = { LevelEditorTabIds::LevelEditorSceneOutliner, LevelEditorTabIds::LevelEditorSceneOutliner2, LevelEditorTabIds::LevelEditorSceneOutliner3, LevelEditorTabIds::LevelEditorSceneOutliner4 };
+
+	for (const FName& TabID : TabIDS)
+	{
+		if (LevelEditorTabManager->FindExistingLiveTab(TabID).IsValid())
+		{
+			LevelEditorTabManager->TryInvokeTab(TabID)->RequestCloseTab();
+			LevelEditorTabManager->TryInvokeTab(TabID);
+		}
+	}
+}
+
 FSceneOutlinerTedsQueryBinder& FSceneOutlinerTedsQueryBinder::GetInstance()
 {
 	static FSceneOutlinerTedsQueryBinder Binder;
@@ -465,7 +483,7 @@ FTreeItemIDDealiaser FSceneOutlinerTedsQueryBinder::GetTreeItemIDDealiaser(const
 {
 	TSharedPtr<FSceneOutlinerTedsBridge>* QueryMapping = FindQueryMapping(Widget);
 
-	if(QueryMapping)
+	if (QueryMapping)
 	{
 		return (*QueryMapping)->GetDealiaser();
 	}
@@ -478,7 +496,7 @@ void FSceneOutlinerTedsQueryBinder::CleanupStaleOutliners()
 	for (TMap<TWeakPtr<ISceneOutliner>, TSharedPtr<FSceneOutlinerTedsBridge>>::TIterator It(SceneOutliners); It; ++It)
 	{
 		// Remove any query mappings where the target Outliner doesn't exist anymore
-		if(!It.Key().IsValid())
+		if (!It.Key().IsValid())
 		{
 			It.RemoveCurrent();
 		}
diff --git a/Engine/Plugins/Experimental/EditorDataStorageFeatures/Source/TedsOutliner/Private/Widgets/OutlinerFolderLabelWidget.cpp b/Engine/Plugins/Experimental/EditorDataStorageFeatures/Source/TedsOutliner/Private/Widgets/OutlinerFolderLabelWidget.cpp
new file mode 100644
--- /dev/null
+++ b/Engine/Plugins/Experimental/EditorDataStorageFeatures/Source/TedsOutliner/Private/Widgets/OutlinerFolderLabelWidget.cpp
@@ -0,0 +1,281 @@
+// Copyright Epic Games, Inc. All Rights Reserved.
+
+#include "Widgets/OutlinerFolderLabelWidget.h"
+
+#include "Columns/ActorFolderColumns.h"
+#include "Columns/SlateDelegateColumns.h"
+#include "Columns/TedsOutlinerColumns.h"
+#include "Compatibility/SceneOutlinerTedsBridge.h"
+#include "Elements/Columns/TypedElementCompatibilityColumns.h"
+#include "Elements/Columns/TypedElementLabelColumns.h"
+#include "Elements/Columns/TypedElementFolderColumns.h"
+#include "Elements/Columns/TypedElementMiscColumns.h"
+#include "Elements/Columns/TypedElementSlateWidgetColumns.h"
+#include "Elements/Framework/TypedElementAttributeBinding.h"
+#include "ISceneOutliner.h"
+#include "SceneOutlinerHelpers.h"
+#include "Widgets/Images/SImage.h"
+#include "Widgets/Layout/SSpacer.h"
+#include "Widgets/Layout/SBox.h"
+#include "Widgets/Text/STextBlock.h"
+#include "Widgets/SBoxPanel.h"
+#include "Widgets/Text/SInlineEditableTextBlock.h"
+
+#define LOCTEXT_NAMESPACE "FOutlinerFolderLabelWidgetConstructor"
+
+
+void UOutlinerFolderLabelWidgetFactory::RegisterWidgetConstructors(UE::Editor::DataStorage::ICoreProvider& DataStorage,
+	UE::Editor::DataStorage::IUiProvider& DataStorageUi) const
+{
+	using namespace UE::Editor::DataStorage::Queries;
+	DataStorageUi.RegisterWidgetFactory<FOutlinerFolderLabelWidgetConstructor>(
+		DataStorageUi.FindPurpose(IUiProvider::FPurposeInfo("SceneOutliner", "RowLabel", NAME_None).GeneratePurposeID()),
+		TColumn<FTypedElementLabelColumn>() && TColumn<FFolderTag>());
+}
+
+FOutlinerFolderLabelWidgetConstructor::FOutlinerFolderLabelWidgetConstructor()
+	: FSimpleWidgetConstructor(StaticStruct())
+{
+}
+
+TSharedPtr<SWidget> FOutlinerFolderLabelWidgetConstructor::CreateWidget(UE::Editor::DataStorage::ICoreProvider* DataStorage,
+	UE::Editor::DataStorage::IUiProvider* DataStorageUi, UE::Editor::DataStorage::RowHandle TargetRow, UE::Editor::DataStorage::RowHandle WidgetRow,
+	const UE::Editor::DataStorage::FMetaDataView& Arguments)
+{
+	using namespace UE::Editor::DataStorage;
+	if (DataStorage->IsRowAvailable(TargetRow))
+	{
+		FAttributeBinder Binder(TargetRow, DataStorage);
+
+		
+
+		return SNew(SHorizontalBox)
+			+SHorizontalBox::Slot()
+			.AutoWidth()
+			.VAlign(VAlign_Center)
+			[
+				SNew(SBox)
+				.WidthOverride(16.0f)
+				.HeightOverride(16.0f)
+				[
+					SNew(SImage)
+						.Image_Lambda([DataStorage, TargetRow, WidgetRow]()
+						{
+							bool bIsExpanded = DataStorage->HasColumns<FFolderExpandedTag>(TargetRow);
+
+							TWeakPtr<ISceneOutlinerTreeItem> OutlinerTreeItem = GetTreeItemForRow(DataStorage, TargetRow, WidgetRow);
+							
+							// If this item does not have any children, we want to treat it as not expanded and show the closed folder icon
+							if (TSharedPtr<ISceneOutlinerTreeItem> Item = OutlinerTreeItem.Pin())
+							{
+								bIsExpanded &= !Item->GetChildren().IsEmpty();
+							}
+
+							return bIsExpanded
+								? FAppStyle::GetBrush("SceneOutliner.FolderOpen")
+								: FAppStyle::GetBrush("SceneOutliner.FolderClosed");
+						})
+						// Disabled for now since folder colors are not properly supported yet
+						//.ColorAndOpacity(Binder.BindData(&FSlateColorColumn::Color))
+				]
+			]
+			+SHorizontalBox::Slot()
+			.AutoWidth()
+			[
+				SNew(SSpacer)
+					.Size(FVector2D(6.0f, 0.0f))
+			]
+			+SHorizontalBox::Slot()
+			.FillWidth(1.0f)
+			.VAlign(VAlign_Center)
+			[
+				CreateLabel(DataStorage, TargetRow, WidgetRow, Arguments)
+			];
+	}
+	else
+	{
+		return SNew(STextBlock)
+			.Text(LOCTEXT("MissingRowReferenceColumn", "Unable to retrieve row reference."));
+	}
+}
+
+TSharedRef<SWidget> FOutlinerFolderLabelWidgetConstructor::CreateLabel(UE::Editor::DataStorage::ICoreProvider* DataStorage,
+	UE::Editor::DataStorage::RowHandle TargetRow, UE::Editor::DataStorage::RowHandle WidgetRow,
+	const UE::Editor::DataStorage::FMetaDataView& Arguments)
+{
+	using namespace UE::Editor::DataStorage;
+	TSharedRef<SWidget> Result = SNullWidget::NullWidget;
+	
+	const bool* IsEditable = Arguments.FindForColumn<FTypedElementLabelColumn>(IsEditableName).TryGetExact<bool>();
+
+	if (IsEditable && *IsEditable)
+	{
+		FAttributeBinder WidgetRowBinder(WidgetRow, DataStorage);
+			
+		TSharedPtr<SInlineEditableTextBlock> TextBlock = SNew(SInlineEditableTextBlock)
+			.OnTextCommitted_Lambda(
+				[DataStorage, TargetRow](const FText& NewText, ETextCommit::Type CommitInfo)
+				{
+					// This callback happens on the game thread so it's safe to directly call into the data storage.
+					FString NewLabelText = NewText.ToString();
+					if (FTypedElementLabelHashColumn* LabelHashColumn = DataStorage->GetColumn<FTypedElementLabelHashColumn>(TargetRow))
+					{
+						LabelHashColumn->LabelHash = CityHash64(reinterpret_cast<const char*>(*NewLabelText), NewLabelText.Len() * sizeof(**NewLabelText));
+					}
+					if (FTypedElementLabelColumn* LabelColumn = DataStorage->GetColumn<FTypedElementLabelColumn>(TargetRow))
+					{
+						LabelColumn->Label = MoveTemp(NewLabelText);
+					}
+					DataStorage->AddColumn<FTypedElementSyncBackToWorldTag>(TargetRow);
+				})
+			.OnVerifyTextChanged_Lambda([DataStorage, TargetRow](const FText& Label, FText& ErrorMessage)
+				{
+					const FFolderCompatibilityColumn* FolderCompatibilityColumn = DataStorage->GetColumn<FFolderCompatibilityColumn>(TargetRow);
+					const FTypedElementWorldColumn* WorldColumn = DataStorage->GetColumn<FTypedElementWorldColumn>(TargetRow);
+
+					if (FolderCompatibilityColumn && WorldColumn)
+					{
+						return SceneOutliner::FSceneOutlinerHelpers::ValidateFolderName(FolderCompatibilityColumn->Folder, WorldColumn->World.Get(),
+							Label, ErrorMessage);
+					}
+
+					ErrorMessage = LOCTEXT("MissingColumns", "Could not find folder information to rename.");
+				
+					return false;
+				})
+			.Text_Static(&FOutlinerFolderLabelWidgetConstructor::GetDisplayText, DataStorage, TargetRow, WidgetRow)
+			.ToolTipText_Static(&FOutlinerFolderLabelWidgetConstructor::GetTooltipText, DataStorage, TargetRow)
+			.ColorAndOpacity_Static(&FOutlinerFolderLabelWidgetConstructor::GetForegroundColor, DataStorage, TargetRow, WidgetRow)
+			.IsSelected(WidgetRowBinder.BindEvent(&FExternalWidgetSelectionColumn::IsSelected))
+			.OnEnterEditingMode_Lambda([DataStorage, WidgetRow]()
+			{
+				DataStorage->AddColumn<FIsInEditingModeTag>(WidgetRow);
+			})
+			.OnExitEditingMode_Lambda([DataStorage, WidgetRow]()
+			{
+				DataStorage->RemoveColumn<FIsInEditingModeTag>(WidgetRow);
+			});
+
+		DataStorage->AddColumn<FWidgetEnterEditModeColumn>(TargetRow, FWidgetEnterEditModeColumn{
+				.OnEnterEditMode = FSimpleDelegate::CreateLambda([TextBlock]()
+				{
+					TextBlock->EnterEditingMode();
+				})
+			});
+
+		Result = TextBlock.ToSharedRef();
+	}
+	else
+	{
+		FAttributeBinder TargetRowBinder(TargetRow, DataStorage);
+		
+		TSharedPtr<STextBlock> TextBlock = SNew(STextBlock)
+			.IsEnabled(false)
+			.Text(TargetRowBinder.BindText(&FTypedElementLabelColumn::Label))
+			.ToolTipText(TargetRowBinder.BindText(&FTypedElementLabelColumn::Label));
+		
+		Result = TextBlock.ToSharedRef();
+	}
+
+	return Result;
+}
+
+FText FOutlinerFolderLabelWidgetConstructor::GetDisplayText(UE::Editor::DataStorage::ICoreProvider* DataStorage, UE::Editor::DataStorage::RowHandle TargetRow,
+	UE::Editor::DataStorage::RowHandle WidgetRow)
+{
+	if (const FTypedElementLabelColumn* LabelColumn = DataStorage->GetColumn<FTypedElementLabelColumn>(TargetRow))
+	{
+		FText Label = FText::FromString(LabelColumn->Label);
+		
+		const FFolderCompatibilityColumn* FolderCompatibilityColumn = DataStorage->GetColumn<FFolderCompatibilityColumn>(TargetRow);
+		const FTypedElementWorldColumn* WorldColumn = DataStorage->GetColumn<FTypedElementWorldColumn>(TargetRow);
+		const bool IsInEditingMode = DataStorage->HasColumns<FIsInEditingModeTag>(WidgetRow);
+
+		if (FolderCompatibilityColumn && WorldColumn)
+		{
+			if (!IsInEditingMode)
+			{
+				FText IsCurrentSuffixText =
+					SceneOutliner::FSceneOutlinerHelpers::IsFolderCurrent(FolderCompatibilityColumn->Folder, WorldColumn->World.Get())
+						?	FText(LOCTEXT("IsCurrentSuffix", " (Current)"))
+						: FText::GetEmpty();
+				
+				return FText::Format(LOCTEXT("LevelInstanceDisplay", "{0}{1}"), Label, IsCurrentSuffixText);
+			}
+		}
+
+		return Label;
+	}
+
+	return FText();
+}
+
+FText FOutlinerFolderLabelWidgetConstructor::GetTooltipText(UE::Editor::DataStorage::ICoreProvider* DataStorage, UE::Editor::DataStorage::RowHandle TargetRow)
+{
+	const FFolderCompatibilityColumn* FolderCompatibilityColumn = DataStorage->GetColumn<FFolderCompatibilityColumn>(TargetRow);
+	const FTypedElementWorldColumn* WorldColumn = DataStorage->GetColumn<FTypedElementWorldColumn>(TargetRow);
+	const FTypedElementLabelColumn* LabelColumn = DataStorage->GetColumn<FTypedElementLabelColumn>(TargetRow);
+
+	if (LabelColumn && WorldColumn && FolderCompatibilityColumn)
+	{
+		FText Description = SceneOutliner::FSceneOutlinerHelpers::IsFolderCurrent(FolderCompatibilityColumn->Folder, WorldColumn->World.Get())
+				? LOCTEXT("ActorFolderIsCurrentDescription", "\nThis is your current folder. New actors you create will appear here.")
+				: FText::GetEmpty();
+	
+		return FText::Format(LOCTEXT("DataLayerTooltipText", "{0}{1}"), FText::FromString(LabelColumn->Label), Description);
+	}
+	return FText();
+}
+
+FSlateColor FOutlinerFolderLabelWidgetConstructor::GetForegroundColor(UE::Editor::DataStorage::ICoreProvider* DataStorage,
+	UE::Editor::DataStorage::RowHandle TargetRow, UE::Editor::DataStorage::RowHandle WidgetRow)
+{
+	// CommonLabelData has some Outliner specific color logic
+	TWeakPtr<ISceneOutlinerTreeItem> OutlinerTreeItem = GetTreeItemForRow(DataStorage, TargetRow, WidgetRow);
+	
+	if (TSharedPtr<ISceneOutlinerTreeItem> Item = OutlinerTreeItem.Pin())
+	{
+		FSceneOutlinerCommonLabelData CommonLabelData;
+		CommonLabelData.WeakSceneOutliner = Item->WeakSceneOutliner;
+
+		if (TOptional<FLinearColor> BaseColor = CommonLabelData.GetForegroundColor(*Item))
+		{
+			return BaseColor.GetValue();
+		}
+		
+		const FFolderCompatibilityColumn* FolderCompatibilityColumn = DataStorage->GetColumn<FFolderCompatibilityColumn>(TargetRow);
+		const FTypedElementWorldColumn* WorldColumn = DataStorage->GetColumn<FTypedElementWorldColumn>(TargetRow);
+
+		if (FolderCompatibilityColumn && WorldColumn)
+		{
+			if (SceneOutliner::FSceneOutlinerHelpers::IsFolderCurrent(FolderCompatibilityColumn->Folder, WorldColumn->World.Get()))
+			{
+				return FAppStyle::Get().GetSlateColor("Colors.AccentGreen");
+			}
+		}
+	}
+
+	return FSlateColor::UseForeground();
+}
+
+TWeakPtr<ISceneOutlinerTreeItem> FOutlinerFolderLabelWidgetConstructor::GetTreeItemForRow(UE::Editor::DataStorage::ICoreProvider* DataStorage,
+	UE::Editor::DataStorage::RowHandle TargetRow, UE::Editor::DataStorage::RowHandle WidgetRow)
+{
+	if (FTedsOutlinerColumn* TedsOutlinerColumn = DataStorage->GetColumn<FTedsOutlinerColumn>(WidgetRow))
+	{
+		if (TSharedPtr<ISceneOutliner> Outliner = TedsOutlinerColumn->Outliner.Pin())
+		{
+			FSceneOutlinerTedsQueryBinder& Binder = FSceneOutlinerTedsQueryBinder::GetInstance();
+          
+			if (FTreeItemIDDealiaser Dealiaser = Binder.GetTreeItemIDDealiaser(Outliner);
+			   Dealiaser.IsBound())
+			{
+				return Outliner->GetTreeItem(Dealiaser.Execute(TargetRow));
+			}
+		}
+	}
+
+	return nullptr;
+}
+
+#undef LOCTEXT_NAMESPACE // "FOutlinerFolderLabelWidgetConstructor"
diff --git a/Engine/Plugins/Experimental/EditorDataStorageFeatures/Source/TedsOutliner/Private/Widgets/OutlinerFolderLabelWidget.h b/Engine/Plugins/Experimental/EditorDataStorageFeatures/Source/TedsOutliner/Private/Widgets/OutlinerFolderLabelWidget.h
new file mode 100644
--- /dev/null
+++ b/Engine/Plugins/Experimental/EditorDataStorageFeatures/Source/TedsOutliner/Private/Widgets/OutlinerFolderLabelWidget.h
@@ -0,0 +1,55 @@
+// Copyright Epic Games, Inc. All Rights Reserved.
+
+#pragma once
+
+#include "Elements/Interfaces/TypedElementDataStorageFactory.h"
+#include "Elements/Interfaces/TypedElementDataStorageUiInterface.h"
+#include "Internationalization/Text.h"
+#include "UObject/ObjectMacros.h"
+#include "Styling/SlateColor.h"
+
+#include "OutlinerFolderLabelWidget.generated.h"
+
+class UScriptStruct;
+struct ISceneOutlinerTreeItem;
+
+UCLASS()
+class UOutlinerFolderLabelWidgetFactory : public UEditorDataStorageFactory
+{
+	GENERATED_BODY()
+
+public:
+	~UOutlinerFolderLabelWidgetFactory() override = default;
+
+	TEDSOUTLINER_API void RegisterWidgetConstructors(UE::Editor::DataStorage::ICoreProvider& DataStorage,
+		UE::Editor::DataStorage::IUiProvider& DataStorageUi) const override;
+};
+
+// Label widget for folders in the Outliner that shows an icon + a text label
+USTRUCT()
+struct FOutlinerFolderLabelWidgetConstructor : public FSimpleWidgetConstructor
+{
+	GENERATED_BODY()
+
+public:
+	TEDSOUTLINER_API FOutlinerFolderLabelWidgetConstructor(); 
+	~FOutlinerFolderLabelWidgetConstructor() override = default;
+
+	TEDSOUTLINER_API virtual TSharedPtr<SWidget> CreateWidget(
+		UE::Editor::DataStorage::ICoreProvider* DataStorage,
+		UE::Editor::DataStorage::IUiProvider* DataStorageUi,
+		UE::Editor::DataStorage::RowHandle TargetRow,
+		UE::Editor::DataStorage::RowHandle WidgetRow, 
+		const UE::Editor::DataStorage::FMetaDataView& Arguments) override;
+
+protected:
+ 
+	TSharedRef<SWidget> CreateLabel(UE::Editor::DataStorage::ICoreProvider* DataStorage, UE::Editor::DataStorage::RowHandle TargetRow,
+		UE::Editor::DataStorage::RowHandle WidgetRow, const UE::Editor::DataStorage::FMetaDataView& Arguments);
+
+	static FText GetDisplayText(UE::Editor::DataStorage::ICoreProvider* DataStorage, UE::Editor::DataStorage::RowHandle TargetRow, UE::Editor::DataStorage::RowHandle WidgetRow);
+	static FText GetTooltipText(UE::Editor::DataStorage::ICoreProvider* DataStorage, UE::Editor::DataStorage::RowHandle TargetRow);
+	static FSlateColor GetForegroundColor(UE::Editor::DataStorage::ICoreProvider* DataStorage, UE::Editor::DataStorage::RowHandle TargetRow, UE::Editor::DataStorage::RowHandle WidgetRow);
+	
+	static TWeakPtr<ISceneOutlinerTreeItem> GetTreeItemForRow(UE::Editor::DataStorage::ICoreProvider* DataStorage, UE::Editor::DataStorage::RowHandle TargetRow, UE::Editor::DataStorage::RowHandle WidgetRow);
+};
\ No newline at end of file
diff --git a/Engine/Plugins/Experimental/EditorDataStorageFeatures/Source/TedsOutliner/Public/Columns/ActorFolderColumns.h b/Engine/Plugins/Experimental/EditorDataStorageFeatures/Source/TedsOutliner/Public/Columns/ActorFolderColumns.h
new file mode 100644
--- /dev/null
+++ b/Engine/Plugins/Experimental/EditorDataStorageFeatures/Source/TedsOutliner/Public/Columns/ActorFolderColumns.h
@@ -0,0 +1,27 @@
+// Copyright Epic Games, Inc. All Rights Reserved.
+
+#pragma once
+
+#include "Elements/Interfaces/TypedElementDataStorageInterface.h"
+#include "Folder.h"
+
+#include "ActorFolderColumns.generated.h"
+
+/**
+ * Column that stores a constructed FFolder
+ */
+USTRUCT(meta = (DisplayName = "FFolder compatibility"))
+struct FFolderCompatibilityColumn final : public FEditorDataStorageColumn
+{
+	GENERATED_BODY()
+	FFolder Folder;
+};
+
+/**
+ * Tag that specifies if a folder is expanded in the UI (child items visible)
+ */
+USTRUCT(meta = (DisplayName = "Is Folder Expanded"))
+struct FFolderExpandedTag final : public FEditorDataStorageTag
+{
+	GENERATED_BODY()
+};
\ No newline at end of file
diff --git a/Engine/Plugins/Experimental/EditorDataStorageFeatures/Source/TedsOutliner/Public/Compatibility/SceneOutlinerTedsBridge.h b/Engine/Plugins/Experimental/EditorDataStorageFeatures/Source/TedsOutliner/Public/Compatibility/SceneOutlinerTedsBridge.h
--- a/Engine/Plugins/Experimental/EditorDataStorageFeatures/Source/TedsOutliner/Public/Compatibility/SceneOutlinerTedsBridge.h
+++ b/Engine/Plugins/Experimental/EditorDataStorageFeatures/Source/TedsOutliner/Public/Compatibility/SceneOutlinerTedsBridge.h
@@ -55,6 +55,9 @@ public:
 	// Get the name of the Outliner column corresponding to the given TEDS column (if any)
 	FName FindOutlinerColumnFromTEDSColumns(TConstArrayView<TWeakObjectPtr<const UScriptStruct>> TEDSColumns) const;
 
+	// Refresh all the Outliners currently open in the level editor
+	void RefreshLevelEditorOutliners() const;
+
 private:
 	FSceneOutlinerTedsQueryBinder();
 	void SetupDefaultColumnMapping();
diff --git a/Engine/Plugins/Experimental/EditorDataStorageFeatures/Source/TedsOutliner/TedsOutliner.Build.cs b/Engine/Plugins/Experimental/EditorDataStorageFeatures/Source/TedsOutliner/TedsOutliner.Build.cs
--- a/Engine/Plugins/Experimental/EditorDataStorageFeatures/Source/TedsOutliner/TedsOutliner.Build.cs
+++ b/Engine/Plugins/Experimental/EditorDataStorageFeatures/Source/TedsOutliner/TedsOutliner.Build.cs
@@ -38,7 +38,8 @@ public class TedsOutliner : ModuleRules
 					"TedsAlerts",
 					"TedsTableViewer",
 					"ToolWidgets",
-					"InputCore"
+					"InputCore",
+					"LevelEditor"
 				});
 
 			DynamicallyLoadedModuleNames.AddRange(new string[] {});
diff --git a/Engine/Source/Editor/SceneOutliner/Private/ActorFolderTreeItem.cpp b/Engine/Source/Editor/SceneOutliner/Private/ActorFolderTreeItem.cpp
--- a/Engine/Source/Editor/SceneOutliner/Private/ActorFolderTreeItem.cpp
+++ b/Engine/Source/Editor/SceneOutliner/Private/ActorFolderTreeItem.cpp
@@ -86,8 +86,13 @@ private:
 	bool IsInActorEditorContext() const
 	{
 		auto Folder = TreeItemPtr.Pin();
-		const bool bIsCurrentFolder = (Folder.IsValid() && Folder->World.IsValid()) ? (FActorFolders::Get().GetActorEditorContextFolder(*Folder->World) == Folder->GetFolder()) : false;
-		return bIsCurrentFolder;
+
+		if(Folder->IsValid())
+		{
+			return SceneOutliner::FSceneOutlinerHelpers::IsFolderCurrent(Folder->GetFolder(), Folder->World.Get());
+		}
+
+		return false;
 	}
 
 	FText GetDisplayText() const
@@ -170,19 +175,7 @@ private:
 			if (!TrimmedLabel.ToString().Equals(TreeItem->GetLeafName().ToString(), ESearchCase::CaseSensitive))
 			{
 				// Rename the item
-				FFolder Folder = TreeItem->GetFolder();
-				FName NewPath = Folder.GetParent().GetPath();
-				if (NewPath.IsNone())
-				{
-					NewPath = FName(*TrimmedLabel.ToString());
-				}
-				else
-				{
-					NewPath = FName(*(NewPath.ToString() / TrimmedLabel.ToString()));
-				}
-				FFolder TreeItemNewFolder(Folder.GetRootObject(), NewPath);
-
-				FActorFolders::Get().RenameFolderInWorld(*TreeItem->World.Get(), Folder, TreeItemNewFolder);
+				SceneOutliner::FSceneOutlinerHelpers::RenameFolder(TreeItem->GetFolder(), TrimmedLabel, TreeItem->World.Get());
 
 				auto Outliner = WeakSceneOutliner.Pin();
 				if (Outliner.IsValid())
diff --git a/Engine/Source/Editor/SceneOutliner/Private/ActorHierarchy.cpp b/Engine/Source/Editor/SceneOutliner/Private/ActorHierarchy.cpp
--- a/Engine/Source/Editor/SceneOutliner/Private/ActorHierarchy.cpp
+++ b/Engine/Source/Editor/SceneOutliner/Private/ActorHierarchy.cpp
@@ -163,7 +163,7 @@ FSceneOutlinerTreeItemPtr FActorHierarchy::FindOrCreateParentItem(const ISceneOu
 				}
 				else
 				{
-					return bCreate ? Mode->CreateItemFor<FActorFolderTreeItem>(FActorFolderTreeItem(Folder, ActorTreeItem->Actor->GetWorld()), true) : nullptr;
+					return bCreate ? CreateItemForFolder(Folder, ActorTreeItem->Actor->GetWorld(), true) : nullptr;
 				}
 			}
 
@@ -219,7 +219,7 @@ FSceneOutlinerTreeItemPtr FActorHierarchy::FindOrCreateParentItem(const ISceneOu
 			}
 			else
 			{
-				return bCreate ? Mode->CreateItemFor<FActorFolderTreeItem>(FActorFolderTreeItem(ParentPath, FolderItem->World), true) : nullptr;
+				return bCreate ? CreateItemForFolder(ParentPath, FolderItem->World, true) : nullptr;
 			}
 		}
 		// Parent Level Instance
@@ -285,7 +285,7 @@ FSceneOutlinerTreeItemPtr FActorHierarchy::FindOrCreateParentItem(const ISceneOu
 					}
 					else
 					{
-						return bCreate ? Mode->CreateItemFor<FActorFolderTreeItem>(FActorFolderTreeItem(ActorDescFolder, RepresentingWorldPtr), true) : nullptr;
+						return bCreate ? CreateItemForFolder(ActorDescFolder, RepresentingWorldPtr, true) : nullptr;
 					}
 				}
 			}
@@ -413,7 +413,7 @@ void FActorHierarchy::CreateFolderItems(UWorld* InWorld, TArray<FSceneOutlinerTr
 		// Add any folders which might match the current search terms
 		FActorFolders::Get().ForEachFolder(*InWorld, [this, &InWorld, &OutItems](const FFolder& Folder)
 		{
-			if (FSceneOutlinerTreeItemPtr FolderItem = Mode->CreateItemFor<FActorFolderTreeItem>(FActorFolderTreeItem(Folder, InWorld)))
+			if (FSceneOutlinerTreeItemPtr FolderItem = CreateItemForFolder(Folder, InWorld))
 			{
 				OutItems.Add(FolderItem);
 			}
@@ -432,6 +432,11 @@ FSceneOutlinerTreeItemPtr FActorHierarchy::CreateItemForActorDescInstance(const
 	return Mode->CreateItemFor<FActorDescTreeItem>(InActorDescInstance, bForce);
 }
 
+FSceneOutlinerTreeItemPtr FActorHierarchy::CreateItemForFolder(const FFolder& InFolder, const TWeakObjectPtr<UWorld>& InWorld, bool bForce) const
+{
+	return Mode->CreateItemFor<FActorFolderTreeItem>(FActorFolderTreeItem(InFolder, InWorld), bForce);
+}
+
 bool FActorHierarchy::CheckLevelInstanceEditing(UWorld* World, AActor* Actor) const
 {
 	const ULevelInstanceSubsystem* LevelInstanceSubsystem = World->GetSubsystem<ULevelInstanceSubsystem>();
@@ -461,13 +466,13 @@ void FActorHierarchy::InsertActorItemAndCreateComponents(AActor* InActor, FScene
 	if (bShowingOnlyActorWithValidComponents)
 	{
 		// If the number of items remained the same after component insertion no components were inserted
-		if (OutItems.Num() == InsertLocation)
+		if (OutItems.Num() == InsertLocation) 
 		{
 			return;
 		}
 	}
 
-	if(ActorItem)
+	if (ActorItem)
 	{
 		// Add the actor before the components
 		OutItems.Insert(ActorItem, InsertLocation);
@@ -572,7 +577,7 @@ void FActorHierarchy::CreateWorldChildren(UWorld* World, TArray<FSceneOutlinerTr
 	{
 		AActor* Actor = *ActorIt;
 
-		if(!CheckLevelInstanceEditing(World, Actor))
+		if (!CheckLevelInstanceEditing(World, Actor))
 		{
 			continue;
 		}
@@ -617,7 +622,7 @@ void FActorHierarchy::CreateChildren(const FSceneOutlinerTreeItemPtr& Item, TArr
 		{
 			if (Folder.IsChildOf(InParentFolder))
 			{
-				if (FSceneOutlinerTreeItemPtr NewFolderItem = Mode->CreateItemFor<FActorFolderTreeItem>(FActorFolderTreeItem(Folder, InWorld)))
+				if (FSceneOutlinerTreeItemPtr NewFolderItem = CreateItemForFolder(Folder, InWorld))
 				{
 					OutChildren.Add(NewFolderItem);
 				}
@@ -857,7 +862,7 @@ void FActorHierarchy::OnActorFolderAdded(UActorFolder* InActorFolder)
 	{
 		FSceneOutlinerHierarchyChangedData EventData;
 		EventData.Type = FSceneOutlinerHierarchyChangedData::Added;
-		EventData.Items.Add(Mode->CreateItemFor<FActorFolderTreeItem>(FActorFolderTreeItem(InActorFolder->GetFolder(), RepresentingWorld)));
+		EventData.Items.Add(CreateItemForFolder(InActorFolder->GetFolder(), RepresentingWorld));
 		HierarchyChangedEvent.Broadcast(EventData);
 	}
 }
@@ -922,7 +927,7 @@ void FActorHierarchy::OnBroadcastFolderCreate(UWorld& InWorld, const FFolder& In
 	{
 		FSceneOutlinerHierarchyChangedData EventData;
 		EventData.Type = FSceneOutlinerHierarchyChangedData::Added;
-		EventData.Items.Add(Mode->CreateItemFor<FActorFolderTreeItem>(FActorFolderTreeItem(InNewFolder, &InWorld)));
+		EventData.Items.Add(CreateItemForFolder(InNewFolder, &InWorld));
 		EventData.ItemActions = SceneOutliner::ENewItemAction::Select | SceneOutliner::ENewItemAction::Rename;
 		HierarchyChangedEvent.Broadcast(EventData);
 	}
diff --git a/Engine/Source/Editor/SceneOutliner/Private/SceneOutlinerHelpers.cpp b/Engine/Source/Editor/SceneOutliner/Private/SceneOutlinerHelpers.cpp
--- a/Engine/Source/Editor/SceneOutliner/Private/SceneOutlinerHelpers.cpp
+++ b/Engine/Source/Editor/SceneOutliner/Private/SceneOutlinerHelpers.cpp
@@ -13,7 +13,6 @@
 #include "WorldPartition/WorldPartitionActorDescInstance.h"
 #include "UObject/Package.h"
 
-#include "SourceControlHelpers.h"
 
 #define LOCTEXT_NAMESPACE "SceneOutlinerHelpers"
 
@@ -144,6 +143,27 @@ namespace SceneOutliner
 		}
 	}
 
+	void FSceneOutlinerHelpers::RenameFolder(const FFolder& InFolder, const FText& NewFolderName, UWorld* World)
+	{
+		if (!World)
+		{
+			return;
+		}
+		
+		FName NewPath = InFolder.GetParent().GetPath();
+		if (NewPath.IsNone())
+		{
+			NewPath = FName(*NewFolderName.ToString());
+		}
+		else
+		{
+			NewPath = FName(*(NewPath.ToString() / NewFolderName.ToString()));
+		}
+		
+		const FFolder TreeItemNewFolder(InFolder.GetRootObject(), NewPath);
+		FActorFolders::Get().RenameFolderInWorld(*World, InFolder, TreeItemNewFolder);
+	}
+
 	bool FSceneOutlinerHelpers::ValidateFolderName(const FFolder& InFolder, UWorld* World, const FText& InLabel, FText& OutErrorMessage)
 	{
 		const FText TrimmedLabel = FText::TrimPrecedingAndTrailing(InLabel);
@@ -196,6 +216,14 @@ namespace SceneOutliner
 		return true;
 	}
 
+	bool FSceneOutlinerHelpers::IsFolderCurrent(const FFolder& InFolder, UWorld* World)
+	{
+		if (World)
+		{
+			return FActorFolders::Get().GetActorEditorContextFolder(*World) == InFolder;
+		}
+		return false;
+	}
 }
 
 #undef LOCTEXT_NAMESPACE
\ No newline at end of file
diff --git a/Engine/Source/Editor/SceneOutliner/Public/ActorHierarchy.h b/Engine/Source/Editor/SceneOutliner/Public/ActorHierarchy.h
--- a/Engine/Source/Editor/SceneOutliner/Public/ActorHierarchy.h
+++ b/Engine/Source/Editor/SceneOutliner/Public/ActorHierarchy.h
@@ -68,7 +68,7 @@ protected:
 
 	void OnLevelActorListChanged();
 	
-	void OnActorFolderAdded(UActorFolder* InActorFolder);
+	SCENEOUTLINER_API virtual void OnActorFolderAdded(UActorFolder* InActorFolder);
 	void OnActorFoldersUpdatedEvent(ULevel* InLevel);
 		
 	SCENEOUTLINER_API virtual void OnLevelAdded(ULevel* InLevel, UWorld* InWorld);
@@ -78,13 +78,13 @@ protected:
 	void OnWorldPartitionUninitialized(UWorldPartition* InWorldPartition);
 
 	/** Called when a folder is to be created */
-	void OnBroadcastFolderCreate(UWorld& InWorld, const FFolder& InNewFolder);
+	SCENEOUTLINER_API virtual void OnBroadcastFolderCreate(UWorld& InWorld, const FFolder& InNewFolder);
 
 	/** Called when a folder is to be moved */
-	void OnBroadcastFolderMove(UWorld& InWorld, const FFolder& InOldFolder, const FFolder& InNewFolder);
+	SCENEOUTLINER_API virtual void OnBroadcastFolderMove(UWorld& InWorld, const FFolder& InOldFolder, const FFolder& InNewFolder);
 
 	/** Called when a folder is to be deleted */
-	void OnBroadcastFolderDelete(UWorld& InWorld, const FFolder& InFolder);
+	SCENEOUTLINER_API virtual void OnBroadcastFolderDelete(UWorld& InWorld, const FFolder& InFolder);
 
 	// Remove the unloaded actor item for the given actor
 	SCENEOUTLINER_API void RemoveActorDesc(AActor& InActor);
@@ -113,6 +113,8 @@ protected:
 	// Create a SceneOutlinerTreeItem for the given actor descriptor instance
 	SCENEOUTLINER_API virtual FSceneOutlinerTreeItemPtr CreateItemForActorDescInstance(const FWorldPartitionActorDescInstance* InActorDescInstance, bool bForce = false) const;
 
+	SCENEOUTLINER_API virtual FSceneOutlinerTreeItemPtr CreateItemForFolder(const FFolder& InFolder, const TWeakObjectPtr<UWorld>& InWorld, bool bForce = false) const;
+	
 protected:
 	/** Send a an event indicating a full refresh of the hierarchy is required */
 	void FullRefreshEvent();
diff --git a/Engine/Source/Editor/SceneOutliner/Public/SceneOutlinerHelpers.h b/Engine/Source/Editor/SceneOutliner/Public/SceneOutlinerHelpers.h
--- a/Engine/Source/Editor/SceneOutliner/Public/SceneOutlinerHelpers.h
+++ b/Engine/Source/Editor/SceneOutliner/Public/SceneOutlinerHelpers.h
@@ -21,5 +21,9 @@ namespace SceneOutliner
 		static void PopulateExtraSearchStrings(const ISceneOutlinerTreeItem& TreeItem, TArray< FString >& OutSearchStrings);
 
 		static bool ValidateFolderName(const FFolder& InFolder, UWorld* World, const FText& InLabel, FText& OutErrorMessage);
+
+		static void RenameFolder(const FFolder& InFolder, const FText& NewFolderName, UWorld* World);
+		
+		static bool IsFolderCurrent(const FFolder& InFolder, UWorld* World);
 	};
 };
diff --git a/Engine/Source/Runtime/TypedElementFramework/Public/Elements/Columns/TypedElementSlateWidgetColumns.h b/Engine/Source/Runtime/TypedElementFramework/Public/Elements/Columns/TypedElementSlateWidgetColumns.h
--- a/Engine/Source/Runtime/TypedElementFramework/Public/Elements/Columns/TypedElementSlateWidgetColumns.h
+++ b/Engine/Source/Runtime/TypedElementFramework/Public/Elements/Columns/TypedElementSlateWidgetColumns.h
@@ -88,3 +88,12 @@ struct FSlateColorColumn final : public FEditorDataStorageColumn
 	UPROPERTY()
 	FSlateColor Color = FSlateColor::UseForeground();
 };
+
+/**
+ * Tag added onto widget rows if they are currently in editing mode (e.g editable text boxes)
+ */
+USTRUCT(meta = (DisplayName = "Is Editing"))
+struct FIsInEditingModeTag final : public FEditorDataStorageTag
+{
+	GENERATED_BODY()
+};
