Fixed an unitialized memory issue when populating the asset registry into teds.

Remove the unsed assets columns that are costly to populate.

Moved the folder type population outside of the ar population.This is done in the CB asset data source queries where this is also only proccessed for the paths where the folders can be displayed in the CB.

Avoid reserving to much memory when doing batched modifications in TED processor.

Those minor performence improvement reduce the time to initialize the TEDS.AssetDataStorage by 15%

Removed the TEDS debugs commands. Those are no longer needed and will not be ported forward.

diff --git a/Engine/Plugins/Experimental/EditorDataStorage/Source/TedsCore/Private/Processors/TypedElementProcessorAdaptors.cpp b/Engine/Plugins/Experimental/EditorDataStorage/Source/TedsCore/Private/Processors/TypedElementProcessorAdaptors.cpp
--- a/Engine/Plugins/Experimental/EditorDataStorage/Source/TedsCore/Private/Processors/TypedElementProcessorAdaptors.cpp
+++ b/Engine/Plugins/Experimental/EditorDataStorage/Source/TedsCore/Private/Processors/TypedElementProcessorAdaptors.cpp
@@ -234,37 +234,68 @@ namespace UE::Editor::DataStorage
 
 					// Process the per row operations
 					{
-						int32 Count = 0;
+						TSet<TPair<FMassTagBitSet,FMassFragmentBitSet>
+							, DefaultKeyFuncs<TPair<FMassTagBitSet, FMassFragmentBitSet>>
+							, TInlineSetAllocator<16>> 
+							UniqueAdd;
+						for (const FAddedColumns* const AddColumns : AddedColumnsQueue)
+						{
+							UniqueAdd.Add(TPair<FMassTagBitSet, FMassFragmentBitSet>(AddColumns->AddDescriptor.Tags, AddColumns->AddDescriptor.Fragments));
+						}
+						const int32 AddedColumnsRowGuess = UniqueAdd.IsEmpty() ? 0 : AddedColumnsQueue.Num() / UniqueAdd.Num();
+						UniqueAdd.Empty();
+
+						TSet<const UScriptStruct*
+							, DefaultKeyFuncs<const UScriptStruct*>
+							, TInlineSetAllocator<16>>
+							UniqueColumns;
+						for (const FAddValueColumn* const AddValueColumn : AddedColumnWithRelocationQueue)
+						{
+							UniqueColumns.Add(AddValueColumn->FragmentType);
+						}
+						const int32 AddedColumnsWithRelocationRowGuess = UniqueColumns.IsEmpty() ? 0 : AddedColumnWithRelocationQueue.Num() / UniqueColumns.Num();
+						UniqueColumns.Empty();
+
+						TSet<FDynamicColumnDescription
+							, DefaultKeyFuncs<FDynamicColumnDescription>
+							, TInlineSetAllocator<16>> UniqueDynamicColumnsAdd;
+						for (const FAddDynamicColumn* const AddDynamicColumn : AddedDynamicColumnWithRelocationQueue)
+						{
+							UniqueDynamicColumnsAdd.Add(AddDynamicColumn->Description);
+						}
+						const int32 AddDynamicColumnRowsGuess = UniqueDynamicColumnsAdd.IsEmpty() ? 0 : AddedDynamicColumnWithRelocationQueue.Num() / UniqueDynamicColumnsAdd.Num();
+						UniqueDynamicColumnsAdd.Empty();
+
+
+						int32 AddedColumnsBathRowsCount = 0;
 						for (FAddedColumnsBatch* const AddColumns : AddedColumnsBatchQueue)
 						{
-							Count += AddColumns->Entities.Num();
+							AddedColumnsBathRowsCount = FMath::Max(AddColumns->Entities.Num(), AddedColumnsBathRowsCount);
 						}
-						const int32 AddedColumnsBathRowsCount = Count;
 
-						Count = 0;
+						int32 AddedDynamicColumnsRowCount = 0;
 						for (FAddDynamicColumnsBatch* const AddedDynamicColumns : AddedDynamicColumnsBatchQueue)
 						{
-							Count += AddedDynamicColumns->Rows.Num();
+							AddedDynamicColumnsRowCount += FMath::Max(AddedDynamicColumns->Rows.Num(), AddedDynamicColumnsRowCount);
 						}
-						const int32 AddedDynamicColumnsRowCount = Count;
 
-						Count = 0;
+						int32 RemovedColumnsBatchRowCount = 0;
 						for (FRemovedColumnsBatch* const RemovedColumns : RemovedColumnsBatchQueue)
 						{
-							Count += RemovedColumns->Entities.Num();
+							RemovedColumnsBatchRowCount = FMath::Max(RemovedColumns->Entities.Num(), RemovedColumnsBatchRowCount);
 						}
-						const int32 RemovedColumnsBatchRowCount = Count;
 
 
-						RowOperations.Reserve(FMath::Max(AddedColumnsQueue.Num(),
+						RowOperations.Reserve(FMath::Max(AddedColumnsRowGuess,
 							AddedColumnsBathRowsCount,
-							AddedColumnWithRelocationQueue.Num(),
-							AddedDynamicColumnWithRelocationQueue.Num(),
+							AddedColumnsWithRelocationRowGuess,
+							AddDynamicColumnRowsGuess,
 							AddedDynamicColumnsRowCount,
 							RemovedColumnsQueue.Num(),
 							RemovedColumnsBatchRowCount
 						));
 
+
 						for (FAddedColumns* const AddColumns : AddedColumnsQueue)
 						{
 							FRowOperations& Operations = RowOperations.FindOrAdd(AddColumns->Entity);
diff --git a/Engine/Plugins/Experimental/EditorDataStorageFeatures/Source/TedsAssetData/Private/CB/TedsAssetDataCBDataSource.cpp b/Engine/Plugins/Experimental/EditorDataStorageFeatures/Source/TedsAssetData/Private/CB/TedsAssetDataCBDataSource.cpp
--- a/Engine/Plugins/Experimental/EditorDataStorageFeatures/Source/TedsAssetData/Private/CB/TedsAssetDataCBDataSource.cpp
+++ b/Engine/Plugins/Experimental/EditorDataStorageFeatures/Source/TedsAssetData/Private/CB/TedsAssetDataCBDataSource.cpp
@@ -6,8 +6,11 @@
 
 #include "AssetDefinitionRegistry.h"
 #include "AssetRegistry/IAssetRegistry.h"
+#include "AssetViewTypes.h"
+#include "AssetViewUtils.h"
 #include "Blueprint/BlueprintSupport.h"
 #include "ContentBrowserDataUtils.h"
+#include "Elements/Columns/TypedElementMiscColumns.h"
 #include "Elements/Common/TypedElementQueryTypes.h"
 #include "Elements/Framework/TypedElementQueryBuilder.h"
 #include "Elements/Interfaces/TypedElementDataStorageInterface.h"
@@ -18,7 +21,6 @@
 #include "Internationalization/Text.h"
 #include "Misc/PathViews.h"
 #include "PluginDescriptor.h"
-#include "Elements/Columns/TypedElementMiscColumns.h"
 #include "Settings/ContentBrowserSettings.h"
 #include "TedsAssetDataModule.h"
 #include "Templates/Function.h"
@@ -361,6 +363,27 @@ void FTedsAssetDataCBDataSource::ProcessPathQueryCallback(DataStorage::IQueryCon
 			VirtualPathColumn.VirtualPath = *VirtualPath;
 			// Todo investigate for a batch add Maybe?
 			Context.AddColumn(Rows[Index], MoveTemp(VirtualPathColumn));
+
+
+			FString Path = PathColumn->Path.ToString();
+			FFolderTypeColumn_Experimental FolderType;
+
+			if (AssetViewUtils::IsPluginFolder(Path))
+			{
+				FolderType.FolderType = EFolderType::PluginRoot;
+			}
+			else if (AssetViewUtils::IsDevelopersFolder(Path))
+			{
+				FolderType.FolderType = EFolderType::Developer;
+			}
+			// TODO: Missing CPP Folders, need further change and conversion from the old CBClassDataSource (see FNativeClassHierarchy)
+			// TODO: Missing Virtual Folders, need further change and conversion from the old CBAssetDataCore (see GetItemAttribute)
+			else
+			{
+				FolderType.FolderType = EFolderType::Normal;
+			}
+
+			Context.AddColumn(Rows[Index], MoveTemp(FolderType));
 		}
 	}
 }
diff --git a/Engine/Plugins/Experimental/EditorDataStorageFeatures/Source/TedsAssetData/Private/Debug/TedsAssetDataDebugCommands.cpp b/Engine/Plugins/Experimental/EditorDataStorageFeatures/Source/TedsAssetData/Private/Debug/TedsAssetDataDebugCommands.cpp
deleted file mode 100644
diff --git a/Engine/Plugins/Experimental/EditorDataStorageFeatures/Source/TedsAssetData/Private/TedsAssetData.cpp b/Engine/Plugins/Experimental/EditorDataStorageFeatures/Source/TedsAssetData/Private/TedsAssetData.cpp
--- a/Engine/Plugins/Experimental/EditorDataStorageFeatures/Source/TedsAssetData/Private/TedsAssetData.cpp
+++ b/Engine/Plugins/Experimental/EditorDataStorageFeatures/Source/TedsAssetData/Private/TedsAssetData.cpp
@@ -5,8 +5,6 @@
 #include "TedsAssetDataColumns.h"
 
 #include "AssetRegistry/IAssetRegistry.h"
-#include "AssetViewTypes.h"
-#include "AssetViewUtils.h"
 #include "Async/ParallelFor.h"
 #include "Containers/Array.h"
 #include "Containers/ChunkedArray.h"
@@ -29,10 +27,7 @@ constexpr int32 ParallelForMinBatchSize = 1024 * 4;
 struct FPopulateAssetDataRowArgs
 {
 	FAssetData AssetData;
-
 	FMapKey ObjectPathKey;
-	FMapKey PathKey;
-	RowHandle PathRow = InvalidRowHandle;
 };
 
 // Only safe if the GT is blocked during the operation
@@ -42,18 +37,13 @@ FPopulateAssetDataRowArgs ThreadSafe_PopulateAssetDataTableRow(TAssetData&& InAs
 	FPopulateAssetDataRowArgs Output;
 	Output.ObjectPathKey = FMapKey(InAssetData.GetSoftObjectPath());
 
-	// Looks safe but might not be depending on the implementation of the database
 	if (Database.IsRowAssigned(Database.LookupMappedRow(Output.ObjectPathKey)))
 	{
 		// No need to initialize the rest of the row here. The invalid's asset data will be used as flag to skip the data generated here.
 		return Output;
 	}
 
-	Output.PathKey = FMapKey(InAssetData.PackagePath);
-
-	// Looks safe but might not be depending on the implementation of the database
-	Output.PathRow = Database.LookupMappedRow(Output.PathKey);
-	Output.AssetData = MoveTempIfPossible(InAssetData);
+	Output.AssetData = Forward<TAssetData>(InAssetData);
 
 	return Output;
 }
@@ -61,28 +51,14 @@ FPopulateAssetDataRowArgs ThreadSafe_PopulateAssetDataTableRow(TAssetData&& InAs
 
 void PopulateAssetDataTableRow(FPopulateAssetDataRowArgs&& InAssetDataRowArgs, ICoreProvider& Database, RowHandle RowHandle)
 {
-	if (InAssetDataRowArgs.PathRow != InvalidRowHandle)
-	{
-		Database.GetColumn<FAssetsInPathColumn_Experimental>(InAssetDataRowArgs.PathRow)->AssetsRow.Add(RowHandle);
-	}
-	else
-	{
-		FUnresolvedAssetsInPathColumn_Experimental UnresolvedAssetsInPathColumn;
-		UnresolvedAssetsInPathColumn.Key = MoveTemp(InAssetDataRowArgs.PathKey);
-		Database.AddColumn<FUnresolvedAssetsInPathColumn_Experimental>(RowHandle, MoveTemp(UnresolvedAssetsInPathColumn));
-	}
-
 	Database.GetColumn<FAssetDataColumn_Experimental>(RowHandle)->AssetData = MoveTemp(InAssetDataRowArgs.AssetData);
 }
 
-
 struct FPopulatePathRowArgs
 {
 	FName AssetRegistryPath;
 	FMapKey AssetRegistryPathKey;
-	FMapKey ParentAssetRegistryPathKey;
 	FName AssetName;
-	uint32 PathDepth;
 
 	operator bool() const 
 	{
@@ -95,15 +71,13 @@ struct FPopulatePathRowArgs
 	}
 };
 
-void GetPathDepthAndParentFolderIndex(FStringView Path, uint32& OutDepth, int32& OutParentFolderIndex)
+void GetParentFolderIndex(FStringView Path, int32& OutParentFolderIndex)
 {
-	OutDepth = 0;
 	OutParentFolderIndex = INDEX_NONE;
 
 	if (Path.Len() > 1)
 	{
 		OutParentFolderIndex = 1;
-		++OutDepth;
 	}
 
 	// Skip the first '/'
@@ -111,7 +85,6 @@ void GetPathDepthAndParentFolderIndex(FStringView Path, uint32& OutDepth, int32&
 	{
 		if (Path[Index] == TEXT('/'))
 		{
-			++OutDepth;
 			OutParentFolderIndex = Index;
 		}
 	}
@@ -120,62 +93,24 @@ void GetPathDepthAndParentFolderIndex(FStringView Path, uint32& OutDepth, int32&
 // Only thread safe if the game thread is blocked 
 FPopulatePathRowArgs ThreadSafe_PopulatePathRowArgs(FMapKey&& AssetRegistryPathKey, FName InAssetRegistryPath, FStringView PathAsString)
 {
-	FMapKey ParentAssetRegistryPathKey;
 	int32 CharacterIndex;
-	uint32 Depth;
-	FString AssetName;
-	GetPathDepthAndParentFolderIndex(PathAsString, Depth, CharacterIndex);
+	FStringView FolderName;
+	GetParentFolderIndex(PathAsString, CharacterIndex);
 	if (CharacterIndex != INDEX_NONE)
 	{
-		FStringView ParentPath = PathAsString.Left(CharacterIndex);
-		ParentAssetRegistryPathKey = FMapKey(FName(ParentPath));
-		AssetName = PathAsString.RightChop(CharacterIndex);
+		FolderName = PathAsString.RightChop(CharacterIndex);
 	}
 
 	FPopulatePathRowArgs Args;
 	Args.AssetRegistryPath = InAssetRegistryPath;
-	Args.AssetName = FName(AssetName);
+	Args.AssetName = FName(FolderName);
 	Args.AssetRegistryPathKey = MoveTemp(AssetRegistryPathKey);
-	Args.ParentAssetRegistryPathKey = MoveTemp(ParentAssetRegistryPathKey);
-	Args.PathDepth = Depth;
 
 	return Args;
 }
 
 void PopulatePathDataTableRow(FPopulatePathRowArgs&& InPopulatePathRowArgs, ICoreProvider& Database, RowHandle InRowHandle)
 {
-	if (InPopulatePathRowArgs.ParentAssetRegistryPathKey.IsSet())
-	{
-		const RowHandle ParentRow = Database.LookupMappedRow(InPopulatePathRowArgs.ParentAssetRegistryPathKey);
-		if (Database.IsRowAssigned(ParentRow))
-		{
-			Database.GetColumn<FChildrenAssetPathColumn_Experimental>(ParentRow)->ChildrenRows.Add(InRowHandle);
-			Database.GetColumn<FParentAssetPathColumn_Experimental>(InRowHandle)->ParentRow = ParentRow;
-		}
-		else 
-		{
-			// If we were unlucky we may be missing the parent path data.This should resolve it self during a latter event on paths added event.
-			FUnresolvedParentAssetPathColumn_Experimental UnresolvedParentRow;
-			UnresolvedParentRow.Key = MoveTemp(InPopulatePathRowArgs.ParentAssetRegistryPathKey);
-			Database.AddColumn(InRowHandle, MoveTemp(UnresolvedParentRow));
-		}
-	}
-
-	if (AssetViewUtils::IsDevelopersFolder(InPopulatePathRowArgs.AssetRegistryPath.ToString()))
-	{
-		Database.GetColumn<FFolderTypeColumn_Experimental>(InRowHandle)->FolderType = EFolderType::Developer;
-	}
-	else if (AssetViewUtils::IsPluginFolder(InPopulatePathRowArgs.AssetRegistryPath.ToString()))
-	{
-		Database.GetColumn<FFolderTypeColumn_Experimental>(InRowHandle)->FolderType = EFolderType::PluginRoot;
-	}
-	// TODO: Missing CPP Folders, need further change and conversion from the old CBClassDataSource (see FNativeClassHierarchy)
-	// TODO: Missing Virtual Folders, need further change and conversion from the old CBAssetDataCore (see GetItemAttribute)
-	else
-	{
-		Database.GetColumn<FFolderTypeColumn_Experimental>(InRowHandle)->FolderType = EFolderType::Normal;
-	}
-
 	Database.GetColumn<FAssetPathColumn_Experimental>(InRowHandle)->Path = InPopulatePathRowArgs.AssetRegistryPath;
 	Database.GetColumn<FNameColumn>(InRowHandle)->Name = InPopulatePathRowArgs.AssetName;
 }
@@ -206,7 +141,7 @@ FTedsAssetData::FTedsAssetData(UE::Editor::DataStorage::ICoreProvider& InDatabas
 	PathsTable = Database.FindTable(FName(TEXT("Editor_AssetRegistryPathsTable")));
 	if (PathsTable == InvalidTableHandle)
 	{
-		PathsTable = Database.RegisterTable<FFolderTag, FFolderTypeColumn_Experimental, FAssetPathColumn_Experimental, FNameColumn, FUpdatedPathTag, FSlateColorColumn, FChildrenAssetPathColumn_Experimental, FParentAssetPathColumn_Experimental, FAssetsInPathColumn_Experimental>(FName(TEXT("Editor_AssetRegistryPathsTable")));
+		PathsTable = Database.RegisterTable<FFolderTag, FAssetPathColumn_Experimental, FNameColumn, FUpdatedPathTag>(FName(TEXT("Editor_AssetRegistryPathsTable")));
 	}
 
 	AssetsDataTable = Database.FindTable(FName(TEXT("Editor_AssetRegistryAssetDataTable")));
@@ -240,77 +175,7 @@ FTedsAssetData::FTedsAssetData(UE::Editor::DataStorage::ICoreProvider& InDatabas
 			.All<FUpdatedAssetDataTag>()
 		.Compile());
 
-	UpdateAssetsInPathQuery = Database.RegisterQuery(
-		Select()
-			.ReadWrite<FAssetsInPathColumn_Experimental>()
-		.Compile());
-
-	ResolveMissingAssetInPathQuery = Database.RegisterQuery(
-		Select(
-			TEXT("FTedsAssetData: Resolve Missing Asset In Path"),
-			FProcessor(EQueryTickPhase::FrameEnd, Database.GetQueryTickGroupName(EQueryTickGroups::Default)),
-			[this](IQueryContext& Context, RowHandle Row, const FUnresolvedAssetsInPathColumn_Experimental& UnresolvedAssetPath)
-			{
-#if TRACK_TEDSASSETDATA_MEMORY
-				LLM_SCOPE_BYNAME(TEXT("FTedsAssetData"))
-#endif
-
-				const RowHandle PathRow = Context.LookupMappedRow(UnresolvedAssetPath.Key);
-				if (Context.IsRowAssigned(PathRow))
-				{
-					Context.RemoveColumns<FUnresolvedAssetsInPathColumn_Experimental>(Row);
-					Context.RunSubquery(0, PathRow, 
-						CreateSubqueryCallbackBinding([Row](FAssetsInPathColumn_Experimental& AssetsInPath)
-						{
-#if TRACK_TEDSASSETDATA_MEMORY
-							LLM_SCOPE_BYNAME(TEXT("FTedsAssetData"))
-#endif
-
-							AssetsInPath.AssetsRow.Add(Row);
-						}));
-				}
-			})
-		.DependsOn()
-			.SubQuery(UpdateAssetsInPathQuery)
-		.Compile());
-
-	UpdateParentToChildrenAssetPathQuery = Database.RegisterQuery(
-		Select()
-			.ReadWrite<FChildrenAssetPathColumn_Experimental>()
-		.Compile());
-
-	ResolveMissingParentPathQuery = Database.RegisterQuery(
-		Select(
-			TEXT("FTedsAssetData: Resolve Missing Parent Path Row"),
-			FProcessor(EQueryTickPhase::FrameEnd, Database.GetQueryTickGroupName(EQueryTickGroups::Default)),
-			[this](IQueryContext& Context, RowHandle Row, const FUnresolvedParentAssetPathColumn_Experimental& UnresolvedParentAssetPath, FParentAssetPathColumn_Experimental& ParentAssetPathColumn)
-			{
-#if TRACK_TEDSASSETDATA_MEMORY
-				LLM_SCOPE_BYNAME(TEXT("FTedsAssetData"))
-#endif
-
-				const RowHandle ParentPathRow = Context.LookupMappedRow(UnresolvedParentAssetPath.Key);
-				if (Context.IsRowAssigned(ParentPathRow))
-				{
-					Context.RemoveColumns<FUnresolvedParentAssetPathColumn_Experimental>(Row);
-					ParentAssetPathColumn.ParentRow = ParentPathRow;
-
-					Context.RunSubquery(0, ParentPathRow, 
-						CreateSubqueryCallbackBinding([Row](FChildrenAssetPathColumn_Experimental& ChildrenPathColumn)
-						{
-#if TRACK_TEDSASSETDATA_MEMORY
-							LLM_SCOPE_BYNAME(TEXT("FTedsAssetData"))
-#endif
-							ChildrenPathColumn.ChildrenRows.Add(Row);
-						}));
-				}
-			})
-		.DependsOn()
-			.SubQuery(UpdateParentToChildrenAssetPathQuery)
-		.Compile());
-
 	// Init with the data existing at moment in asset registry
-
 	TArray<FAssetData> AssetsData;
 	AssetRegistry.GetAllAssets(AssetsData);
 
@@ -369,7 +234,7 @@ FTedsAssetData::FTedsAssetData(UE::Editor::DataStorage::ICoreProvider& InDatabas
 	TArrayView<TPair<FMapKey, RowHandle>> KeysToResevedAssetRows(KeysToReservedRows.GetData() + ReservedPopulatePathRows.Num(), ReservedAssetDataRows.Num());
 	ParallelFor(TEXT("Populating TEDS Asset Registry Asset Data Keys"), PopulateAssetDataRowArgs.Num(), ParallelForMinBatchSize, [&PopulateAssetDataRowArgs, &ReservedAssetDataRows, &KeysToResevedAssetRows](int32 Index)
 		{
-			KeysToResevedAssetRows[Index] = TPair<FMapKey, RowHandle>(PopulateAssetDataRowArgs[Index].ObjectPathKey, ReservedAssetDataRows[Index]);
+			KeysToResevedAssetRows[Index] = TPair<FMapKey, RowHandle>(MoveTemp(PopulateAssetDataRowArgs[Index].ObjectPathKey), ReservedAssetDataRows[Index]);
 		});
 	Database.BatchMapRows(KeysToResevedAssetRows);
 
@@ -394,10 +259,6 @@ FTedsAssetData::~FTedsAssetData()
 	{
 		if (IAssetRegistry* AssetRegistry = IAssetRegistry::Get())
 		{
-			Database.UnregisterQuery(ResolveMissingParentPathQuery);
-			Database.UnregisterQuery(UpdateParentToChildrenAssetPathQuery);
-			Database.UnregisterQuery(ResolveMissingAssetInPathQuery);
-			Database.UnregisterQuery(UpdateAssetsInPathQuery);
 			Database.UnregisterQuery(RemoveUpdatedAssetDataTagQuery);
 			Database.UnregisterQuery(RemoveUpdatedPathTagQuery);
 
@@ -488,7 +349,7 @@ void FTedsAssetData::OnAssetsAdded(TConstArrayView<FAssetData> InAssetsAdded)
 					ARowArgs = MoveTemp(RowArgs[Index]);
 				}
 
-				KeyToRow.Emplace(ARowArgs.ObjectPathKey, InRowHandle);
+				KeyToRow.Emplace(MoveTemp(ARowArgs.ObjectPathKey), InRowHandle);
 				PopulateAssetDataTableRow(MoveTemp(ARowArgs), Database, InRowHandle);
 				++Index;
 			});
@@ -505,24 +366,20 @@ void FTedsAssetData::OnAssetsRemoved(TConstArrayView<FAssetData> InAssetsRemoved
 	LLM_SCOPE_BYNAME(TEXT("FTedsAssetData"))
 #endif
 
+	TArray<RowHandle> RowsToRemove;
+	RowsToRemove.Reserve(InAssetsRemoved.Num());
+
 	for (const FAssetData& Asset : InAssetsRemoved)
 	{
 		const FMapKey AssetKey = FMapKey(Asset.GetSoftObjectPath());
 		const RowHandle AssetRow = Database.LookupMappedRow(AssetKey);
 		if (Database.IsRowAssigned(AssetRow))
 		{
-			if (const FAssetDataColumn_Experimental* AssetDataColunm = Database.GetColumn<FAssetDataColumn_Experimental>(AssetRow))
-			{
-				const FMapKey FolderPathKey = FMapKey(AssetDataColunm->AssetData.PackagePath);
-				const RowHandle FolderRow = Database.LookupMappedRow(FolderPathKey);
-				if (FAssetsInPathColumn_Experimental* AssetInFolder = Database.GetColumn<FAssetsInPathColumn_Experimental>(FolderRow))
-				{
-					AssetInFolder->AssetsRow.Remove(AssetRow);
-				}
-			}
-			Database.RemoveRow(AssetRow);
+			RowsToRemove.Add(AssetRow);
 		}
 	}
+
+	Database.BatchRemoveRows(RowsToRemove);
 }
 
 void FTedsAssetData::OnAssetsUpdated(TConstArrayView<FAssetData> InAssetsUpdated)
@@ -594,31 +451,6 @@ void FTedsAssetData::OnAssetRenamed(const FAssetData& InAsset, const FString& In
 
 		const FMapKey OldFolderHash = FMapKey(FName(OldPackagePath));
 
-		if (NewFolderHash != OldFolderHash)
-		{
-			const RowHandle NewPathRow = Database.LookupMappedRow(NewFolderHash);
-			if (NewPathRow !=  InvalidRowHandle)
-			{
-				Database.GetColumn<FAssetsInPathColumn_Experimental>(NewPathRow)->AssetsRow.Add(Row);
-			}
-			else
-			{
-				FUnresolvedAssetsInPathColumn_Experimental UnresolvedAssetsInPathColumn;
-				UnresolvedAssetsInPathColumn.Key = NewFolderHash.CreateKey();
-				Database.AddColumn<FUnresolvedAssetsInPathColumn_Experimental>(Row, MoveTemp(UnresolvedAssetsInPathColumn));
-			}
-
-			const RowHandle OldPathRow = Database.LookupMappedRow(OldFolderHash);
-			if (NewPathRow !=  InvalidRowHandle)
-			{
-				Database.GetColumn<FAssetsInPathColumn_Experimental>(NewPathRow)->AssetsRow.Remove(Row);
-			}
-			else
-			{
-				Database.RemoveColumn<FUnresolvedAssetsInPathColumn_Experimental>(Row);
-			}
-		}
-
 		Database.AddColumn<FUpdatedPathTag>(Row);
 		Database.RemapRow(OldAssetHash, MoveTemp(NewAssetHash));
 	}
@@ -634,7 +466,7 @@ void FTedsAssetData::OnPathsAdded(TConstArrayView<FStringView> InPathsAdded)
 
 	TArray<int32> PreparePathRowContexts;
 	TArray<FPopulatePathRowArgs> PopulateRowArgs;
-	PopulateRowArgs.AddUninitialized(InPathsAdded.Num());
+	PopulateRowArgs.AddDefaulted(InPathsAdded.Num());
 
 	ParallelForWithTaskContext(TEXT("Populating TEDS Asset Registry Path"), PreparePathRowContexts, InPathsAdded.Num(), ParallelForMinBatchSize, [&PopulateRowArgs, &InPathsAdded, this](int32& WorkerValidCount, int32 Index)
 		{
diff --git a/Engine/Plugins/Experimental/EditorDataStorageFeatures/Source/TedsAssetData/Private/TedsAssetData.h b/Engine/Plugins/Experimental/EditorDataStorageFeatures/Source/TedsAssetData/Private/TedsAssetData.h
--- a/Engine/Plugins/Experimental/EditorDataStorageFeatures/Source/TedsAssetData/Private/TedsAssetData.h
+++ b/Engine/Plugins/Experimental/EditorDataStorageFeatures/Source/TedsAssetData/Private/TedsAssetData.h
@@ -42,11 +42,6 @@ namespace UE::Editor::AssetData::Private
 		DataStorage::TableHandle PathsTable = DataStorage::InvalidTableHandle;
 		DataStorage::TableHandle AssetsDataTable = DataStorage::InvalidTableHandle;
 
-		DataStorage::QueryHandle UpdateAssetsInPathQuery = DataStorage::InvalidQueryHandle;
-		DataStorage::QueryHandle UpdateAssetsInPath = DataStorage::InvalidQueryHandle;
-		DataStorage::QueryHandle ResolveMissingAssetInPathQuery = DataStorage::InvalidQueryHandle;
-		DataStorage::QueryHandle UpdateParentToChildrenAssetPathQuery = DataStorage::InvalidQueryHandle;
-		DataStorage::QueryHandle ResolveMissingParentPathQuery = DataStorage::InvalidQueryHandle;
 		DataStorage::QueryHandle RemoveUpdatedPathTagQuery = DataStorage::InvalidQueryHandle;
 		DataStorage::QueryHandle RemoveUpdatedAssetDataTagQuery = DataStorage::InvalidQueryHandle;
 	};
diff --git a/Engine/Plugins/Experimental/EditorDataStorageFeatures/Source/TedsAssetData/Private/Tests/TedsAssetDataTest.cpp b/Engine/Plugins/Experimental/EditorDataStorageFeatures/Source/TedsAssetData/Private/Tests/TedsAssetDataTest.cpp
--- a/Engine/Plugins/Experimental/EditorDataStorageFeatures/Source/TedsAssetData/Private/Tests/TedsAssetDataTest.cpp
+++ b/Engine/Plugins/Experimental/EditorDataStorageFeatures/Source/TedsAssetData/Private/Tests/TedsAssetDataTest.cpp
@@ -82,11 +82,6 @@ bool FTedsAssetDataTest::RunTest(const FString& Parameters)
 
 		bHadError |= !TestEqual(TEXT("Path Stored In TEDS Match Path From Asset Registry"), NameStoredInTeds, PathAndAssetCount.Key);
 
-		if (const FAssetsInPathColumn_Experimental* AssetsInPathColumn = Database->GetColumn<FAssetsInPathColumn_Experimental>(Row))
-		{
-			bHadError |= !TestEqual(TEXT("Teds should have the same assets amount in a path as the asset registry has."), AssetsInPathColumn->AssetsRow.Num(), PathAndAssetCount.Value);
-		}
-
 		if (bHadError)
 		{
 			TArray<FStringFormatArg> FormatArgs;
@@ -116,23 +111,6 @@ bool FTedsAssetDataTest::RunTest(const FString& Parameters)
 
 			bHadError |= !TestNotNull(TEXT("Teds doesn't have a asset column for an asset of the asset registry"), Database->GetColumn<FAssetDataColumn_Experimental>(Row));
 
-			if (!Database->GetColumn<FUnresolvedAssetsInPathColumn_Experimental>(Row))
-			{
-				RowHandle PathRow = Database->LookupMappedRow(FMapKeyView(AssetData.PackagePath));
-			
-				const FAssetsInPathColumn_Experimental* AssetsInPathColumn = Database->GetColumn<FAssetsInPathColumn_Experimental>(PathRow);
-
-				const bool bPassNullTest = TestNotNull(TEXT("Asset in path column does exist for the package path of the asset"), AssetsInPathColumn);
-				if (bPassNullTest)
-				{
-					bHadError |= !TestTrue(TEXT("The asset is stored properly in the assets in path column"), AssetsInPathColumn->AssetsRow.Contains(Row));
-				}
-				else
-				{
-					bHadError = true;
-				}
-			}
-
 			if (bHadError)
 			{
 				TArray<FStringFormatArg> FormatArgs;
diff --git a/Engine/Plugins/Experimental/EditorDataStorageFeatures/Source/TedsAssetData/Public/TedsAssetDataColumns.h b/Engine/Plugins/Experimental/EditorDataStorageFeatures/Source/TedsAssetData/Public/TedsAssetDataColumns.h
--- a/Engine/Plugins/Experimental/EditorDataStorageFeatures/Source/TedsAssetData/Public/TedsAssetDataColumns.h
+++ b/Engine/Plugins/Experimental/EditorDataStorageFeatures/Source/TedsAssetData/Public/TedsAssetDataColumns.h
@@ -26,30 +26,6 @@ struct FAssetPathColumn_Experimental : public FEditorDataStorageColumn
 	FName Path;
 };
 
-USTRUCT()
-struct FParentAssetPathColumn_Experimental : public FEditorDataStorageColumn
-{
-	GENERATED_BODY()
-	
-	UE::Editor::DataStorage::RowHandle ParentRow;
-};
-
-USTRUCT()
-struct FChildrenAssetPathColumn_Experimental : public FEditorDataStorageColumn
-{
-	GENERATED_BODY()
-	
-	TSet<UE::Editor::DataStorage::RowHandle> ChildrenRows;
-};
-
-USTRUCT()
-struct FAssetsInPathColumn_Experimental : public FEditorDataStorageColumn
-{
-	GENERATED_BODY()
-	
-	TSet<UE::Editor::DataStorage::RowHandle> AssetsRow;
-};
-
 USTRUCT()
 struct FFolderTypeColumn_Experimental : public FEditorDataStorageColumn
 {
@@ -67,22 +43,6 @@ struct FAssetDataColumn_Experimental : public FEditorDataStorageColumn
 	FAssetData AssetData;
 };
 
-USTRUCT()
-struct FUnresolvedParentAssetPathColumn_Experimental : public FEditorDataStorageColumn
-{
-	GENERATED_BODY()
-	
-	UE::Editor::DataStorage::FMapKey Key;
-};
-
-USTRUCT()
-struct FUnresolvedAssetsInPathColumn_Experimental : public FEditorDataStorageColumn
-{
-	GENERATED_BODY()
-	
-	UE::Editor::DataStorage::FMapKey Key;
-};
-
 // Tag to identify assets
 USTRUCT(meta = (DisplayName = "Asset"))
 struct FAssetTag final : public FEditorDataStorageTag
