CB 2.0 - add a new purpose for CB thumbnail widgets and use all columns on the target row with longest match to find the best matching thumbnail widget.
This allows individual asset types to override the thumbnail widget

diff --git a/Engine/Plugins/Experimental/EditorDataStorageFeatures/Source/TedsAssetData/Private/Widgets/AssetDataThumbnailWidget.cpp b/Engine/Plugins/Experimental/EditorDataStorageFeatures/Source/TedsAssetData/Private/Widgets/AssetDataThumbnailWidget.cpp
--- a/Engine/Plugins/Experimental/EditorDataStorageFeatures/Source/TedsAssetData/Private/Widgets/AssetDataThumbnailWidget.cpp
+++ b/Engine/Plugins/Experimental/EditorDataStorageFeatures/Source/TedsAssetData/Private/Widgets/AssetDataThumbnailWidget.cpp
@@ -242,12 +242,21 @@ namespace Private
 	};
 }
 
+void UAssetThumbnailWidgetFactory::RegisterWidgetPurposes(UE::Editor::DataStorage::IUiProvider& DataStorageUi) const
+{
+	DataStorageUi.RegisterWidgetPurpose(
+	UE::Editor::DataStorage::IUiProvider::FPurposeInfo("ContentBrowser", "Thumbnail", NAME_None,
+		UE::Editor::DataStorage::IUiProvider::EPurposeType::UniqueByNameAndColumn,
+		LOCTEXT("CBThumbnailPurpose", "Specific purpose display thumbnails in the CB.")));
+}
+
+
 void UAssetThumbnailWidgetFactory::RegisterWidgetConstructors(UE::Editor::DataStorage::ICoreProvider& DataStorage,
-	UE::Editor::DataStorage::IUiProvider& DataStorageUi) const
+                                                              UE::Editor::DataStorage::IUiProvider& DataStorageUi) const
 {
 	using namespace UE::Editor::DataStorage::Queries;
 	DataStorageUi.RegisterWidgetFactory<FAssetThumbnailWidgetConstructor>(
-		DataStorageUi.FindPurpose(DataStorageUi.GetGeneralWidgetPurposeID()),
+		DataStorageUi.FindPurpose(IUiProvider::FPurposeInfo("ContentBrowser", "Thumbnail", NAME_None).GeneratePurposeID()),
 		TColumn<FAssetTag>() || TColumn<FFolderTag>());
 }
 
diff --git a/Engine/Plugins/Experimental/EditorDataStorageFeatures/Source/TedsAssetData/Private/Widgets/AssetDataThumbnailWidget.h b/Engine/Plugins/Experimental/EditorDataStorageFeatures/Source/TedsAssetData/Private/Widgets/AssetDataThumbnailWidget.h
--- a/Engine/Plugins/Experimental/EditorDataStorageFeatures/Source/TedsAssetData/Private/Widgets/AssetDataThumbnailWidget.h
+++ b/Engine/Plugins/Experimental/EditorDataStorageFeatures/Source/TedsAssetData/Private/Widgets/AssetDataThumbnailWidget.h
@@ -19,6 +19,7 @@ class UAssetThumbnailWidgetFactory : public UEditorDataStorageFactory
 public:
 	virtual ~UAssetThumbnailWidgetFactory() override = default;
 
+	TEDSASSETDATA_API virtual void RegisterWidgetPurposes(UE::Editor::DataStorage::IUiProvider& DataStorageUi) const override;
 	TEDSASSETDATA_API virtual void RegisterWidgetConstructors(UE::Editor::DataStorage::ICoreProvider& DataStorage,
 		UE::Editor::DataStorage::IUiProvider& DataStorageUi) const override;
 };
diff --git a/Engine/Plugins/Experimental/EditorDataStorageFeatures/Source/TedsAssetData/Private/Widgets/ContentBrowserListViewNameWidget.cpp b/Engine/Plugins/Experimental/EditorDataStorageFeatures/Source/TedsAssetData/Private/Widgets/ContentBrowserListViewNameWidget.cpp
--- a/Engine/Plugins/Experimental/EditorDataStorageFeatures/Source/TedsAssetData/Private/Widgets/ContentBrowserListViewNameWidget.cpp
+++ b/Engine/Plugins/Experimental/EditorDataStorageFeatures/Source/TedsAssetData/Private/Widgets/ContentBrowserListViewNameWidget.cpp
@@ -83,14 +83,23 @@ TSharedPtr<SWidget> FContentBrowserListViewNameWidgetConstructor::CreateWidget(U
 		};
 
 		// TODO: Arguments need to be developed further to customize the Thumbnail Config if needed
-		TArray<TWeakObjectPtr<const UScriptStruct>> ThumbnailColumns = GetThumbnailColumns();
 		UE::Editor::DataStorage::FMetaData ThumbnailMeta;
 		ThumbnailMeta.AddOrSetMutableData(TedsAssetDataHelper::MetaDataNames::GetThumbnailStatusMetaDataName(), true);
 		ThumbnailMeta.AddOrSetMutableData(TedsAssetDataHelper::MetaDataNames::GetThumbnailFadeInMetaDataName(), true);
 		ThumbnailMeta.AddOrSetMutableData(TedsAssetDataHelper::MetaDataNames::GetThumbnailHintTextMetaDataName(), false);
 		// TODO: use our own OnMouseEnter/Leave for setting the realtime flag, also check CB Setting
 		ThumbnailMeta.AddOrSetMutableData(TedsAssetDataHelper::MetaDataNames::GetThumbnailRealTimeOnHoveredMetaDataName(), false);
-		DataStorageUi->CreateWidgetConstructors(DataStorageUi->FindPurpose(DataStorageUi->GetGeneralWidgetPurposeID()), IUiProvider::EMatchApproach::ExactMatch, ThumbnailColumns, Arguments, AssignThumbnailWidgetToColumn);
+		
+		// List all columns on the row so the thumbnail widget can be override based on longest match on the current row
+		TArray<TWeakObjectPtr<const UScriptStruct>> ThumbnailColumns;
+		DataStorage->ListColumns(TargetRow, [&ThumbnailColumns](const UScriptStruct& ColumnType)
+			{
+				ThumbnailColumns.Emplace(&ColumnType);
+				return true;
+			});
+
+		const IUiProvider::FPurposeID ThumbnailPurposeID = IUiProvider::FPurposeInfo("ContentBrowser", "Thumbnail", NAME_None).GeneratePurposeID();
+		DataStorageUi->CreateWidgetConstructors(DataStorageUi->FindPurpose(ThumbnailPurposeID), IUiProvider::EMatchApproach::LongestMatch, ThumbnailColumns, Arguments, AssignThumbnailWidgetToColumn);
 
 		TArray<TWeakObjectPtr<const UScriptStruct>> LabelColumns = GetLabelColumns();
 
@@ -169,12 +178,6 @@ TSharedPtr<SWidget> FContentBrowserListViewNameWidgetConstructor::CreateWidget(U
 	return SNullWidget::NullWidget;
 }
 
-TArray<TWeakObjectPtr<const UScriptStruct>> FContentBrowserListViewNameWidgetConstructor::GetThumbnailColumns()
-{
-	static TArray<TWeakObjectPtr<const UScriptStruct>> ThumbnailColumns({ TWeakObjectPtr(FAssetTag::StaticStruct()), TWeakObjectPtr(FFolderTag::StaticStruct()) });
-	return ThumbnailColumns;
-}
-
 TArray<TWeakObjectPtr<const UScriptStruct>> FContentBrowserListViewNameWidgetConstructor::GetLabelColumns()
 {
 	static TArray<TWeakObjectPtr<const UScriptStruct>> LabelColumns({ TWeakObjectPtr(FNameColumn::StaticStruct()) });
diff --git a/Engine/Plugins/Experimental/EditorDataStorageFeatures/Source/TedsAssetData/Private/Widgets/ContentBrowserListViewNameWidget.h b/Engine/Plugins/Experimental/EditorDataStorageFeatures/Source/TedsAssetData/Private/Widgets/ContentBrowserListViewNameWidget.h
--- a/Engine/Plugins/Experimental/EditorDataStorageFeatures/Source/TedsAssetData/Private/Widgets/ContentBrowserListViewNameWidget.h
+++ b/Engine/Plugins/Experimental/EditorDataStorageFeatures/Source/TedsAssetData/Private/Widgets/ContentBrowserListViewNameWidget.h
@@ -50,6 +50,5 @@ public:
 		const UE::Editor::DataStorage::FMetaDataView& Arguments) override;
 
 private:
-	TArray<TWeakObjectPtr<const UScriptStruct>> GetThumbnailColumns();
 	TArray<TWeakObjectPtr<const UScriptStruct>> GetLabelColumns();
 };
diff --git a/Engine/Plugins/Experimental/EditorDataStorageFeatures/Source/TedsContentBrowser/Private/Widgets/ContentBrowserTileViewWidget.cpp b/Engine/Plugins/Experimental/EditorDataStorageFeatures/Source/TedsContentBrowser/Private/Widgets/ContentBrowserTileViewWidget.cpp
--- a/Engine/Plugins/Experimental/EditorDataStorageFeatures/Source/TedsContentBrowser/Private/Widgets/ContentBrowserTileViewWidget.cpp
+++ b/Engine/Plugins/Experimental/EditorDataStorageFeatures/Source/TedsContentBrowser/Private/Widgets/ContentBrowserTileViewWidget.cpp
@@ -502,8 +502,7 @@ TSharedPtr<SWidget> FContentBrowserTileViewWidgetConstructor::CreateThumbnailWid
 		OutThumbnailWidgetConstructorPtr = TSharedPtr<FTypedElementWidgetConstructor>(Constructor.Release());
 		return false;
 	};
-
-	TArray<TWeakObjectPtr<const UScriptStruct>> ThumbnailColumns = GetThumbnailColumns();
+	
 	UE::Editor::DataStorage::FMetaData ThumbnailMeta;
 	ThumbnailMeta.AddOrSetMutableData(TedsAssetDataHelper::MetaDataNames::GetThumbnailStatusMetaDataName(), true);
 	ThumbnailMeta.AddOrSetMutableData(TedsAssetDataHelper::MetaDataNames::GetThumbnailFadeInMetaDataName(), true);
@@ -518,8 +517,18 @@ TSharedPtr<SWidget> FContentBrowserTileViewWidgetConstructor::CreateThumbnailWid
 	}
 	const FGenericMetaDataView ThumbnailMetaView = FGenericMetaDataView(ThumbnailMeta);
 
-	const RowHandle DefaultPurposeRowHandle = InTedsTileViewCommonArguments.DataStorageUi->FindPurpose(InTedsTileViewCommonArguments.DataStorageUi->GetGeneralWidgetPurposeID());
-	InTedsTileViewCommonArguments.DataStorageUi->CreateWidgetConstructors(DefaultPurposeRowHandle, IUiProvider::EMatchApproach::ExactMatch, ThumbnailColumns, ThumbnailMetaView, AssignThumbnailWidgetToColumn);
+	// List all columns on the row so the thumbnail widget can be override based on longest match on the current row
+	TArray<TWeakObjectPtr<const UScriptStruct>> ThumbnailColumns;
+	InTedsTileViewCommonArguments.DataStorage->ListColumns(InTedsTileViewCommonArguments.TargetRow, [&ThumbnailColumns](const UScriptStruct& ColumnType)
+		{
+			ThumbnailColumns.Emplace(&ColumnType);
+			return true;
+		});
+	const IUiProvider::FPurposeID ThumbnailPurposeID = IUiProvider::FPurposeInfo("ContentBrowser", "Thumbnail", NAME_None).GeneratePurposeID();
+	const RowHandle ThumbnailPurposeRowHandle = InTedsTileViewCommonArguments.DataStorageUi->FindPurpose(ThumbnailPurposeID);
+	
+	InTedsTileViewCommonArguments.DataStorageUi->CreateWidgetConstructors(ThumbnailPurposeRowHandle, IUiProvider::EMatchApproach::LongestMatch,
+		ThumbnailColumns, ThumbnailMetaView, AssignThumbnailWidgetToColumn);
 
 	TSharedPtr<SWidget> ThumbnailWidget = SNullWidget::NullWidget;
 	if (OutThumbnailWidgetConstructorPtr)
@@ -607,12 +616,6 @@ TSharedPtr<SWidget> FContentBrowserTileViewWidgetConstructor::CreateItemTypeWidg
 	return ItemTypeWidget;
 }
 
-TArray<TWeakObjectPtr<const UScriptStruct>> FContentBrowserTileViewWidgetConstructor::GetThumbnailColumns() const
-{
-	static TArray<TWeakObjectPtr<const UScriptStruct>> ThumbnailColumns({ TWeakObjectPtr(FAssetTag::StaticStruct()), TWeakObjectPtr(FFolderTag::StaticStruct()) });
-	return ThumbnailColumns;
-}
-
 TArray<TWeakObjectPtr<const UScriptStruct>> FContentBrowserTileViewWidgetConstructor::GetItemTypeColumns() const
 {
 	static TArray<TWeakObjectPtr<const UScriptStruct>> ItemTypeColumns({ TWeakObjectPtr(FAssetClassColumn::StaticStruct()), TWeakObjectPtr(FFolderTag::StaticStruct()) });
diff --git a/Engine/Plugins/Experimental/EditorDataStorageFeatures/Source/TedsContentBrowser/Private/Widgets/ContentBrowserTileViewWidget.h b/Engine/Plugins/Experimental/EditorDataStorageFeatures/Source/TedsContentBrowser/Private/Widgets/ContentBrowserTileViewWidget.h
--- a/Engine/Plugins/Experimental/EditorDataStorageFeatures/Source/TedsContentBrowser/Private/Widgets/ContentBrowserTileViewWidget.h
+++ b/Engine/Plugins/Experimental/EditorDataStorageFeatures/Source/TedsContentBrowser/Private/Widgets/ContentBrowserTileViewWidget.h
@@ -106,9 +106,6 @@ private:
 	/** Create the item type widget through TEDS */
 	TSharedPtr<SWidget> CreateItemTypeWidget(FTedsTileViewCommonArguments InTedsTileViewCommonArguments, UE::Editor::DataStorage::RowHandle& OutItemTypeWidgetRowHandle) const;
 
-	/** Get the Thumbnail columns */
-	TArray<TWeakObjectPtr<const UScriptStruct>> GetThumbnailColumns() const;
-	
 	/** Get the ItemType columns */
 	TArray<TWeakObjectPtr<const UScriptStruct>> GetItemTypeColumns() const;
 };
