TedsSettings : Add a FindSettingsSection method.

diff --git a/Engine/Plugins/Experimental/EditorDataStorageFeatures/Source/TedsSettings/Private/TedsSettingsEditorSubsystem.cpp b/Engine/Plugins/Experimental/EditorDataStorageFeatures/Source/TedsSettings/Private/TedsSettingsEditorSubsystem.cpp
--- a/Engine/Plugins/Experimental/EditorDataStorageFeatures/Source/TedsSettings/Private/TedsSettingsEditorSubsystem.cpp
+++ b/Engine/Plugins/Experimental/EditorDataStorageFeatures/Source/TedsSettings/Private/TedsSettingsEditorSubsystem.cpp
@@ -31,6 +31,11 @@ UTedsSettingsEditorSubsystem::FOnEnabledChanged& UTedsSettingsEditorSubsystem::O
 	return EnabledChangedDelegate;
 }
 
+UE::Editor::DataStorage::RowHandle UTedsSettingsEditorSubsystem::FindSettingsSection(const FName& ContainerName, const FName& CategoryName, const FName& SectionName)
+{
+	return SettingsManager->FindSettingsSection(ContainerName, CategoryName, SectionName);
+}
+
 UE::Editor::DataStorage::RowHandle UTedsSettingsEditorSubsystem::FindOrAddSettingsSection(const FName& ContainerName, const FName& CategoryName, const FName& SectionName)
 {
 	return SettingsManager->FindOrAddSettingsSection(ContainerName, CategoryName, SectionName);
diff --git a/Engine/Plugins/Experimental/EditorDataStorageFeatures/Source/TedsSettings/Private/TedsSettingsManager.cpp b/Engine/Plugins/Experimental/EditorDataStorageFeatures/Source/TedsSettings/Private/TedsSettingsManager.cpp
--- a/Engine/Plugins/Experimental/EditorDataStorageFeatures/Source/TedsSettings/Private/TedsSettingsManager.cpp
+++ b/Engine/Plugins/Experimental/EditorDataStorageFeatures/Source/TedsSettings/Private/TedsSettingsManager.cpp
@@ -110,6 +110,28 @@ void FTedsSettingsManager::Shutdown()
 	}
 }
 
+UE::Editor::DataStorage::RowHandle FTedsSettingsManager::FindSettingsSection(const FName& ContainerName, const FName& CategoryName, const FName& SectionName)
+{
+	using namespace UE::Editor::DataStorage;
+	using namespace UE::Editor::Settings::Private;
+
+	if (bIsInitialized)
+	{
+		if (ICoreProvider* DataStorage = GetMutableDataStorageFeature<ICoreProvider>(StorageFeatureName))
+		{
+			// Ensure that the SettingsInactiveSectionTable is registered since it is possible for OnDataStorageInterfacesSet to be called
+			// for a delegate bound in another module before it has been called for our bound delegate.
+			RegisterTables(*DataStorage);
+
+			FMapKey SectionIndexKey = GenerateIndexKey(ContainerName, CategoryName, SectionName);
+
+			return DataStorage->LookupMappedRow(SectionIndexKey);
+		}
+	}
+
+	return InvalidRowHandle;
+}
+
 UE::Editor::DataStorage::RowHandle FTedsSettingsManager::FindOrAddSettingsSection(const FName& ContainerName, const FName& CategoryName, const FName& SectionName)
 {
 	using namespace UE::Editor::DataStorage;
diff --git a/Engine/Plugins/Experimental/EditorDataStorageFeatures/Source/TedsSettings/Private/TedsSettingsManager.h b/Engine/Plugins/Experimental/EditorDataStorageFeatures/Source/TedsSettings/Private/TedsSettingsManager.h
--- a/Engine/Plugins/Experimental/EditorDataStorageFeatures/Source/TedsSettings/Private/TedsSettingsManager.h
+++ b/Engine/Plugins/Experimental/EditorDataStorageFeatures/Source/TedsSettings/Private/TedsSettingsManager.h
@@ -21,6 +21,8 @@ public:
 	void Initialize();
 	void Shutdown();
 
+	UE::Editor::DataStorage::RowHandle FindSettingsSection(const FName& ContainerName, const FName& CategoryName, const FName& SectionName);
+
 	UE::Editor::DataStorage::RowHandle FindOrAddSettingsSection(const FName& ContainerName, const FName& CategoryName, const FName& SectionName);
 
 	bool GetSettingsSectionFromRow(UE::Editor::DataStorage::RowHandle Row, FName& OutContainerName, FName& OutCategoryName, FName& OutSectionName);
diff --git a/Engine/Plugins/Experimental/EditorDataStorageFeatures/Source/TedsSettings/Public/TedsSettingsEditorSubsystem.h b/Engine/Plugins/Experimental/EditorDataStorageFeatures/Source/TedsSettings/Public/TedsSettingsEditorSubsystem.h
--- a/Engine/Plugins/Experimental/EditorDataStorageFeatures/Source/TedsSettings/Public/TedsSettingsEditorSubsystem.h
+++ b/Engine/Plugins/Experimental/EditorDataStorageFeatures/Source/TedsSettings/Public/TedsSettingsEditorSubsystem.h
@@ -23,7 +23,26 @@ public:
 	DECLARE_MULTICAST_DELEGATE(FOnEnabledChanged)
 	FOnEnabledChanged& OnEnabledChanged();
 
-	/** Finds an existing row (may be active or inactive) or adds a new inactive settings section row if no existing row is found. */
+	/** 
+	 * Finds an existing row (may be active or inactive).
+	 * 
+	 * @param	ContainerName	The ContainerName to search for.
+	 * @param	CategoryName	The CategoryName to search for.
+	 * @param	SectionName		The SectionName to search for.
+	 * 
+	 * @return	The found row or InvalidRowHandle if no existing row is found.
+	 */
+	UE::Editor::DataStorage::RowHandle FindSettingsSection(const FName& ContainerName, const FName& CategoryName, const FName& SectionName);
+
+	/**
+	 * Finds an existing row (may be active or inactive) or adds a new inactive settings section row if no existing row is found.
+	 *
+	 * @param	ContainerName	The ContainerName to search for.
+	 * @param	CategoryName	The CategoryName to search for.
+	 * @param	SectionName		The SectionName to search for.
+	 *
+	 * @return	The row that was either found or added, or InvalidRowHandle if data storage is not available.
+	 */
 	UE::Editor::DataStorage::RowHandle FindOrAddSettingsSection(const FName& ContainerName, const FName& CategoryName, const FName& SectionName);
 
 	/**
