Teds Actor Folder Factory: fix OnActorFolderRemoved() hooked on OnActorFolderAdded()
Add a cvar for populating folders in TEDS to fix a crash on deleting folders by disabling the integration for now

diff --git a/Engine/Plugins/Experimental/EditorDataStorageFeatures/Source/TedsOutliner/Private/ActorFolderFactory.cpp b/Engine/Plugins/Experimental/EditorDataStorageFeatures/Source/TedsOutliner/Private/ActorFolderFactory.cpp
--- a/Engine/Plugins/Experimental/EditorDataStorageFeatures/Source/TedsOutliner/Private/ActorFolderFactory.cpp
+++ b/Engine/Plugins/Experimental/EditorDataStorageFeatures/Source/TedsOutliner/Private/ActorFolderFactory.cpp
@@ -20,6 +20,13 @@
 namespace UE::TedsActorFolderFactory::Local
 {
 	static const FName TableName("Editor_ActorFolderTable");
+
+	static bool bRegisterFoldersinTEDS = false;
+	
+	static FAutoConsoleVariableRef CVarUseTEDSOutliner(
+		TEXT("TEDS.Feature.ActorFolders"),
+		bRegisterFoldersinTEDS,
+		TEXT("Populate FFolders and Actor Folders in TEDS. Must be set at startup."));
 }
 
 void UTedsActorFolderFactory::RegisterTables(UE::Editor::DataStorage::ICoreProvider& InDataStorage)
@@ -33,20 +40,23 @@ void UTedsActorFolderFactory::RegisterTables(UE::Editor::DataStorage::ICoreProvi
 void UTedsActorFolderFactory::PreRegister(UE::Editor::DataStorage::ICoreProvider& InDataStorage)
 {
 	Super::PreRegister(InDataStorage);
-	
-	FActorFolders::Get().OnFolderCreated.AddUObject(this, &UTedsActorFolderFactory::OnFolderCreated);
-	FActorFolders::Get().OnFolderDeleted.AddUObject(this, &UTedsActorFolderFactory::OnFolderDeleted);
-	FActorFolders::Get().OnFolderMoved.AddUObject(this, &UTedsActorFolderFactory::OnFolderMoved);
 
-	FEditorDelegates::MapChange.AddUObject(this, &UTedsActorFolderFactory::OnMapChange);
-	FEditorDelegates::PostPIEStarted.AddUObject(this, &UTedsActorFolderFactory::OnPIEStarted);
+	if (UE::TedsActorFolderFactory::Local::bRegisterFoldersinTEDS)
+	{
+		FActorFolders::Get().OnFolderCreated.AddUObject(this, &UTedsActorFolderFactory::OnFolderCreated);
+		FActorFolders::Get().OnFolderDeleted.AddUObject(this, &UTedsActorFolderFactory::OnFolderDeleted);
+		FActorFolders::Get().OnFolderMoved.AddUObject(this, &UTedsActorFolderFactory::OnFolderMoved);
+
+		FEditorDelegates::MapChange.AddUObject(this, &UTedsActorFolderFactory::OnMapChange);
+		FEditorDelegates::PostPIEStarted.AddUObject(this, &UTedsActorFolderFactory::OnPIEStarted);
 	
-	FWorldDelegates::LevelAddedToWorld.AddUObject(this, &UTedsActorFolderFactory::OnLevelAdded);
+		FWorldDelegates::LevelAddedToWorld.AddUObject(this, &UTedsActorFolderFactory::OnLevelAdded);
 
-	GEngine->OnActorFolderAdded().AddUObject(this, &UTedsActorFolderFactory::OnActorFolderAdded);
-	GEngine->OnActorFolderRemoved().AddUObject(this, &UTedsActorFolderFactory::OnActorFolderAdded);
+		GEngine->OnActorFolderAdded().AddUObject(this, &UTedsActorFolderFactory::OnActorFolderAdded);
+		GEngine->OnActorFolderRemoved().AddUObject(this, &UTedsActorFolderFactory::OnActorFolderRemoved);
 	
-	InDataStorage.OnUpdateCompleted().AddUObject(this, &UTedsActorFolderFactory::Tick);
+		InDataStorage.OnUpdateCompleted().AddUObject(this, &UTedsActorFolderFactory::Tick);
+	}
 	
 	// Store a pointer to data storage so we don't have to grab it from the global instance of the registry every time
 	DataStorage = &InDataStorage;
@@ -57,6 +67,11 @@ void UTedsActorFolderFactory::RegisterQueries(UE::Editor::DataStorage::ICoreProv
 	using namespace UE::Editor::DataStorage::Queries;
 	using namespace UE::Editor;
 
+	if (!UE::TedsActorFolderFactory::Local::bRegisterFoldersinTEDS)
+	{
+		return;
+	}
+
 	InDataStorage.RegisterQuery(
 		Select(
 			TEXT("Delete folder with deleted world"),
@@ -128,19 +143,23 @@ void UTedsActorFolderFactory::RegisterQueries(UE::Editor::DataStorage::ICoreProv
 
 void UTedsActorFolderFactory::PreShutdown(UE::Editor::DataStorage::ICoreProvider& InDataStorage)
 {
-	FActorFolders::Get().OnFolderCreated.RemoveAll(this);
-	FActorFolders::Get().OnFolderDeleted.RemoveAll(this);
-	FActorFolders::Get().OnFolderMoved.RemoveAll(this);
+	if (UE::TedsActorFolderFactory::Local::bRegisterFoldersinTEDS)
+	{
+		FActorFolders::Get().OnFolderCreated.RemoveAll(this);
+		FActorFolders::Get().OnFolderDeleted.RemoveAll(this);
+		FActorFolders::Get().OnFolderMoved.RemoveAll(this);
 
-	FEditorDelegates::MapChange.RemoveAll(this);
-	FEditorDelegates::PostPIEStarted.RemoveAll(this);
+		FEditorDelegates::MapChange.RemoveAll(this);
+		FEditorDelegates::PostPIEStarted.RemoveAll(this);
 	
-	FWorldDelegates::LevelAddedToWorld.RemoveAll(this);
+		FWorldDelegates::LevelAddedToWorld.RemoveAll(this);
 
-	GEngine->OnActorFolderAdded().RemoveAll(this);
-	GEngine->OnActorFolderRemoved().RemoveAll(this);
+		GEngine->OnActorFolderAdded().RemoveAll(this);
+		GEngine->OnActorFolderRemoved().RemoveAll(this);
+	
+		InDataStorage.OnUpdateCompleted().RemoveAll(this);
+	}
 	
-	InDataStorage.OnUpdateCompleted().RemoveAll(this);
 
 	Super::PreShutdown(InDataStorage);
 }
