CB 2.0 - Add an FFavoriteTag to TEDS and sync CB favorite folders to TEDS

diff --git a/Engine/Plugins/Experimental/EditorDataStorageFeatures/Source/TedsAssetData/Private/Processors/AssetProcessors.cpp b/Engine/Plugins/Experimental/EditorDataStorageFeatures/Source/TedsAssetData/Private/Processors/AssetProcessors.cpp
--- a/Engine/Plugins/Experimental/EditorDataStorageFeatures/Source/TedsAssetData/Private/Processors/AssetProcessors.cpp
+++ b/Engine/Plugins/Experimental/EditorDataStorageFeatures/Source/TedsAssetData/Private/Processors/AssetProcessors.cpp
@@ -8,6 +8,7 @@
 #include "ContentBrowserDataSubsystem.h"
 #include "ContentBrowserModule.h"
 #include "IContentBrowserDataModule.h"
+#include "IContentBrowserSingleton.h"
 #include "Elements/Columns/TypedElementFolderColumns.h"
 #include "Elements/Columns/TypedElementSlateWidgetColumns.h"
 #include "Elements/Framework/TypedElementQueryBuilder.h"
@@ -17,6 +18,7 @@
 #include "Elements/Columns/TypedElementMiscColumns.h"
 #include "String/LexFromString.h"
 #include "TedsAlerts.h"
+#include "TedsAssetDataModule.h"
 
 #define LOCTEXT_NAMESPACE "UTedsAssetDataFactory"
 
@@ -204,11 +206,19 @@ void UTedsAssetDataFactory::PreRegister(UE::Editor::DataStorage::ICoreProvider&
 	{
 		return;
 	}
+	
+	if (UE::Editor::AssetData::FTedsAssetDataModule* TedsAssetDataModule = FModuleManager::Get().GetModulePtr<UE::Editor::AssetData::FTedsAssetDataModule>("TedsAssetData"))
+	{
+		TedsAssetDataModule->OnAssetRegistryStorageInit().AddUObject(this, &UTedsAssetDataFactory::OnAssetDataStorageEnabled, &DataStorage);
+	}
 
 	if (FContentBrowserModule* ContentBrowserModule = FModuleManager::Get().GetModulePtr<FContentBrowserModule>("ContentBrowser"))
 	{
 		ContentBrowserModule->GetOnSetFolderColor().AddUObject(this, &UTedsAssetDataFactory::OnSetFolderColor, &DataStorage);
 	}
+	
+	OnFavoritesChangedDelegateHandle = IContentBrowserSingleton::Get().RegisterOnFavoritesChangedHandler(
+		FOnFavoritesChanged::FDelegate::CreateUObject(this, &UTedsAssetDataFactory::OnFavoritesChanged, &DataStorage));
 }
 
 void UTedsAssetDataFactory::PreShutdown(UE::Editor::DataStorage::ICoreProvider& DataStorage)
@@ -216,6 +226,12 @@ void UTedsAssetDataFactory::PreShutdown(UE::Editor::DataStorage::ICoreProvider&
 	if (FContentBrowserModule* ContentBrowserModule = FModuleManager::Get().GetModulePtr<FContentBrowserModule>("ContentBrowser"))
 	{
 		ContentBrowserModule->GetOnSetFolderColor().RemoveAll(this);
+		IContentBrowserSingleton::Get().UnregisterOnFavoritesChangedDelegate(OnFavoritesChangedDelegateHandle);
+	}
+	
+	if (UE::Editor::AssetData::FTedsAssetDataModule* TedsAssetDataModule = FModuleManager::Get().GetModulePtr<UE::Editor::AssetData::FTedsAssetDataModule>("TedsAssetData"))
+	{
+		TedsAssetDataModule->OnAssetRegistryStorageInit().RemoveAll(this);
 	}
 }
 
@@ -230,4 +246,27 @@ void UTedsAssetDataFactory::OnSetFolderColor(const FString& Path, UE::Editor::Da
 	}
 }
 
+void UTedsAssetDataFactory::OnFavoritesChanged(const FContentBrowserItemPath& ItemPath, bool bAdded, UE::Editor::DataStorage::ICoreProvider* DataStorage)
+{
+	if (ItemPath.HasInternalPath())
+	{
+		const UE::Editor::DataStorage::FMapKey PathKey = UE::Editor::DataStorage::FMapKey(ItemPath.GetInternalPathName());
+		const UE::Editor::DataStorage::RowHandle Row = DataStorage->LookupMappedRow(PathKey);
+
+		if (DataStorage->IsRowAvailable(Row))
+		{
+			bAdded ? DataStorage->AddColumn<FFavoriteTag>(Row) : DataStorage->RemoveColumn<FFavoriteTag>(Row);
+		}
+	}
+	
+}
+
+void UTedsAssetDataFactory::OnAssetDataStorageEnabled(UE::Editor::DataStorage::ICoreProvider* DataStorage)
+{
+	IContentBrowserSingleton::Get().ForEachFavoriteFolder([DataStorage, this](const FContentBrowserItemPath& Path)
+	{
+		OnFavoritesChanged(Path, true, DataStorage);
+	});
+}
+
 #undef LOCTEXT_NAMESPACE
diff --git a/Engine/Plugins/Experimental/EditorDataStorageFeatures/Source/TedsAssetData/Private/Processors/AssetProcessors.h b/Engine/Plugins/Experimental/EditorDataStorageFeatures/Source/TedsAssetData/Private/Processors/AssetProcessors.h
--- a/Engine/Plugins/Experimental/EditorDataStorageFeatures/Source/TedsAssetData/Private/Processors/AssetProcessors.h
+++ b/Engine/Plugins/Experimental/EditorDataStorageFeatures/Source/TedsAssetData/Private/Processors/AssetProcessors.h
@@ -6,6 +6,8 @@
 
 #include "AssetProcessors.generated.h"
 
+struct FContentBrowserItemPath;
+
 namespace UE::Editor::DataStorage
 {
 	class ICoreProvider;
@@ -26,4 +28,9 @@ public:
 protected:
 
 	void OnSetFolderColor(const FString& Path, UE::Editor::DataStorage::ICoreProvider* DataStorage);
+	void OnFavoritesChanged(const FContentBrowserItemPath& ItemPath, bool bAdded, UE::Editor::DataStorage::ICoreProvider* DataStorage);
+	void OnAssetDataStorageEnabled(UE::Editor::DataStorage::ICoreProvider* DataStorage);
+
+protected:
+	FDelegateHandle OnFavoritesChangedDelegateHandle;
 };
diff --git a/Engine/Plugins/Experimental/EditorDataStorageFeatures/Source/TedsAssetData/Private/TedsAssetDataModule.cpp b/Engine/Plugins/Experimental/EditorDataStorageFeatures/Source/TedsAssetData/Private/TedsAssetDataModule.cpp
--- a/Engine/Plugins/Experimental/EditorDataStorageFeatures/Source/TedsAssetData/Private/TedsAssetDataModule.cpp
+++ b/Engine/Plugins/Experimental/EditorDataStorageFeatures/Source/TedsAssetData/Private/TedsAssetDataModule.cpp
@@ -111,6 +111,8 @@ void FTedsAssetDataModule::InitAssetRegistryStorage()
 
 	AssetDataCBDataSource = MakeUnique<Private::FTedsAssetDataCBDataSource>(MutableDataStorage);
 	AssetRegistryStorage = MakeUnique<Private::FTedsAssetData>(MutableDataStorage);
+	
+	OnAssetRegistryStorageInitDelegate.Broadcast();
 }
 
 void FTedsAssetDataModule::EnableAssetDataMetadataStorage()
@@ -129,4 +131,8 @@ void FTedsAssetDataModule::DisableAssetDataMetadataStorage()
 	}
 }
 
+FSimpleMulticastDelegate& FTedsAssetDataModule::OnAssetRegistryStorageInit()
+{
+	return OnAssetRegistryStorageInitDelegate;
+}
 }// namespace UE::Editor::AssetData
diff --git a/Engine/Plugins/Experimental/EditorDataStorageFeatures/Source/TedsAssetData/Private/TedsAssetDataModule.h b/Engine/Plugins/Experimental/EditorDataStorageFeatures/Source/TedsAssetData/Private/TedsAssetDataModule.h
--- a/Engine/Plugins/Experimental/EditorDataStorageFeatures/Source/TedsAssetData/Private/TedsAssetDataModule.h
+++ b/Engine/Plugins/Experimental/EditorDataStorageFeatures/Source/TedsAssetData/Private/TedsAssetDataModule.h
@@ -2,6 +2,7 @@
 
 #pragma once
 
+#include "Delegates/Delegate.h"
 #include "Modules/ModuleInterface.h"
 #include "Templates/UniquePtr.h"
 
@@ -27,6 +28,7 @@ public:
 	bool IsTedsAssetRegistryStorageEnabled() const;
 	void EnableAssetDataMetadataStorage();
 	void DisableAssetDataMetadataStorage();
+	FSimpleMulticastDelegate& OnAssetRegistryStorageInit();
 
 	/**
 	 * Process now any pending event that might make the Teds database out of sync with the asset registry.
@@ -39,6 +41,7 @@ private:
 
 	TUniquePtr<Private::FTedsAssetDataCBDataSource> AssetDataCBDataSource;
 	TUniquePtr<Private::FTedsAssetData> AssetRegistryStorage;
+	FSimpleMulticastDelegate OnAssetRegistryStorageInitDelegate;
 };
 
 } // namespace UE::Editor::AssetData
diff --git a/Engine/Source/Editor/ContentBrowser/Private/ContentBrowserSingleton.cpp b/Engine/Source/Editor/ContentBrowser/Private/ContentBrowserSingleton.cpp
--- a/Engine/Source/Editor/ContentBrowser/Private/ContentBrowserSingleton.cpp
+++ b/Engine/Source/Editor/ContentBrowser/Private/ContentBrowserSingleton.cpp
@@ -1219,19 +1219,27 @@ void FContentBrowserSingleton::UnregisterIsFolderShowPrivateContentToggleableDel
 	IsFolderShowPrivateContentToggleableDelegate = FIsFolderShowPrivateContentToggleableDelegate();
 }
 
-FDelegateHandle FContentBrowserSingleton::RegisterOnFavoritesChangedHandler(FSimpleDelegate InOnFavoritesChanged)
+FDelegateHandle FContentBrowserSingleton::RegisterOnFavoritesChangedHandler(FOnFavoritesChanged::FDelegate InOnFavoritesChanged)
 {
 	return OnFavoritesChanged.Add(InOnFavoritesChanged);
 }
 
+FDelegateHandle FContentBrowserSingleton::RegisterOnFavoritesChangedHandler(FSimpleDelegate InOnFavoritesChanged)
+{
+	return OnFavoritesChanged.AddLambda([InOnFavoritesChanged](const FContentBrowserItemPath&, bool)
+	{
+		InOnFavoritesChanged.ExecuteIfBound();
+	});
+}
+
 void FContentBrowserSingleton::UnregisterOnFavoritesChangedDelegate(FDelegateHandle Handle)
 {
 	OnFavoritesChanged.Remove(Handle);
 }
 
-void FContentBrowserSingleton::BroadcastFavoritesChanged() const
+void FContentBrowserSingleton::BroadcastFavoritesChanged(const FContentBrowserItemPath& ChangedItem, bool bAdded) const
 {
-	OnFavoritesChanged.Broadcast();
+	OnFavoritesChanged.Broadcast(ChangedItem, bAdded);
 }
 
 void FContentBrowserSingleton::PopulateConfigValues()
@@ -1467,6 +1475,14 @@ void FContentBrowserSingleton::ToggleFolderFavorite(const TArray<FString>& Folde
 	}
 }
 
+void FContentBrowserSingleton::ForEachFavoriteFolder(TFunctionRef<void(const FContentBrowserItemPath&)> InDelegate)
+{
+	for (const FString& Path : FavoriteFolderPaths)
+	{
+		InDelegate(FContentBrowserItemPath(Path, EContentBrowserPathType::Internal));
+	}
+}
+
 void FContentBrowserSingleton::SetFilterLayout(const EFilterBarLayout InFilterBarLayout)
 {
 	if (!HasPrimaryContentBrowser())
diff --git a/Engine/Source/Editor/ContentBrowser/Private/ContentBrowserSingleton.h b/Engine/Source/Editor/ContentBrowser/Private/ContentBrowserSingleton.h
--- a/Engine/Source/Editor/ContentBrowser/Private/ContentBrowserSingleton.h
+++ b/Engine/Source/Editor/ContentBrowser/Private/ContentBrowserSingleton.h
@@ -116,6 +116,8 @@ PRAGMA_ENABLE_DEPRECATION_WARNINGS
 	virtual void SetPrivateContentPermissionListDirty() override;
 	virtual void RegisterIsFolderShowPrivateContentToggleableDelegate(FIsFolderShowPrivateContentToggleableDelegate InIsFolderShowPrivateContentToggleableDelegate) override;
 	virtual void UnregisterIsFolderShowPrivateContentToggleableDelegate() override;
+	virtual FDelegateHandle RegisterOnFavoritesChangedHandler(FOnFavoritesChanged::FDelegate InOnFavoritesChanged) override;
+	UE_DEPRECATED(5.7, "Use the overload that takes in an FOnFavoritesChanged delegate")
 	virtual FDelegateHandle RegisterOnFavoritesChangedHandler(FSimpleDelegate OnFavoritesChanged) override;
 	virtual void UnregisterOnFavoritesChangedDelegate(FDelegateHandle Handle) override;
 	virtual TArray<FString> GetAliasesForPath(const FSoftObjectPath& InPath) const override;
@@ -125,6 +127,7 @@ PRAGMA_ENABLE_DEPRECATION_WARNINGS
 	virtual bool IsAssetViewDoneFiltering() override;
 	virtual void SetSearchText(const FText& NewText) override;
 	virtual void ToggleFolderFavorite(const TArray<FString>& FolderPaths) override;
+	virtual void ForEachFavoriteFolder(TFunctionRef<void(const FContentBrowserItemPath&)> InDelegate) override;
 	virtual void SetFilterLayout(EFilterBarLayout InFilterBarLayout) override;
 	virtual TSharedPtr<SWidget> GetActiveFilterContainer() override;
 	virtual EFilterBarLayout GetFilterLayout() override;
@@ -135,7 +138,7 @@ PRAGMA_ENABLE_DEPRECATION_WARNINGS
 	virtual FOnContentSourceFactoriesChanged& OnContentSourceFactoriesChanged() override;
 
 	/** Broadcast that the favorites have changed. */
-	void BroadcastFavoritesChanged() const;
+	void BroadcastFavoritesChanged(const FContentBrowserItemPath& ChangedItem, bool bAdded) const;
 
 	/** Gets the content browser singleton as a FContentBrowserSingleton */
 	static FContentBrowserSingleton& Get();
@@ -227,8 +230,7 @@ private:
 	FShowPrivateContentState ShowPrivateContentState;
 	FShowPrivateContentState ShowEpicInternalContentState;
 
-	DECLARE_MULTICAST_DELEGATE(FOnFavoritesChangedDelegate);
-	FOnFavoritesChangedDelegate OnFavoritesChanged;
+	FOnFavoritesChanged OnFavoritesChanged;
 
 	TMap<FName, FContentSourceFactory> ContentSourceFactories;
 	FOnContentSourceFactoriesChanged ContentSourceFactoriesChangedEvent;
diff --git a/Engine/Source/Editor/ContentBrowser/Private/ContentBrowserUtils.cpp b/Engine/Source/Editor/ContentBrowser/Private/ContentBrowserUtils.cpp
--- a/Engine/Source/Editor/ContentBrowser/Private/ContentBrowserUtils.cpp
+++ b/Engine/Source/Editor/ContentBrowser/Private/ContentBrowserUtils.cpp
@@ -1019,7 +1019,7 @@ void ContentBrowserUtils::AddFavoriteFolder(const FContentBrowserItemPath& Folde
 		UContentBrowserConfig::Get()->SaveEditorConfig();
 	}
 
-	FContentBrowserSingleton::Get().BroadcastFavoritesChanged();
+	FContentBrowserSingleton::Get().BroadcastFavoritesChanged(FolderPath, true);
 }
 
 void ContentBrowserUtils::RemoveFavoriteFolder(const FContentBrowserItemPath& FolderPath)
@@ -1041,7 +1041,7 @@ void ContentBrowserUtils::RemoveFavoriteFolder(const FContentBrowserItemPath& Fo
 		UContentBrowserConfig::Get()->SaveEditorConfig();
 	}
 
-	FContentBrowserSingleton::Get().BroadcastFavoritesChanged();
+	FContentBrowserSingleton::Get().BroadcastFavoritesChanged(FolderPath, false);
 }
 
 void ContentBrowserUtils::RemoveFavoriteFolder(const FString& FolderPath, bool bFlushConfig)
diff --git a/Engine/Source/Editor/ContentBrowser/Private/SPathView.cpp b/Engine/Source/Editor/ContentBrowser/Private/SPathView.cpp
--- a/Engine/Source/Editor/ContentBrowser/Private/SPathView.cpp
+++ b/Engine/Source/Editor/ContentBrowser/Private/SPathView.cpp
@@ -2735,7 +2735,8 @@ void SFavoritePathView::Construct(const FArguments& InArgs)
 	// Bind the favorites menu to update after folder changes
 	AssetViewUtils::OnFolderPathChanged().AddSP(this, &SFavoritePathView::FixupFavoritesFromExternalChange); 
 
-	OnFavoritesChangedHandle = FContentBrowserSingleton::Get().RegisterOnFavoritesChangedHandler(FSimpleDelegate::CreateSP(this, &SFavoritePathView::OnFavoriteAdded));
+	OnFavoritesChangedHandle = FContentBrowserSingleton::Get().RegisterOnFavoritesChangedHandler(
+		FOnFavoritesChanged::FDelegate::CreateSP(this, &SFavoritePathView::OnFavoriteAdded));
 
 	SPathView::Construct(InArgs);
 }
@@ -2916,7 +2917,7 @@ TSharedRef<ITableRow> SFavoritePathView::GenerateTreeRow(TSharedPtr<FTreeItem> T
 		];
 }
 
-void SFavoritePathView::OnFavoriteAdded()
+void SFavoritePathView::OnFavoriteAdded(const FContentBrowserItemPath&, bool)
 {
 	if (!bIsLoadingSettings)
 	{
diff --git a/Engine/Source/Editor/ContentBrowser/Private/SPathView.h b/Engine/Source/Editor/ContentBrowser/Private/SPathView.h
--- a/Engine/Source/Editor/ContentBrowser/Private/SPathView.h
+++ b/Engine/Source/Editor/ContentBrowser/Private/SPathView.h
@@ -611,7 +611,7 @@ public:
 private:
 	virtual TSharedRef<ITableRow> GenerateTreeRow(TSharedPtr<FTreeItem> TreeItem, const TSharedRef<STableViewBase>& OwnerTable) override;
 
-	void OnFavoriteAdded();
+	void OnFavoriteAdded(const FContentBrowserItemPath& ChangedItem, bool bAdded);
 
 	/** Handles updating the view when content items are changed */
 	virtual void HandleItemDataUpdated(TArrayView<const FContentBrowserItemDataUpdate> InUpdatedItems) override;
diff --git a/Engine/Source/Editor/ContentBrowser/Public/ContentBrowserDelegates.h b/Engine/Source/Editor/ContentBrowser/Public/ContentBrowserDelegates.h
--- a/Engine/Source/Editor/ContentBrowser/Public/ContentBrowserDelegates.h
+++ b/Engine/Source/Editor/ContentBrowser/Public/ContentBrowserDelegates.h
@@ -137,6 +137,9 @@ DECLARE_DELEGATE_RetVal_OneParam(bool, FIsFolderShowPrivateContentToggleableDele
 /** Called when a custom folder color is set on a folder */
 DECLARE_MULTICAST_DELEGATE_OneParam(FOnSetFolderColor, const FString& /*FolderPath*/)
 
+/** Called when an item is added or removed from favorites */
+DECLARE_MULTICAST_DELEGATE_TwoParams(FOnFavoritesChanged, const FContentBrowserItemPath& /*ItemPath*/, bool /* bAdded */)
+
 /** Contains the delegates used to handle a custom drag-and-drop in the asset view */
 struct FAssetViewDragAndDropExtender
 {
diff --git a/Engine/Source/Editor/ContentBrowser/Public/IContentBrowserSingleton.h b/Engine/Source/Editor/ContentBrowser/Public/IContentBrowserSingleton.h
--- a/Engine/Source/Editor/ContentBrowser/Public/IContentBrowserSingleton.h
+++ b/Engine/Source/Editor/ContentBrowser/Public/IContentBrowserSingleton.h
@@ -800,7 +800,7 @@ public:
 	virtual void UnregisterIsFolderShowPrivateContentToggleableDelegate() = 0;
 
 	/** Register a delegate to be called when the Favorites changes. */
-	virtual FDelegateHandle RegisterOnFavoritesChangedHandler(FSimpleDelegate OnFavoritesChanged) = 0;
+	virtual FDelegateHandle RegisterOnFavoritesChangedHandler(FOnFavoritesChanged::FDelegate InOnFavoritesChanged) = 0;
 
 	/** Unregister a previously-registered handler for when Favorites changes. */
 	virtual void UnregisterOnFavoritesChangedDelegate(FDelegateHandle Handle) = 0;
@@ -831,6 +831,9 @@ public:
 	/** Toggle the Favorite status of the given folder list */
 	virtual void ToggleFolderFavorite(const TArray<FString>& FolderPaths) = 0;
 
+	/** Iterate through each favorite folder and call the delegate */
+	virtual void ForEachFavoriteFolder(TFunctionRef<void(const FContentBrowserItemPath&)> InDelegate) = 0;
+
 	/** Sets the Content Browser filter layout to either vertical or horizontal */
 	virtual void SetFilterLayout(EFilterBarLayout InFilterBarLayout) = 0;
 
@@ -858,4 +861,8 @@ public:
 	/** Delegate fired when a new content source factory is registered/unregistered */
 	DECLARE_MULTICAST_DELEGATE(FOnContentSourceFactoriesChanged)
 	virtual FOnContentSourceFactoriesChanged& OnContentSourceFactoriesChanged() = 0;
+	
+	/** Register a delegate to be called when the Favorites changes. */
+	UE_DEPRECATED(5.7, "Use the overload that takes in an FOnFavoritesChanged delegate")
+	virtual FDelegateHandle RegisterOnFavoritesChangedHandler(FSimpleDelegate OnFavoritesChanged) = 0;
 };
diff --git a/Engine/Source/Runtime/TypedElementFramework/Public/Elements/Columns/TypedElementMiscColumns.h b/Engine/Source/Runtime/TypedElementFramework/Public/Elements/Columns/TypedElementMiscColumns.h
--- a/Engine/Source/Runtime/TypedElementFramework/Public/Elements/Columns/TypedElementMiscColumns.h
+++ b/Engine/Source/Runtime/TypedElementFramework/Public/Elements/Columns/TypedElementMiscColumns.h
@@ -58,8 +58,6 @@ struct FTypedElementRowReferenceColumn final : public FEditorDataStorageColumn
 
 /**
  * A name for this row.
- * 
- * This can be used as a dynamic column to specify names for multiple items in a row.
  */
 USTRUCT(meta = (DisplayName = "Name"))
 struct FNameColumn final : public FEditorDataStorageColumn
@@ -69,3 +67,12 @@ struct FNameColumn final : public FEditorDataStorageColumn
 	UPROPERTY(meta = (Searchable))
 	FName Name;
 };
+
+/**
+ * Tag used to mark a row as "favorite" in TEDS - which can then be used to filter in table views.
+ */
+USTRUCT(meta = (DisplayName = "Favorite"))
+struct FFavoriteTag final : public FEditorDataStorageTag
+{
+	GENERATED_BODY()
+};
\ No newline at end of file
