TEDS-Outliner:
Add a new query to re-filter items on label change
Defer the add query to avoid accessing data storage from an observer

CVD:
Bug fixes for search not behaving correctly with the CVD specific code for the TEDS Outliner

diff --git a/Engine/Plugins/ChaosVD/Source/ChaosVD/Private/ChaosVDWorldOutlinerMode.cpp b/Engine/Plugins/ChaosVD/Source/ChaosVD/Private/ChaosVDWorldOutlinerMode.cpp
--- a/Engine/Plugins/ChaosVD/Source/ChaosVD/Private/ChaosVDWorldOutlinerMode.cpp
+++ b/Engine/Plugins/ChaosVD/Source/ChaosVD/Private/ChaosVDWorldOutlinerMode.cpp
@@ -165,6 +165,9 @@ FChaosVDWorldOutlinerMode::FChaosVDWorldOutlinerMode(const UE::Editor::Outliner:
 				.Compile();
 
 		CustomRowRemovalQueryHandle = DataStorage->RegisterQuery(MoveTemp(RowRemovalQueryDescription));
+
+		// Use TEDS' update instead of a regular tick to make sure all processors have run and the data is correct
+		DataStorage->OnUpdateCompleted().AddRaw(this, &FChaosVDWorldOutlinerMode::OnDataStorageUpdateCompleted);
 	}
 
 }
@@ -181,6 +184,8 @@ FChaosVDWorldOutlinerMode::~FChaosVDWorldOutlinerMode()
 	{
 		DataStorage->UnregisterQuery(CustomRowAdditionQueryHandle);
 		DataStorage->UnregisterQuery(CustomRowRemovalQueryHandle);
+		
+		DataStorage->OnUpdateCompleted().RemoveAll(this);
 	}
 }
 
@@ -329,11 +334,12 @@ void FChaosVDWorldOutlinerMode::ProcessPendingActionsList(double TimeBudgetInSec
 
 		if (Type == EPendingOperationType::Add)
 		{
-			EventData.ItemIDs.Add(*PendingItemsIDIterator);
+			EventData.Items.Add(CreateItemFor<UE::Editor::Outliner::FTedsOutlinerTreeItem>(UE::Editor::Outliner::FTedsOutlinerTreeItem(*PendingItemsIDIterator, TedsOutlinerImpl.ToSharedRef())));
+			
 		}
 		else
 		{
-			EventData.Items.Add(CreateItemFor<UE::Editor::Outliner::FTedsOutlinerTreeItem>(UE::Editor::Outliner::FTedsOutlinerTreeItem(*PendingItemsIDIterator, TedsOutlinerImpl.ToSharedRef())));
+			EventData.ItemIDs.Add(*PendingItemsIDIterator);
 		}
 
 		PendingItemsIDIterator.RemoveCurrent();
@@ -343,7 +349,7 @@ void FChaosVDWorldOutlinerMode::ProcessPendingActionsList(double TimeBudgetInSec
 	GetHierarchy()->OnHierarchyChanged().Broadcast(MoveTemp(EventData));
 }
 
-bool FChaosVDWorldOutlinerMode::Tick(float DeltaTime)
+void FChaosVDWorldOutlinerMode::OnDataStorageUpdateCompleted()
 {
 	if (!CanPopulate())
 	{
@@ -355,7 +361,6 @@ bool FChaosVDWorldOutlinerMode::Tick(float DeltaTime)
 		ProcessPendingActionsList(TimeBudgetInSeconds, RowsPendingAddition, EPendingOperationType::Add);
 		ProcessPendingActionsList(TimeBudgetInSeconds, RowsPendingRemoval, EPendingOperationType::Remove);
 	}
-	return true;
 }
 
 #undef LOCTEXT_NAMESPACE
diff --git a/Engine/Plugins/ChaosVD/Source/ChaosVD/Private/ChaosVDWorldOutlinerMode.h b/Engine/Plugins/ChaosVD/Source/ChaosVD/Private/ChaosVDWorldOutlinerMode.h
--- a/Engine/Plugins/ChaosVD/Source/ChaosVD/Private/ChaosVDWorldOutlinerMode.h
+++ b/Engine/Plugins/ChaosVD/Source/ChaosVD/Private/ChaosVDWorldOutlinerMode.h
@@ -38,7 +38,7 @@ struct FChaosVDActorOutlinerSelector
  * It has a more limited view compared to the normal outliner, hiding features we don't support,
  * and it is integrated with the CVD local selection system
  */
-class FChaosVDWorldOutlinerMode : public UE::Editor::Outliner::FTedsOutlinerMode, public FChaosVDSceneSelectionObserver, public FTSTickerObjectBase
+class FChaosVDWorldOutlinerMode : public UE::Editor::Outliner::FTedsOutlinerMode, public FChaosVDSceneSelectionObserver
 {
 public:
 
@@ -77,7 +77,7 @@ public:
 		return SNullWidget::NullWidget;
 	}
 
-	virtual bool Tick(float DeltaTime) override;
+	void OnDataStorageUpdateCompleted();
 
 private:
 
diff --git a/Engine/Plugins/ChaosVD/Source/ChaosVD/Private/ChaosVDWorldOutlinerTab.cpp b/Engine/Plugins/ChaosVD/Source/ChaosVD/Private/ChaosVDWorldOutlinerTab.cpp
--- a/Engine/Plugins/ChaosVD/Source/ChaosVD/Private/ChaosVDWorldOutlinerTab.cpp
+++ b/Engine/Plugins/ChaosVD/Source/ChaosVD/Private/ChaosVDWorldOutlinerTab.cpp
@@ -41,7 +41,7 @@ void FChaosVDWorldOutlinerTab::CreateWorldOutlinerWidget()
 			.ReadOnly<FTypedElementLabelColumn>()
 		.Where()
 			.Any<FChaosVDObjectDataTag, FTypedElementActorTag>()
-			.All<FTypedElementFromCVDWorldTag>()
+			.All<FTypedElementFromCVDWorldTag, FChaosVDActiveObjectTag>()
 		.Compile();
 
 	FSceneOutlinerInitializationOptions InitOptions;
@@ -55,6 +55,7 @@ void FChaosVDWorldOutlinerTab::CreateWorldOutlinerWidget()
 	Params.QueryDescription = OutlinerQueryDescription;
 	Params.bUseDefaultTedsFilters = false;
 	Params.bShowRowHandleColumn = false;
+	Params.bUseDefaultObservers = false; // CVD uses custom observers to time slice addition and removal currently
 
 	const FTedsOutlinerHierarchyData::FGetParentRowHandle RowHandleGetter = FTedsOutlinerHierarchyData::FGetParentRowHandle::CreateLambda([](const void* InColumnData)
 	{
diff --git a/Engine/Plugins/ChaosVD/Source/ChaosVD/Private/TEDS/ChaosVDParticleEditorDataFactory.cpp b/Engine/Plugins/ChaosVD/Source/ChaosVD/Private/TEDS/ChaosVDParticleEditorDataFactory.cpp
--- a/Engine/Plugins/ChaosVD/Source/ChaosVD/Private/TEDS/ChaosVDParticleEditorDataFactory.cpp
+++ b/Engine/Plugins/ChaosVD/Source/ChaosVD/Private/TEDS/ChaosVDParticleEditorDataFactory.cpp
@@ -74,6 +74,7 @@ void UChaosVDParticleEditorDataFactory::RegisterQueries(UE::Editor::DataStorage:
 						FString Name = SceneObject->GetDisplayName();
 						LabelHash.LabelHash = ObjectLabelHash;
 						Label.Label = MoveTemp(Name);
+						Context.AddColumns<FTypedElementSyncBackToWorldTag>(Row);
 					}
 				}
 			})
diff --git a/Engine/Plugins/Experimental/EditorDataStorageFeatures/Source/TedsOutliner/Private/TedsOutlinerImpl.cpp b/Engine/Plugins/Experimental/EditorDataStorageFeatures/Source/TedsOutliner/Private/TedsOutlinerImpl.cpp
--- a/Engine/Plugins/Experimental/EditorDataStorageFeatures/Source/TedsOutliner/Private/TedsOutlinerImpl.cpp
+++ b/Engine/Plugins/Experimental/EditorDataStorageFeatures/Source/TedsOutliner/Private/TedsOutlinerImpl.cpp
@@ -87,6 +87,7 @@ FTedsOutlinerParams::FTedsOutlinerParams(SSceneOutliner* InSceneOutliner)
 	, bUseDefaultTedsFilters(false)
 	, bShowRowHandleColumn(true)
 	, bForceShowParents(true)
+	, bUseDefaultObservers(true)
 	, HierarchyData(FTedsOutlinerHierarchyData::GetDefaultHierarchyData())
 {
 	CellWidgetPurpose =
@@ -185,21 +186,17 @@ void FTedsOutlinerImpl::CreateFilterQueries()
 void FTedsOutlinerImpl::Init()
 {
 	CreateFilterQueries();
-	
-	if (SelectionSetName.IsSet())
-	{
-		// Ticker for selection updates so we don't fire the delegate multiple times in one frame for multi select
-		TickerHandle = FTSTicker::GetCoreTicker().AddTicker(FTickerDelegate::CreateLambda(
-				[this](float DeltaTimeInSeconds)
-				{
-					Tick();
-					return true;
-				}));
-	}
+
+	// Tick post TEDS update to make sure all processors have run and the data is correct
+	Storage->OnUpdateCompleted().AddRaw(this, &FTedsOutlinerImpl::Tick);
 }
 
 FTedsOutlinerImpl::~FTedsOutlinerImpl()
 {
+	if (Storage)
+	{
+		Storage->OnUpdateCompleted().RemoveAll(this);
+	}
 	UnregisterQueries();
 	FTSTicker::GetCoreTicker().RemoveTicker(TickerHandle);
 }
@@ -562,10 +559,7 @@ void FTedsOutlinerImpl::OnItemAdded(DataStorage::RowHandle ItemRowHandle)
 		return;
 	}
 	
-	FSceneOutlinerHierarchyChangedData EventData;
-	EventData.Type = FSceneOutlinerHierarchyChangedData::Added;
-	EventData.Items.Add(SceneOutlinerMode->CreateItemFor<FTedsOutlinerTreeItem>(FTedsOutlinerTreeItem(ItemRowHandle, AsShared())));
-	HierarchyChangedEvent.Broadcast(EventData);
+	RowsPendingAddition.Add(ItemRowHandle);
 }
 
 void FTedsOutlinerImpl::OnItemRemoved(DataStorage::RowHandle ItemRowHandle)
@@ -596,34 +590,40 @@ void FTedsOutlinerImpl::RecompileQueries()
 	// Add the filters the user has active to the query
 	AppendExternalQueries(FinalQueryDescription);
 
-	// Row to track addition of rows to the Outliner
-	FQueryDescription RowAdditionQueryDescription =
-		Select(
-			TEXT("Add Row to Outliner"),
-			FObserver::OnAdd<FTypedElementLabelColumn>().SetExecutionMode(EExecutionMode::GameThread),
-			[this](IQueryContext& Context, DataStorage::RowHandle Row)
-			{
-				OnItemAdded(Row);
-			})
-		.Compile();
-
-	// Add the conditions from FinalQueryDescription to ensure we are tracking addition of the rows the user requested
-	AppendQuery(RowAdditionQueryDescription, FinalQueryDescription);
-	
-	// Row to track removal of rows from the Outliner
-	FQueryDescription RowRemovalQueryDescription =
-		Select(
-				TEXT("Remove Row from Outliner"),
-				FObserver::OnRemove<FTypedElementLabelColumn>().SetExecutionMode(EExecutionMode::GameThread),
+	if (CreationParams.bUseDefaultObservers)
+	{
+		// Row to track addition of rows to the Outliner
+		FQueryDescription RowAdditionQueryDescription =
+			Select(
+				TEXT("Add Row to Outliner"),
+				FObserver::OnAdd<FTypedElementLabelColumn>().SetExecutionMode(EExecutionMode::GameThread),
 				[this](IQueryContext& Context, DataStorage::RowHandle Row)
 				{
-					OnItemRemoved(Row);
+					OnItemAdded(Row);
 				})
 			.Compile();
 
-	// Add the conditions from FinalQueryDescription to ensure we are tracking removal of the rows the user requested
-	AppendQuery(RowRemovalQueryDescription, FinalQueryDescription);
+		// Add the conditions from FinalQueryDescription to ensure we are tracking addition of the rows the user requested
+		AppendQuery(RowAdditionQueryDescription, FinalQueryDescription);
+	
+		// Row to track removal of rows from the Outliner
+		FQueryDescription RowRemovalQueryDescription =
+			Select(
+					TEXT("Remove Row from Outliner"),
+					FObserver::OnRemove<FTypedElementLabelColumn>().SetExecutionMode(EExecutionMode::GameThread),
+					[this](IQueryContext& Context, DataStorage::RowHandle Row)
+					{
+						OnItemRemoved(Row);
+					})
+				.Compile();
 
+		// Add the conditions from FinalQueryDescription to ensure we are tracking removal of the rows the user requested
+		AppendQuery(RowRemovalQueryDescription, FinalQueryDescription);
+		
+		RowAdditionQuery = Storage->RegisterQuery(MoveTemp(RowAdditionQueryDescription));
+		RowRemovalQuery = Storage->RegisterQuery(MoveTemp(RowRemovalQueryDescription));
+	}
+	
 	// Queries to track parent info, only required if we have hierarchy data
 	if (HierarchyData.IsSet())
 	{
@@ -755,11 +755,40 @@ void FTedsOutlinerImpl::RecompileQueries()
 		SelectionAddedQuery = Storage->RegisterQuery(MoveTemp(SelectionAddedQueryDescription));
 		SelectionRemovedQuery = Storage->RegisterQuery(MoveTemp(SelectionRemovedQueryDescription));
 	}
+
+	// Query to track when the label of a row we are observing changes, to re-filter/re-search for the item
+	FQueryDescription LabelUpdateQueryDescription = 
+		Select(
+			TEXT("Re-Filter Teds Outliner Item on label change"),
+			FProcessor(EQueryTickPhase::PostPhysics, Storage->GetQueryTickGroupName(EQueryTickGroups::SyncExternalToDataStorage))
+				.SetExecutionMode(EExecutionMode::GameThread),
+			[this](IQueryContext& Context, RowHandle Row, const FTypedElementLabelColumn& LabelColumn)
+			{
+				RowsPendingLabelUpdate.Add(Row);
+			}
+		)
+		.Compile();
 	
-	RowHandleQuery = Storage->RegisterQuery(MoveTemp(FinalQueryDescription));
-	RowAdditionQuery = Storage->RegisterQuery(MoveTemp(RowAdditionQueryDescription));
-	RowRemovalQuery = Storage->RegisterQuery(MoveTemp(RowRemovalQueryDescription));
+	if (bUsingQueryConditionsSyntax)
+	{
+		LabelUpdateQueryDescription.Conditions.Emplace(TColumn<FTypedElementSyncBackToWorldTag>());
+		LabelUpdateQueryDescription.Conditions.Emplace(TColumn<FTypedElementSyncFromWorldTag>());
+	}
+	else
+	{
+		LabelUpdateQueryDescription.ConditionTypes.Add(FQueryDescription::EOperatorType::SimpleAny);
+		LabelUpdateQueryDescription.ConditionOperators.AddZeroed_GetRef().Type = FTypedElementSyncBackToWorldTag::StaticStruct();
+		
+		LabelUpdateQueryDescription.ConditionTypes.Add(FQueryDescription::EOperatorType::SimpleAny);
+		LabelUpdateQueryDescription.ConditionOperators.AddZeroed_GetRef().Type = FTypedElementSyncFromWorldTag::StaticStruct();
+	}
+		
+	// Add the conditions from FinalQueryDescription to ensure we are the rows the user requested
+	AppendQuery(LabelUpdateQueryDescription, FinalQueryDescription);
 
+	LabelUpdateQuery = Storage->RegisterQuery(MoveTemp(LabelUpdateQueryDescription));
+	
+	RowHandleQuery = Storage->RegisterQuery(MoveTemp(FinalQueryDescription));
 }
 
 void FTedsOutlinerImpl::UnregisterQueries() const
@@ -774,6 +803,7 @@ void FTedsOutlinerImpl::UnregisterQueries() const
 		Storage->UnregisterQuery(SelectedRowsQuery);
 		Storage->UnregisterQuery(SelectionAddedQuery);
 		Storage->UnregisterQuery(SelectionRemovedQuery);
+		Storage->UnregisterQuery(LabelUpdateQuery);
 	}
 }
 
@@ -822,6 +852,42 @@ void FTedsOutlinerImpl::Tick()
 		OnTedsOutlinerSelectionChanged.Broadcast();
 		bSelectionDirty = false;
 	}
+
+	// Process any new rows that need to be added
+	for (DataStorage::RowHandle Row : RowsPendingAddition)
+	{
+		FSceneOutlinerHierarchyChangedData EventData;
+		EventData.Type = FSceneOutlinerHierarchyChangedData::Added;
+		EventData.Items.Add(SceneOutlinerMode->CreateItemFor<FTedsOutlinerTreeItem>(FTedsOutlinerTreeItem(Row, AsShared())));
+		HierarchyChangedEvent.Broadcast(EventData);
+	}
+	RowsPendingAddition.Empty();
+
+	// Update the label for any rows that might need it
+	for (DataStorage::RowHandle Row : RowsPendingLabelUpdate)
+	{
+		// If the item already exists, it only needs an update if it passed a filter previously and does not now (or vice versa)
+		if (FSceneOutlinerTreeItemPtr ExistingItem = SceneOutliner->GetTreeItem(Row))
+		{
+			bool bCachedFilteredFlag = ExistingItem->Flags.bIsFilteredOut;
+
+			// This implicitly calls into the data storage to get the label of the row and check against the search query
+			ExistingItem->Flags.bIsFilteredOut = !SceneOutliner->PassesAllFilters(ExistingItem);
+
+			if (bCachedFilteredFlag != ExistingItem->Flags.bIsFilteredOut)
+			{
+				SceneOutliner->OnItemLabelChanged(ExistingItem, false);
+			}
+		}
+		// If the item doesn't exist, create a dummy item to see if it would match the current search/filter queries and should be actually added
+		else if (FSceneOutlinerTreeItemPtr PotentialItem = SceneOutlinerMode->CreateItemFor<FTedsOutlinerTreeItem>(
+				FTedsOutlinerTreeItem(Row, AsShared()), true))
+		{
+			SceneOutliner->OnItemLabelChanged(PotentialItem, false);
+		}
+	}
+	
+	RowsPendingLabelUpdate.Empty();
 }
 
 DataStorage::ICoreProvider* FTedsOutlinerImpl::GetStorage() const
diff --git a/Engine/Plugins/Experimental/EditorDataStorageFeatures/Source/TedsOutliner/Public/TedsOutlinerImpl.h b/Engine/Plugins/Experimental/EditorDataStorageFeatures/Source/TedsOutliner/Public/TedsOutlinerImpl.h
--- a/Engine/Plugins/Experimental/EditorDataStorageFeatures/Source/TedsOutliner/Public/TedsOutlinerImpl.h
+++ b/Engine/Plugins/Experimental/EditorDataStorageFeatures/Source/TedsOutliner/Public/TedsOutlinerImpl.h
@@ -101,6 +101,9 @@ struct FTedsOutlinerParams
 	// If true, parent nodes will be remain visible if a child passes all filters even if the parent fails a filter
 	bool bForceShowParents;
 
+	// If true, the Teds Outliner will create observers to track addition/removal of rows and update the Outliner
+	bool bUseDefaultObservers;
+
 	// If specified, this is how the TEDS Outliner will handle hierarchies. If not specified - there will be no hierarchies shown as a
 	// parent-child relation in the tree view
 	TOptional<FTedsOutlinerHierarchyData> HierarchyData;
@@ -244,6 +247,9 @@ protected:
 	DataStorage::QueryHandle SelectedRowsQuery = DataStorage::InvalidQueryHandle;
 	DataStorage::QueryHandle SelectionAddedQuery = DataStorage::InvalidQueryHandle;
 	DataStorage::QueryHandle SelectionRemovedQuery = DataStorage::InvalidQueryHandle;
+
+	// Query to track when a row's label gets changed
+	DataStorage::QueryHandle LabelUpdateQuery;
 	
 	TOptional<FName> SelectionSetName;
 	bool bSelectionDirty = false;
@@ -268,5 +274,10 @@ protected:
 
 	// Whether the row query is using the FConditions syntax (TColumn<FMyColumn>()) or the old syntax (.All().Any().None())
 	bool bUsingQueryConditionsSyntax = false;
+
+	// Addition and Label Updates are deferred because they access data storage implicitly
+	// Deletion is currently not deferred to work nicely with object lifecycles in some cases - but can be once everything goes through the query stack
+	TSet<DataStorage::RowHandle> RowsPendingAddition;
+	TSet<DataStorage::RowHandle> RowsPendingLabelUpdate;
 };
 } // namsepace UE::Editor::Outliner
diff --git a/Engine/Source/Editor/SceneOutliner/Private/SSceneOutliner.cpp b/Engine/Source/Editor/SceneOutliner/Private/SSceneOutliner.cpp
--- a/Engine/Source/Editor/SceneOutliner/Private/SSceneOutliner.cpp
+++ b/Engine/Source/Editor/SceneOutliner/Private/SSceneOutliner.cpp
@@ -2140,7 +2140,7 @@ void SSceneOutliner::PostUndo(bool bSuccess)
 	}
 }
 
-void SSceneOutliner::OnItemLabelChanged(FSceneOutlinerTreeItemPtr ChangedItem)
+void SSceneOutliner::OnItemLabelChanged(FSceneOutlinerTreeItemPtr ChangedItem, bool bFlashHighlight)
 {
 	// If the item already exists
 	if (FSceneOutlinerTreeItemPtr* ExistingItem = TreeItemMap.Find(ChangedItem->GetID()))
@@ -2150,7 +2150,10 @@ void SSceneOutliner::OnItemLabelChanged(FSceneOutlinerTreeItemPtr ChangedItem)
 		// The changed item flags will have been set already
 		if (!ChangedItem->Flags.bIsFilteredOut)
 		{
-			OutlinerTreeView->FlashHighlightOnItem(*ExistingItem);
+			if (bFlashHighlight)
+			{
+				OutlinerTreeView->FlashHighlightOnItem(*ExistingItem);
+			}
 			RequestSort();
 		}
 		else
diff --git a/Engine/Source/Editor/SceneOutliner/Public/SSceneOutliner.h b/Engine/Source/Editor/SceneOutliner/Public/SSceneOutliner.h
--- a/Engine/Source/Editor/SceneOutliner/Public/SSceneOutliner.h
+++ b/Engine/Source/Editor/SceneOutliner/Public/SSceneOutliner.h
@@ -676,7 +676,7 @@ private:
 	}
 
 	/** Handler for when a property changes on any item. Called by the mode */
-	void OnItemLabelChanged(FSceneOutlinerTreeItemPtr ChangedItem);
+	void OnItemLabelChanged(FSceneOutlinerTreeItemPtr ChangedItem, bool bFlashHighlight = true);
 private:
 
 	/** BindCommands for this Outliner */
