Teds Outliner: Initialize label update query to a valid default value, and set all unregistered query handles to the default value as well
Speculative fix for crash in the Teds Outliner

diff --git a/Engine/Plugins/Experimental/EditorDataStorageFeatures/Source/TedsOutliner/Private/TedsOutlinerImpl.cpp b/Engine/Plugins/Experimental/EditorDataStorageFeatures/Source/TedsOutliner/Private/TedsOutlinerImpl.cpp
--- a/Engine/Plugins/Experimental/EditorDataStorageFeatures/Source/TedsOutliner/Private/TedsOutlinerImpl.cpp
+++ b/Engine/Plugins/Experimental/EditorDataStorageFeatures/Source/TedsOutliner/Private/TedsOutlinerImpl.cpp
@@ -791,7 +791,7 @@ void FTedsOutlinerImpl::RecompileQueries()
 	RowHandleQuery = Storage->RegisterQuery(MoveTemp(FinalQueryDescription));
 }
 
-void FTedsOutlinerImpl::UnregisterQueries() const
+void FTedsOutlinerImpl::UnregisterQueries()
 {
 	if (Storage)
 	{
@@ -804,6 +804,16 @@ void FTedsOutlinerImpl::UnregisterQueries() const
 		Storage->UnregisterQuery(SelectionAddedQuery);
 		Storage->UnregisterQuery(SelectionRemovedQuery);
 		Storage->UnregisterQuery(LabelUpdateQuery);
+		
+		RowHandleQuery = DataStorage::InvalidQueryHandle;
+		RowAdditionQuery = DataStorage::InvalidQueryHandle;
+		RowRemovalQuery = DataStorage::InvalidQueryHandle;
+		ChildRowHandleQuery = DataStorage::InvalidQueryHandle;
+		UpdateParentQuery = DataStorage::InvalidQueryHandle;
+		SelectedRowsQuery = DataStorage::InvalidQueryHandle;
+		SelectionAddedQuery = DataStorage::InvalidQueryHandle;
+		SelectionRemovedQuery = DataStorage::InvalidQueryHandle;
+		LabelUpdateQuery = DataStorage::InvalidQueryHandle;
 	}
 }
 
diff --git a/Engine/Plugins/Experimental/EditorDataStorageFeatures/Source/TedsOutliner/Public/TedsOutlinerImpl.h b/Engine/Plugins/Experimental/EditorDataStorageFeatures/Source/TedsOutliner/Public/TedsOutlinerImpl.h
--- a/Engine/Plugins/Experimental/EditorDataStorageFeatures/Source/TedsOutliner/Public/TedsOutlinerImpl.h
+++ b/Engine/Plugins/Experimental/EditorDataStorageFeatures/Source/TedsOutliner/Public/TedsOutlinerImpl.h
@@ -195,7 +195,7 @@ public:
 	void RecompileQueries();
 
 	// Unregister all queries used by this table viewer
-	void UnregisterQueries() const;
+	void UnregisterQueries();
 
 protected:
 	
@@ -249,7 +249,7 @@ protected:
 	DataStorage::QueryHandle SelectionRemovedQuery = DataStorage::InvalidQueryHandle;
 
 	// Query to track when a row's label gets changed
-	DataStorage::QueryHandle LabelUpdateQuery;
+	DataStorage::QueryHandle LabelUpdateQuery = DataStorage::InvalidQueryHandle;
 	
 	TOptional<FName> SelectionSetName;
 	bool bSelectionDirty = false;
