Teds Revision Control: Stop monitoring unresolved packages if they get resolved (or the row is removed)

diff --git a/Engine/Plugins/Experimental/EditorDataStorageFeatures/Source/TedsRevisionControl/Private/Processors/RevisionControlProcessors.cpp b/Engine/Plugins/Experimental/EditorDataStorageFeatures/Source/TedsRevisionControl/Private/Processors/RevisionControlProcessors.cpp
--- a/Engine/Plugins/Experimental/EditorDataStorageFeatures/Source/TedsRevisionControl/Private/Processors/RevisionControlProcessors.cpp
+++ b/Engine/Plugins/Experimental/EditorDataStorageFeatures/Source/TedsRevisionControl/Private/Processors/RevisionControlProcessors.cpp
@@ -255,6 +255,9 @@ void URevisionControlDataStorageFactory::RegisterQueries(UE::Editor::DataStorage
 			{
 				DataStorage.UnregisterQuery(FetchUpdates);
 				FetchUpdates = InvalidQueryHandle;
+				
+				DataStorage.UnregisterQuery(StopFetchUpdates);
+				StopFetchUpdates = InvalidQueryHandle;
 			}
 		}
 	);
@@ -347,6 +350,25 @@ void URevisionControlDataStorageFactory::RegisterFetchUpdates(UE::Editor::DataSt
 				})
 			.Compile());
 	}
+	
+	if (StopFetchUpdates == InvalidQueryHandle)
+	{
+		StopFetchUpdates = DataStorage.RegisterQuery(
+			Select(
+				TEXT("Stop monitoring source control statuses for objects with resolved package paths"),
+				FObserver::OnRemove<FTypedElementPackageUnresolvedReference>()
+					.SetExecutionMode(EExecutionMode::GameThread),
+				[this, &FileStatusMonitor](IQueryContext& Context, const FTypedElementPackageUnresolvedReference& UnresolvedReference)
+				{
+					static FSourceControlFileStatusMonitor::FOnSourceControlFileStatus EmptyDelegate{};
+				
+					FileStatusMonitor.StopMonitoringFile(
+						reinterpret_cast<uintptr_t>(this),
+						UnresolvedReference.PathOnDisk
+					);
+				})
+			.Compile());
+	}
 }
 
 void URevisionControlDataStorageFactory::RegisterApplyOverlays(UE::Editor::DataStorage::ICoreProvider& DataStorage)
diff --git a/Engine/Plugins/Experimental/EditorDataStorageFeatures/Source/TedsRevisionControl/Private/Processors/RevisionControlProcessors.h b/Engine/Plugins/Experimental/EditorDataStorageFeatures/Source/TedsRevisionControl/Private/Processors/RevisionControlProcessors.h
--- a/Engine/Plugins/Experimental/EditorDataStorageFeatures/Source/TedsRevisionControl/Private/Processors/RevisionControlProcessors.h
+++ b/Engine/Plugins/Experimental/EditorDataStorageFeatures/Source/TedsRevisionControl/Private/Processors/RevisionControlProcessors.h
@@ -38,6 +38,7 @@ private:
 	void RegisterGeneralQueries(UE::Editor::DataStorage::ICoreProvider& DataStorage);
 
 	UE::Editor::DataStorage::QueryHandle FetchUpdates = UE::Editor::DataStorage::InvalidQueryHandle;
+	UE::Editor::DataStorage::QueryHandle StopFetchUpdates = UE::Editor::DataStorage::InvalidQueryHandle;
 	UE::Editor::DataStorage::QueryHandle ApplyOverlaysObjectToSCC = UE::Editor::DataStorage::InvalidQueryHandle;
 	UE::Editor::DataStorage::QueryHandle RemoveOverlays = UE::Editor::DataStorage::InvalidQueryHandle;
 	UE::Editor::DataStorage::QueryHandle UpdateSCCForActors = UE::Editor::DataStorage::InvalidQueryHandle;
