[TEDS] Query Stack node to sort rows by column values.

[RN] TEDS Query Stack node to sort rows by column values.

diff --git a/Engine/Plugins/Experimental/EditorDataStorageFeatures/Source/TedsQueryStack/Private/Sorting/SortTasks.cpp b/Engine/Plugins/Experimental/EditorDataStorageFeatures/Source/TedsQueryStack/Private/Sorting/SortTasks.cpp
new file mode 100644
--- /dev/null
+++ b/Engine/Plugins/Experimental/EditorDataStorageFeatures/Source/TedsQueryStack/Private/Sorting/SortTasks.cpp
@@ -0,0 +1,650 @@
+// Copyright Epic Games, Inc. All Rights Reserved.
+
+#include "Sorting/SortTasks.h"
+
+#include "Algo/StableSort.h"
+#include "Async/ParallelFor.h"
+#include "Elements/Framework/TypedElementRowHandleArray.h"
+#include "Elements/Framework/TypedElementSorter.h"
+#include "ProfilingDebugging/CpuProfilerTrace.h"
+
+namespace UE::Editor::DataStorage::QueryStack::Sorting
+{
+	void FSortTasks::AddBarrier()
+	{
+		AddBarrier(MainTaskList);
+	}
+
+	void FSortTasks::CollectPrefixes(TArrayView<FPrefix> Output, TConstArrayView<RowHandle> Input,
+		const ICoreProvider& Storage, TSharedPtr<const FColumnSorterInterface> Sorter, int32 BatchSize)
+	{
+		CollectPrefixes(MainTaskList, Output, Input, Storage, Sorter, BatchSize);
+	}
+
+	void FSortTasks::RefreshPrefixes(TArrayView<FPrefix> Prefixes, const ICoreProvider& Storage,
+		TSharedPtr<const FColumnSorterInterface> Sorter, uint32 ByteIndex, int32 BatchSize)
+	{
+		RefreshPrefixes(MainTaskList, Prefixes, Storage, Sorter, ByteIndex, BatchSize);
+	}
+
+	void FSortTasks::RadixSortPrefixes(TArrayView<FPrefix> Output)
+	{
+		RadixSortPrefixes(MainTaskList, Output);
+	}
+
+	void FSortTasks::CopyPrefixesToRowHandleArray(TArrayView<RowHandle> Output, TConstArrayView<FPrefix> Input)
+	{
+		CopyPrefixesToRowHandleArray(MainTaskList, Output, Input);
+	}
+
+	void FSortTasks::ComparativeSortRowHandleRanges(TArrayView<RowHandle> Rows, const ICoreProvider& Storage,
+		TSharedPtr<const FColumnSorterInterface> Sorter, int32 BatchSize)
+	{
+		ComparativeSortRowHandleRanges(MainTaskList, Rows, Storage, Sorter, BatchSize);
+	}
+
+	void FSortTasks::MergeSortedAdjacentRanges(TArrayView<RowHandle> Array, TArrayView<RowHandle> Scratch,
+		const ICoreProvider& Storage, TSharedPtr<const FColumnSorterInterface> Sorter, int32 SplitPosition)
+	{
+		MergeSortedAdjacentRanges(MainTaskList, Array, Scratch, Storage, Sorter, SplitPosition);
+	}
+
+	void FSortTasks::BucketizeRange(TArrayView<RowHandle> Array, TArrayView<FPrefix> Prefixes, TArrayView<FPrefix> PrefixesShadow, 
+		TSharedPtr<const FColumnSorterInterface> Sorter, const ICoreProvider& Storage, int32 PrefixBatchSize, int32 SortBatchSize)
+	{
+		BucketizeRange(MainTaskList, Array, Prefixes, PrefixesShadow, Sorter, Storage, PrefixBatchSize, SortBatchSize, 0);
+	}
+
+	void FSortTasks::BuildTaskChain(TaskReferenceChain& Output)
+	{
+		// First collect all newly added chains.
+		while (TaskChain* Next = PendingTasksQueue.Peek())
+		{
+			PendingTasks.Add(MoveTemp(*Next));
+			PendingTasksQueue.Dequeue();
+		}
+
+		// Then for each chain, add their commands to the overall chain.
+		auto AddTaskToChain = [](TaskReferenceChain& Output, TaskChain& Chain)
+			{
+				for (FPrefixCollectionBatch& Batch : Chain)
+				{
+					if (!Batch.bCompleted)
+					{
+						if (!Batch.Task.IsType<FTaskBarrier>())
+						{
+							Output.Add(&Batch);
+						}
+						else
+						{
+							break;
+						}
+					}
+				}
+			};
+		for (TaskChain& Chain : PendingTasks)
+		{
+			AddTaskToChain(Output, Chain);
+		}
+		AddTaskToChain(Output, MainTaskList);
+	}
+
+	void FSortTasks::RemoveCompletedTasks()
+	{
+		auto HasCompleted = [](TaskChain& Chain)
+			{
+				for (FPrefixCollectionBatch& Batch : Chain)
+				{
+					if (!Batch.bCompleted)
+					{
+						if (Batch.Task.IsType<FTaskBarrier>())
+						{
+							Batch.bCompleted = true;
+						}
+						else
+						{
+							return false;
+						}
+					}
+				}
+				return true;
+			};
+		for (auto It = PendingTasks.CreateIterator(); It; ++It)
+		{
+			if (HasCompleted(*It))
+			{
+				It.RemoveCurrentSwap();
+			}
+		}
+		if (HasCompleted(MainTaskList))
+		{
+			MainTaskList.Reset();
+		}
+	}
+
+	void FSortTasks::Update(FTimespan AllottedTime)
+	{
+		if (HasRemainingTasks())
+		{
+			bIsUpdating = true;
+			uint64 StartTime = FPlatformTime::Cycles64();
+			do
+			{
+				BuildTaskChain(LocalTaskChain);
+
+				// Execute the tasks in this batch.
+				ParallelFor(LocalTaskChain.Num(), [this, StartTime, AllottedTime](int32 Batch)
+					{
+						if (FTimespan(FPlatformTime::Cycles64() - StartTime) < AllottedTime)
+						{
+							FPrefixCollectionBatch* BatchInfo = LocalTaskChain[Batch];
+							checkf(BatchInfo->bCompleted == false, TEXT("Trying to execute an already completed sort task."));
+							BatchInfo->bCompleted = true;
+							Visit(FTaskVisitor
+								{
+									.Owner = *this
+								}, BatchInfo->Task);
+						}
+					}, EParallelForFlags::Unbalanced);
+
+				RemoveCompletedTasks();
+				LocalTaskChain.Reset();
+			} while (HasRemainingTasks() && FTimespan(FPlatformTime::Cycles64() - StartTime) < AllottedTime);
+			bIsUpdating = false;
+		}
+	}
+
+	bool FSortTasks::HasRemainingTasks() const
+	{
+		return !MainTaskList.IsEmpty() || !PendingTasks.IsEmpty() || !PendingTasksQueue.IsEmpty();
+	}
+
+	void FSortTasks::FTaskVisitor::operator()(const FTaskBarrier& Task)
+	{
+		checkf(false, TEXT("Task barriers should never be processed."));
+	}
+
+	void FSortTasks::FTaskVisitor::operator()(const FCollectPrefixesTask& Task)
+	{
+		TRACE_CPUPROFILER_EVENT_SCOPE("[TEDS Sorting] Collect prefixes");
+
+		const RowHandle* Source = Task.Input.GetData();
+		FPrefix* Target = Task.Output.GetData();
+		int32 Count = Task.Input.Num();
+
+		for (int32 Index = 0; Index < Count; ++Index)
+		{
+			FPrefixInfo Result = Task.Sorter->CalculatePrefix(Task.Storage, *Source, 0);
+			Target->Prefix = Result.Prefix;
+			Target->Row = *Source;
+			Target->bHasRemainingBytes = Result.bHasRemainingBytes;
+
+			Source++;
+			Target++;
+		}
+	}
+
+	void FSortTasks::FTaskVisitor::operator()(const FRefreshPrefixesTask& Task)
+	{
+		TRACE_CPUPROFILER_EVENT_SCOPE("[TEDS Sorting] Refresh prefixes");
+
+		for (FPrefix& Prefix : Task.Prefixes)
+		{
+			if (Prefix.bHasRemainingBytes)
+			{
+				FPrefixInfo Result = Task.Sorter->CalculatePrefix(Task.Storage, Prefix.Row, Task.ByteIndex);
+				Prefix.Prefix = Result.Prefix;
+				Prefix.bHasRemainingBytes = Result.bHasRemainingBytes;
+			}
+			else
+			{
+				Prefix.Prefix = 0;
+			}
+		}
+	}
+
+	void FSortTasks::FTaskVisitor::operator()(const FRadixSortPrefixesTask& Task)
+	{
+		TRACE_CPUPROFILER_EVENT_SCOPE("[TEDS Sorting] Radix sort prefixes");
+
+		RadixSort64(Task.Output.GetData(), Task.Output.Num(),
+			[](const FPrefix& Prefix)
+			{
+				return Prefix.Prefix;
+			});
+	}
+
+	void FSortTasks::FTaskVisitor::operator()(const FCopyPrefixesToRowHandleArrayTask& Task)
+	{
+		TRACE_CPUPROFILER_EVENT_SCOPE("[TEDS Sorting] Copy out prefixes to array");
+
+		const FPrefix* InputIt = Task.Input.GetData();
+		const FPrefix* InputEnd = InputIt + Task.Input.Num();
+		RowHandle* OutputIt = Task.Output.GetData();
+
+		while (InputIt < InputEnd)
+		{
+			*OutputIt = InputIt->Row;
+			OutputIt++;
+			InputIt++;
+		}
+	}
+
+	void FSortTasks::FTaskVisitor::operator()(const FCopyPrefixesToRowHandleArrayViewTask& Task)
+	{
+		TRACE_CPUPROFILER_EVENT_SCOPE("[TEDS Sorting] Copy out prefixes to array view");
+
+		RowHandle* WriteIt = Task.Output.GetData();
+		const FPrefix* ReadIt = Task.Input.GetData();
+		const FPrefix* ReadEnd = ReadIt + Task.Input.Num();
+
+		for (; ReadIt < ReadEnd; ++ReadIt)
+		{
+			*WriteIt++ = ReadIt->Row;
+		}
+	}
+
+	void FSortTasks::FTaskVisitor::operator()(const FComparativeSortRowHandleRangeTask& Task)
+	{
+		TRACE_CPUPROFILER_EVENT_SCOPE("[TEDS Sorting] Comparative sort");
+
+		Algo::StableSort(Task.Rows, [&Task](RowHandle Left, RowHandle Right)
+			{
+				return Task.Sorter->Compare(Task.Storage, Left, Right) < 0;
+			});
+	}
+
+	void FSortTasks::FTaskVisitor::operator()(const FComparativeSortRowHandlePrefixRangeTask& Task)
+	{
+		TRACE_CPUPROFILER_EVENT_SCOPE("[TEDS Sorting] Comparative sort with prefix");
+
+		Algo::StableSort(Task.Rows, [&Task](const FPrefix& Left, const FPrefix& Right)
+			{
+				if (Left.Prefix == Right.Prefix)
+				{
+					return Task.Sorter->Compare(Task.Storage, Left.Row, Right.Row) < 0;
+				}
+				else
+				{
+					return (Left.Prefix < Right.Prefix) ? true : false;
+				}
+			});
+	}
+
+	void FSortTasks::FTaskVisitor::operator()(const FMergeSortedAdjacentRangesTask& Task)
+	{
+		TRACE_CPUPROFILER_EVENT_SCOPE("[TEDS Sorting] Merge");
+
+		// Only merge if the first element from the right range is within the left range, otherwise the 
+		// array is already sorted.
+		TArrayView<RowHandle> Left(Task.Array.GetData(), Task.SplitPosition);
+		
+		RowHandle* LeftBegin = Task.Array.GetData();
+		RowHandle* RightBegin = LeftBegin + Task.SplitPosition;
+		RowHandle* LeftLast = RightBegin - 1;
+		RowHandle* LeftEnd = RightBegin;
+		RowHandle* RightEnd = LeftBegin + Task.Array.Num();
+		RowHandle* RightIt = RightEnd - 1;
+		RowHandle* InsertBack = RightIt;
+
+		if (Task.Sorter->Compare(Task.Storage, *LeftLast, *RightBegin) > 0)
+		{
+			int32 RightCount = static_cast<int32>(RightEnd - RightBegin);
+
+			// Copy the right section to a temporary location so it doesn't interfere with the sorting. Update the pointers
+			// to for right as the values are now in the scratch buffer.
+			FMemory::Memcpy(Task.Scratch.GetData(), RightBegin, RightCount * sizeof(RowHandle));
+			RightBegin = Task.Scratch.GetData();
+			RightEnd = RightBegin + RightCount;
+			RightIt = RightEnd - 1;
+			
+			// For each entry in the right section, find the insertion point in the left section. Move everything from the
+			// left up to that point and insert the right value.
+			while (RightIt >= RightBegin)
+			{
+				int32 InsertPoint = Algo::LowerBound(Left, *RightIt, 
+					[&Task](RowHandle Left, RowHandle Right)
+					{
+						return Task.Sorter->Compare(Task.Storage, Left, Right) < 0;
+					});
+				if (int32 ReadBlockSize = Left.Num() - InsertPoint; ReadBlockSize > 0)
+				{
+					InsertBack -= ReadBlockSize - 1;
+					FMemory::Memmove(InsertBack, Left.GetData() + InsertPoint, ReadBlockSize * sizeof(RowHandle));
+					Left.LeftChopInline(ReadBlockSize);
+					InsertBack--;
+				}
+				*InsertBack = *RightIt;
+				RightIt--;
+				InsertBack--;
+			}
+			// Any values remaining on the left side will already be sorted.
+		}
+	}
+
+	void FSortTasks::FTaskVisitor::operator()(const FBucketizeRangesTask& Task)
+	{
+		TRACE_CPUPROFILER_EVENT_SCOPE("[TEDS Sorting] Bucketize");
+
+		int32 Counts[256];
+		FMemory::Memzero(Counts, sizeof(Counts));
+		bool bHasRemainingBytes[256];
+		FMemory::Memzero(bHasRemainingBytes, sizeof(bHasRemainingBytes));
+
+		// Collect the prefix count at the request byte.
+		int32 Shift = 56 - ((Task.Index & 7) << 3);
+		for (const FPrefix& Prefix : Task.Prefixes)
+		{
+			uint8 Index = static_cast<uint8>(Prefix.Prefix >> Shift);
+			Counts[Index]++;
+			bHasRemainingBytes[Index] = bHasRemainingBytes[Index] || Prefix.bHasRemainingBytes;
+		}
+
+		// Convert counts into offsets.
+		int32 Offsets[256];
+		Offsets[0] = 0;
+		for (int32 Index = 0; Index < 255; ++Index)
+		{
+			Offsets[Index + 1] = Offsets[Index] + Counts[Index];
+		}
+
+		// Distribute the prefixes.
+		for (const FPrefix& Prefix : Task.Prefixes)
+		{
+			uint8 Index = static_cast<uint8>(Prefix.Prefix >> Shift);
+			Task.PrefixesShadow[Offsets[Index]++] = Prefix;
+		}
+		
+		// Create new tasks for each group that still needs work.
+		uint32 NextIndex = Task.Index + 1;
+		TaskChain SubTasks;
+		if ((NextIndex & 7) == 0)
+		{
+			int32 Offset = 0;
+			bool bAddedRefreshPrefixesTask = false;
+			for (int32 Index = 0; Index < 256; ++Index)
+			{
+				int32 Count = Counts[Index];
+				if (Count == 1)
+				{
+					// If there's only one value in the bucket, then it's already in the correct place and there's no need for further 
+					// sorting.
+					Task.Array[Offset] = Task.PrefixesShadow[Offset].Row;
+				}
+				else if (bHasRemainingBytes[Index]) // If there are no entries in this slot this is guaranteed to be false.
+				{
+					Owner.RefreshPrefixes(
+						SubTasks, Task.PrefixesShadow.Slice(Offset, Count), Task.Storage, Task.Sorter, NextIndex, Task.PrefixBatchSize);
+					bAddedRefreshPrefixesTask = true;
+				}
+				else
+				{
+					// If there are no prefixes left, there's no more sorting that can happen, which means the values all have the same value.
+					for (int32 SubIndex = Offset; SubIndex < Offset + Count; ++SubIndex)
+					{
+						Task.Array[SubIndex] = Task.PrefixesShadow[SubIndex].Row;
+					}
+				}
+				Offset += Count;
+			}
+
+			if (bAddedRefreshPrefixesTask)
+			{
+				Owner.AddBarrier(SubTasks);
+				
+				Offset = 0;
+				for (int32 Index = 0; Index < 256; ++Index)
+				{
+					int32 Count = Counts[Index];
+					if (bHasRemainingBytes[Index])
+					{
+						// Can't start a comparative sort here because the barrier wouldn't protect against work in a separate task list.
+						// This would mean adding the compares in the same task list which would result in additional barriers in this list
+						// and less parallelization. Instead let another loop of bucketization happen which guarantees the prefixes are
+						// collected.
+						Owner.BucketizeRange(
+							SubTasks,
+							Task.Array.Slice(Offset, Count),
+							Task.PrefixesShadow.Slice(Offset, Count),
+							Task.Prefixes.Slice(Offset, Count),
+							Task.Sorter, Task.Storage, Task.PrefixBatchSize, Task.SortBatchSize, NextIndex);
+					}
+					Offset += Count;
+				}
+			}
+		}
+		else
+		{
+			int32 Offset = 0;
+			for (int32 Index = 0; Index < 256; ++Index)
+			{
+				int32 Count = Counts[Index];
+				if (Count == 1)
+				{
+					// If there's only one value in the bucket, then it's already in the correct place and there's no need for further 
+					// sorting.
+					Task.Array[Offset] = Task.PrefixesShadow[Offset].Row;
+				}
+				else if (Count > 1)
+				{
+					if (Count < Task.SortBatchSize)
+					{
+						// Do the final sort and copy in a separate list to avoid the barrier causing additional delays and maximizing
+						// parallelization.
+						TArrayView<FPrefix> SubPrefixView = Task.PrefixesShadow.Slice(Offset, Count);
+						TaskChain CompareSortTask;
+						CompareSortTask.Reserve(3);
+						Owner.ComparativeSortRowHandlePrefixRanges(CompareSortTask, SubPrefixView, Task.Storage, Task.Sorter);
+						Owner.AddBarrier(CompareSortTask);
+						Owner.CopyPrefixesToRowHandleArrayView(CompareSortTask, Task.Array.Slice(Offset, Count), SubPrefixView);
+						Owner.PendingTasksQueue.Enqueue(MoveTemp(CompareSortTask));
+					}
+					else
+					{
+						Owner.BucketizeRange(
+							SubTasks,
+							Task.Array.Slice(Offset, Count),
+							Task.PrefixesShadow.Slice(Offset, Count),
+							Task.Prefixes.Slice(Offset, Count),
+							Task.Sorter, Task.Storage, Task.PrefixBatchSize, Task.SortBatchSize, NextIndex);
+					}
+				}
+				Offset += Count;
+			}
+		}
+		if (!SubTasks.IsEmpty())
+		{
+			Owner.PendingTasksQueue.Enqueue(MoveTemp(SubTasks));
+		}
+	}
+
+	void FSortTasks::AddBarrier(TaskChain& Tasks)
+	{
+		Tasks.Emplace(FPrefixCollectionBatch
+			{
+				.Task = TaskVariant(TInPlaceType<FTaskBarrier>(), FTaskBarrier{})
+			});
+	}
+
+	void FSortTasks::CollectPrefixes(TaskChain& Tasks, TArrayView<FPrefix> Output, TConstArrayView<RowHandle> Input,
+		const ICoreProvider& Storage, TSharedPtr<const FColumnSorterInterface> Sorter, int32 BatchSize)
+	{
+		checkf(Output.Num() >= Input.Num(), TEXT("A prefix array can't be generated from an input that's smaller than the prefix table."));
+
+		int32 Count = Input.Num();
+		Tasks.Reserve(Tasks.Num() + (Count / BatchSize) + 1);
+		for (int32 Batch = 0; Batch < Count; Batch += BatchSize)
+		{
+			int32 SubBatchEnd = FMath::Min(Batch + BatchSize, Count);
+			int32 SubBatchCount = SubBatchEnd - Batch;
+
+			Tasks.Emplace(FPrefixCollectionBatch
+				{
+					.Task = TaskVariant(TInPlaceType<FCollectPrefixesTask>(), FCollectPrefixesTask
+						{
+							.Output = Output.Slice(Batch, SubBatchCount),
+							.Input = Input.Slice(Batch, SubBatchCount),
+							.Sorter = Sorter,
+							.Storage = Storage
+						})
+				});
+		}
+	}
+
+	void FSortTasks::RefreshPrefixes(TaskChain& Tasks, TArrayView<FPrefix> Prefixes, const ICoreProvider& Storage,
+		TSharedPtr<const FColumnSorterInterface> Sorter, uint32 ByteIndex, int32 BatchSize)
+	{
+		int32 Count = Prefixes.Num();
+		Tasks.Reserve(Tasks.Num() + (Count / BatchSize) + 1);
+		for (int32 Batch = 0; Batch < Count; Batch += BatchSize)
+		{
+			int32 SubBatchEnd = FMath::Min(Batch + BatchSize, Count);
+			int32 SubBatchCount = SubBatchEnd - Batch;
+
+			Tasks.Emplace(FPrefixCollectionBatch
+				{
+					.Task = TaskVariant(TInPlaceType<FRefreshPrefixesTask>(), FRefreshPrefixesTask
+					{
+						.Prefixes = Prefixes.Slice(Batch, SubBatchCount),
+						.Sorter = Sorter,
+						.Storage = Storage,
+						.ByteIndex = ByteIndex
+					})
+				});
+		}
+	}
+
+	void FSortTasks::RadixSortPrefixes(TaskChain& Tasks, TArrayView<FPrefix> Output)
+	{
+		Tasks.Emplace(FPrefixCollectionBatch
+			{
+				.Task = TaskVariant(TInPlaceType<FRadixSortPrefixesTask>(), FRadixSortPrefixesTask
+					{
+						.Output = Output
+					})
+			});
+	}
+
+	void FSortTasks::CopyPrefixesToRowHandleArray(TaskChain& Tasks, TArrayView<RowHandle> Output, TConstArrayView<FPrefix> Input)
+	{
+		checkf(Output.Num() == Input.Num(),
+			TEXT("The input sort prefix list doesn't have the same number of entries as the row handle output array."));
+		Tasks.Emplace(FPrefixCollectionBatch
+			{
+				.Task = TaskVariant(TInPlaceType<FCopyPrefixesToRowHandleArrayTask>(), FCopyPrefixesToRowHandleArrayTask
+					{
+						.Output = Output,
+						.Input = Input
+					})
+			});
+	}
+
+	void FSortTasks::CopyPrefixesToRowHandleArrayView(TaskChain& Tasks, TArrayView<RowHandle> Output, TConstArrayView<FPrefix> Input)
+	{
+		checkf(Output.Num() == Input.Num(),
+			TEXT("The input sort prefix list doesn't have the same number of entries as the row handle output array view."));
+		Tasks.Emplace(FPrefixCollectionBatch
+			{
+				.Task = TaskVariant(TInPlaceType<FCopyPrefixesToRowHandleArrayViewTask>(), FCopyPrefixesToRowHandleArrayViewTask
+					{
+						.Output = Output,
+						.Input = Input
+					})
+			});
+	}
+
+	void FSortTasks::ComparativeSortRowHandleRanges(TaskChain& Tasks, TArrayView<RowHandle> Rows, const ICoreProvider& Storage,
+		TSharedPtr<const FColumnSorterInterface> Sorter, int32 BatchSize)
+	{
+		int32 Count = Rows.Num();
+		Tasks.Reserve(Tasks.Num() + (Count / BatchSize) + 1);
+		for (int32 Batch = 0; Batch < Count; Batch += BatchSize)
+		{
+			int32 SubBatchCount = FMath::Min(Batch + BatchSize, Count) - Batch;
+			Tasks.Emplace(FPrefixCollectionBatch
+				{
+					.Task = TaskVariant(TInPlaceType<FComparativeSortRowHandleRangeTask>(), FComparativeSortRowHandleRangeTask
+						{
+							.Rows = Rows.Slice(Batch, SubBatchCount),
+							.Sorter = Sorter,
+							.Storage = Storage
+						})
+				});
+		}
+	}
+
+	void FSortTasks::ComparativeSortRowHandlePrefixRanges(TaskChain& Tasks, TArrayView<FPrefix> Rows, const ICoreProvider& Storage,
+		TSharedPtr<const FColumnSorterInterface> Sorter)
+	{
+		Tasks.Emplace(FPrefixCollectionBatch
+			{
+				.Task = TaskVariant(TInPlaceType<FComparativeSortRowHandlePrefixRangeTask>(), FComparativeSortRowHandlePrefixRangeTask
+					{
+						.Rows = Rows,
+						.Sorter = Sorter,
+						.Storage = Storage
+					})
+			});
+	}
+
+	void FSortTasks::MergeSortedAdjacentRanges(TaskChain& Tasks, TArrayView<RowHandle> Array, TArrayView<RowHandle> Scratch,
+		const ICoreProvider& Storage, TSharedPtr<const FColumnSorterInterface> Sorter, int32 SplitPosition)
+	{
+		int32 Count = Array.Num();
+		// Only sort if the number of entries is larger than the split position, otherwise there is no right side to merge in.
+		while (Count > SplitPosition)
+		{
+			int32 ScratchBatch = 0;
+			int32 BatchSize = SplitPosition * 2;
+			Tasks.Reserve(Tasks.Num() + (Count / BatchSize) + 1);
+			for (int32 Batch = 0; Batch < Count; Batch += BatchSize)
+			{
+				int32 SubBatchCount = FMath::Min(Batch + BatchSize, Count) - Batch;
+				if (SubBatchCount > SplitPosition)
+				{
+					int32 SubBatchScratchCount = SubBatchCount - SplitPosition;
+
+					Tasks.Emplace(FPrefixCollectionBatch
+						{
+							.Task = TaskVariant(TInPlaceType<FMergeSortedAdjacentRangesTask>(), FMergeSortedAdjacentRangesTask
+								{
+									.Array = Array.Slice(Batch, SubBatchCount),
+									.Scratch = Scratch.Slice(ScratchBatch, SubBatchScratchCount),
+									.Sorter = Sorter,
+									.Storage = Storage,
+									.SplitPosition = SplitPosition
+								})
+						});
+					ScratchBatch += SubBatchScratchCount;
+				}
+			}
+
+			Tasks.Emplace(FPrefixCollectionBatch
+				{
+					.Task = TaskVariant(TInPlaceType<FTaskBarrier>(), FTaskBarrier{})
+				});
+
+			SplitPosition *= 2;
+		}
+	}
+
+	void FSortTasks::BucketizeRange(TaskChain& Tasks, TArrayView<RowHandle> Array, TArrayView<FPrefix> Prefixes, 
+		TArrayView<FPrefix> PrefixesShadow, TSharedPtr<const FColumnSorterInterface> Sorter, const ICoreProvider& Storage, 
+		int32 PrefixBatchSize, int32 SortBatchSize, uint32 ByteIndex)
+	{
+		Tasks.Emplace(FPrefixCollectionBatch
+			{
+				.Task = TaskVariant(TInPlaceType<FBucketizeRangesTask>(), FBucketizeRangesTask
+					{
+						.Array = Array,
+						.Prefixes = Prefixes,
+						.PrefixesShadow = PrefixesShadow,
+						.Sorter = Sorter,
+						.Storage = Storage,
+						.PrefixBatchSize = PrefixBatchSize,
+						.SortBatchSize = SortBatchSize,
+						.Index = ByteIndex
+					})
+			});
+	}
+} // namespace UE::Editor::DataStorage::QueryStack::Sorting
diff --git a/Engine/Plugins/Experimental/EditorDataStorageFeatures/Source/TedsQueryStack/Private/Sorting/SortTasks.h b/Engine/Plugins/Experimental/EditorDataStorageFeatures/Source/TedsQueryStack/Private/Sorting/SortTasks.h
new file mode 100644
--- /dev/null
+++ b/Engine/Plugins/Experimental/EditorDataStorageFeatures/Source/TedsQueryStack/Private/Sorting/SortTasks.h
@@ -0,0 +1,178 @@
+// Copyright Epic Games, Inc. All Rights Reserved.
+
+#pragma once
+
+#include "Async/Mutex.h"
+#include "Containers/ArrayView.h"
+#include "Containers/MpscQueue.h"
+#include "Elements/Common/TypedElementHandles.h"
+#include "Misc/Timespan.h"
+#include "Misc/TVariant.h"
+#include "Templates/SharedPointer.h"
+
+namespace UE::Editor::DataStorage
+{
+	class FColumnSorterInterface;
+	class FRowHandleArray;
+	class ICoreProvider;
+}
+
+namespace UE::Editor::DataStorage::QueryStack::Sorting
+{
+	struct FPrefix
+	{
+		uint64 Prefix;
+		RowHandle Row;
+		bool bHasRemainingBytes;
+	};
+
+	/** Simple task queue system for the various operations needed to sort row handles using a column sorter. */
+	class FSortTasks
+	{
+	public:
+		void AddBarrier();
+		void CollectPrefixes(TArrayView<FPrefix> Output, TConstArrayView<RowHandle> Input,
+			const ICoreProvider& Storage, TSharedPtr<const FColumnSorterInterface> Sorter, int32 BatchSize);
+		void RefreshPrefixes(TArrayView<FPrefix> Prefixes, const ICoreProvider& Storage, 
+			TSharedPtr<const FColumnSorterInterface> Sorter, uint32 ByteIndex, int32 BatchSize);
+		void RadixSortPrefixes(TArrayView<FPrefix> Output);
+		void CopyPrefixesToRowHandleArray(TArrayView<RowHandle> Output, TConstArrayView<FPrefix> Input);
+		void ComparativeSortRowHandleRanges(TArrayView<RowHandle> Rows, const ICoreProvider& Storage, 
+			TSharedPtr<const FColumnSorterInterface> Sorter, int32 BatchSize);
+		void MergeSortedAdjacentRanges(TArrayView<RowHandle> Array, TArrayView<RowHandle> Scratch,
+			const ICoreProvider& Storage, TSharedPtr<const FColumnSorterInterface> Sorter, int32 SplitPosition);
+		void BucketizeRange(TArrayView<RowHandle> Array, TArrayView<FPrefix> Prefixes, TArrayView<FPrefix> PrefixesShadow, 
+			TSharedPtr<const FColumnSorterInterface> Sorter, const ICoreProvider& Storage, int32 PrefixBatchSize, int32 SortBatchSize);
+		
+		void Update(FTimespan AllottedTime);
+
+		bool HasRemainingTasks() const;
+
+	private:
+		struct FTaskBarrier {};
+		struct FCollectPrefixesTask
+		{
+			TArrayView<FPrefix> Output;
+			TConstArrayView<RowHandle> Input;
+			TSharedPtr<const FColumnSorterInterface> Sorter;
+			const ICoreProvider& Storage;
+		};
+		struct FRefreshPrefixesTask
+		{
+			TArrayView<FPrefix> Prefixes;
+			TSharedPtr<const FColumnSorterInterface> Sorter;
+			const ICoreProvider& Storage;
+			uint32 ByteIndex;
+		};
+		struct FRadixSortPrefixesTask
+		{
+			TArrayView<FPrefix> Output;
+		};
+		struct FCopyPrefixesToRowHandleArrayTask
+		{
+			TArrayView<RowHandle> Output;
+			TConstArrayView<FPrefix> Input;
+		};
+		struct FCopyPrefixesToRowHandleArrayViewTask
+		{
+			TArrayView<RowHandle> Output;
+			TConstArrayView<FPrefix> Input;
+		};
+		struct FComparativeSortRowHandleRangeTask
+		{
+			TArrayView<RowHandle> Rows;
+			TSharedPtr<const FColumnSorterInterface> Sorter;
+			const ICoreProvider& Storage;
+		};
+		struct FComparativeSortRowHandlePrefixRangeTask
+		{
+			TArrayView<FPrefix> Rows;
+			TSharedPtr<const FColumnSorterInterface> Sorter;
+			const ICoreProvider& Storage;
+		};
+		struct FMergeSortedAdjacentRangesTask
+		{
+			TArrayView<RowHandle> Array;
+			TArrayView<RowHandle> Scratch;
+			TSharedPtr<const FColumnSorterInterface> Sorter;
+			const ICoreProvider& Storage;
+			int32 SplitPosition;
+		};
+		struct FBucketizeRangesTask
+		{
+			TArrayView<RowHandle> Array;
+			TArrayView<FPrefix> Prefixes;
+			TArrayView<FPrefix> PrefixesShadow;
+			TSharedPtr<const FColumnSorterInterface> Sorter;
+			const ICoreProvider& Storage;
+			int32 PrefixBatchSize;
+			int32 SortBatchSize;
+			uint32 Index;
+		};
+		using TaskVariant = TVariant<
+			FTaskBarrier, 
+			FCollectPrefixesTask,
+			FRefreshPrefixesTask,
+			FRadixSortPrefixesTask,
+			FCopyPrefixesToRowHandleArrayTask,
+			FCopyPrefixesToRowHandleArrayViewTask,
+			FComparativeSortRowHandleRangeTask,
+			FComparativeSortRowHandlePrefixRangeTask,
+			FMergeSortedAdjacentRangesTask,
+			FBucketizeRangesTask
+		>;
+
+		struct FPrefixCollectionBatch
+		{
+			TaskVariant Task;
+			bool bCompleted = false;
+		};
+		
+		struct FTaskVisitor
+		{
+			void operator()(const FTaskBarrier& Task);
+			void operator()(const FCollectPrefixesTask& Task);
+			void operator()(const FRefreshPrefixesTask& Task);
+			void operator()(const FRadixSortPrefixesTask& Task);
+			void operator()(const FCopyPrefixesToRowHandleArrayTask& Task);
+			void operator()(const FCopyPrefixesToRowHandleArrayViewTask& Task);
+			void operator()(const FComparativeSortRowHandleRangeTask& Task);
+			void operator()(const FComparativeSortRowHandlePrefixRangeTask& Task);
+			void operator()(const FMergeSortedAdjacentRangesTask& Task);
+			void operator()(const FBucketizeRangesTask& Task);
+
+			FSortTasks& Owner;
+		};
+		
+		using TaskChain = TArray<FPrefixCollectionBatch>;
+		using TaskReferenceChain = TArray<FPrefixCollectionBatch*>;
+
+		void AddBarrier(TaskChain& Tasks);
+		void CollectPrefixes(TaskChain& Tasks, TArrayView<FPrefix> Output, TConstArrayView<RowHandle> Input,
+			const ICoreProvider& Storage, TSharedPtr<const FColumnSorterInterface> Sorter, int32 BatchSize);
+		void RefreshPrefixes(TaskChain& Tasks, TArrayView<FPrefix> Prefixes, const ICoreProvider& Storage,
+			TSharedPtr<const FColumnSorterInterface> Sorter, uint32 ByteIndex, int32 BatchSize);
+		void RadixSortPrefixes(TaskChain& Tasks, TArrayView<FPrefix> Output);
+		void CopyPrefixesToRowHandleArray(TaskChain& Tasks, TArrayView<RowHandle> Output, TConstArrayView<FPrefix> Input);
+		void CopyPrefixesToRowHandleArrayView(TaskChain& Tasks, TArrayView<RowHandle> Output, TConstArrayView<FPrefix> Input);
+		void ComparativeSortRowHandleRanges(TaskChain& Tasks, TArrayView<RowHandle> Rows, const ICoreProvider& Storage,
+			TSharedPtr<const FColumnSorterInterface> Sorter, int32 BatchSize);
+		void ComparativeSortRowHandlePrefixRanges(TaskChain& Tasks, TArrayView<FPrefix> Rows, const ICoreProvider& Storage,
+			TSharedPtr<const FColumnSorterInterface> Sorter);
+		void MergeSortedAdjacentRanges(TaskChain& Tasks, TArrayView<RowHandle> Array, TArrayView<RowHandle> Scratch,
+			const ICoreProvider& Storage, TSharedPtr<const FColumnSorterInterface> Sorter, int32 SplitPosition);
+		void BucketizeRange(TaskChain& Tasks, TArrayView<RowHandle> Array, TArrayView<FPrefix> Prefixes, TArrayView<FPrefix> PrefixesShadow,
+			TSharedPtr<const FColumnSorterInterface> Sorter, const ICoreProvider& Storage, int32 PrefixBatchSize, int32 SortBatchSize, 
+			uint32 ByteIndex);
+
+		void BuildTaskChain(TaskReferenceChain& Output);
+		void RemoveCompletedTasks();
+
+		TMpscQueue<TaskChain> PendingTasksQueue;
+		TArray<TaskChain> PendingTasks;
+		TaskChain MainTaskList;
+		TaskReferenceChain LocalTaskChain;
+
+		bool bIsUpdating = false;
+	};
+} // namespace UE::Editor::DataStorage::QueryStack::Sorting
\ No newline at end of file
diff --git a/Engine/Plugins/Experimental/EditorDataStorageFeatures/Source/TedsQueryStack/Private/TedsRowSortNode.cpp b/Engine/Plugins/Experimental/EditorDataStorageFeatures/Source/TedsQueryStack/Private/TedsRowSortNode.cpp
new file mode 100644
--- /dev/null
+++ b/Engine/Plugins/Experimental/EditorDataStorageFeatures/Source/TedsQueryStack/Private/TedsRowSortNode.cpp
@@ -0,0 +1,221 @@
+// Copyright Epic Games, Inc. All Rights Reserved.
+
+#include "TedsRowSortNode.h"
+
+#include "Algo/BinarySearch.h"
+#include "Algo/Sort.h"
+#include "Async/ParallelFor.h"
+#include "Misc/ScopeLock.h"
+#include "ProfilingDebugging/CpuProfilerTrace.h"
+#include "Sorting/SortTasks.h"
+
+namespace UE::Editor::DataStorage::QueryStack
+{
+	namespace Private
+	{
+		struct FSortContext
+		{
+			static constexpr int32 PrefixBatchSize = 2500;
+			static constexpr int32 SortBatchSize = 2500;
+
+			TArray<Sorting::FPrefix> PrefixArray;
+			TArray<Sorting::FPrefix> PrefixShadowArray;
+			TArray<RowHandle> RowHandleArray;
+			Sorting::FSortTasks Tasks;
+		};
+	} // namespace Private
+
+	FRowSortNode::FRowSortNode(ICoreProvider& InStorage, TSharedPtr<IRowNode> InParent, 
+		FTimespan InMaxFrameDuration)
+		: Parent(MoveTemp(InParent))
+		, MaxFrameDuration(InMaxFrameDuration)
+		, Storage(InStorage)
+	{
+		FRowHandleArray* ParentRows = Parent->GetMutableRows();
+		checkf(ParentRows != nullptr,
+			TEXT("FRowHandleSortNode requires a parent node that provides write access to its row handle array."));
+	}
+
+	FRowSortNode::FRowSortNode(ICoreProvider& InStorage, TSharedPtr<IRowNode> InParent, 
+		TSharedPtr<const FColumnSorterInterface> InColumnSorter, FTimespan InMaxFrameDuration)
+		: ColumnSorter(MoveTemp(InColumnSorter))
+		, Parent(MoveTemp(InParent))
+		, MaxFrameDuration(InMaxFrameDuration)
+		, Storage(InStorage)
+	{
+		FRowHandleArray* ParentRows = Parent->GetMutableRows();
+		checkf(ParentRows != nullptr, 
+			TEXT("FRowHandleSortNode requires a parent node that provides write access to its row handle array."));
+	}
+
+	void FRowSortNode::SetColumnSorter(TSharedPtr<const FColumnSorterInterface> InColumnSorter)
+	{
+		ColumnSorter = MoveTemp(InColumnSorter);
+
+		// Restart sort on next update.
+		LastUpdatedRevision = TNumericLimits<RevisionId>::Max();
+	}
+
+	bool FRowSortNode::IsSorting() const
+	{
+		return SortContext.IsValid();
+	}
+
+	INode::RevisionId FRowSortNode::GetRevision() const
+	{
+		return Revision;
+	}
+
+	void FRowSortNode::Update()
+	{
+		TRACE_CPUPROFILER_EVENT_SCOPE("[TEDS Sorting] Sort node update");
+
+		uint64 StartTime = FPlatformTime::Cycles64();
+
+		Parent->Update();
+		if (Parent->GetRevision() != LastUpdatedRevision)
+		{
+			LastUpdatedRevision = Parent->GetRevision();
+
+			if (ColumnSorter)
+			{
+				// If there's already a sort in progress, this will effectively cancel it, clean it up and start a new one.
+				SortContext = MakePimpl<Private::FSortContext>();
+
+				switch (ColumnSorter->GetSortType())
+				{
+				case FColumnSorterInterface::ESortType::FixedSize64:
+					SetupFixedSize64();
+					break;
+				case FColumnSorterInterface::ESortType::FixedSizeOnly:
+					SetupFixedSizeOnly();
+					break;
+				case FColumnSorterInterface::ESortType::ComparativeSort:
+					SetupComparativeSort();
+					break;
+				case FColumnSorterInterface::ESortType::HybridSort:
+					SetupHybridSort();
+					break;
+				default:
+					checkf(false, TEXT("Unsupported sort type: %i"), static_cast<int32>(ColumnSorter->GetSortType()));
+					break;
+				}
+			}
+			else
+			{
+				Revision++;
+			}
+		}
+
+		if (SortContext)
+		{
+			if (SortContext->Tasks.HasRemainingTasks())
+			{
+				SortContext->Tasks.Update(MaxFrameDuration);
+			}
+
+			// Check remaining tasks again as the update above could have added new tasks.
+			if (!SortContext->Tasks.HasRemainingTasks())
+			{
+				// No work left delete the sort context and stop further processing.
+				Revision++;
+				SortContext.Reset();
+			}
+		}
+	}
+
+	void FRowSortNode::SetupFixedSize64()
+	{
+		if (FRowHandleArray* MutableRows = Parent->GetMutableRows())
+		{
+			FRowHandleArrayView Rows = Parent->GetRows();
+			TArrayView<RowHandle> RowView = MutableRows->GetMutableRows(MutableRows->IsUnique()
+				? FRowHandleArray::EFlags::IsUnique
+				: FRowHandleArray::EFlags::None);
+			SortContext->PrefixArray.AddUninitialized(Rows.Num());
+			SortContext->Tasks.CollectPrefixes(SortContext->PrefixArray, Rows, Storage, ColumnSorter, Private::FSortContext::PrefixBatchSize);
+			SortContext->Tasks.AddBarrier();
+
+			SortContext->Tasks.RadixSortPrefixes(SortContext->PrefixArray);
+			SortContext->Tasks.AddBarrier();
+
+			SortContext->Tasks.CopyPrefixesToRowHandleArray(RowView, SortContext->PrefixArray);
+		}
+		else
+		{
+			checkf(false, TEXT("A 64-bit fixed size sort was attempted on a parent node that can't provide mutable rows."));
+		}
+	}
+
+	void FRowSortNode::SetupFixedSizeOnly()
+	{
+		if (FRowHandleArray* MutableRows = Parent->GetMutableRows())
+		{
+			FRowHandleArrayView Rows = Parent->GetRows();
+			SortContext->PrefixArray.AddUninitialized(Rows.Num());
+			SortContext->Tasks.CollectPrefixes(
+				SortContext->PrefixArray, Rows, Storage, ColumnSorter, Private::FSortContext::PrefixBatchSize);
+			SortContext->Tasks.AddBarrier();
+
+			SortContext->PrefixShadowArray.AddUninitialized(Rows.Num());
+			FRowHandleArray::EFlags Flags = MutableRows->IsUnique() ? FRowHandleArray::EFlags::IsUnique : FRowHandleArray::EFlags::None;
+			SortContext->Tasks.BucketizeRange(MutableRows->GetMutableRows(Flags), SortContext->PrefixArray,
+				SortContext->PrefixShadowArray, ColumnSorter, Storage, Private::FSortContext::PrefixBatchSize,
+				0 /* Never fall down to a comparative sort. The assumption is that this is called when all prefixes are the same limited size, though this is not mandatory. */);
+		}
+		else
+		{
+			checkf(false, TEXT("A fixed size sort was attempted on a parent node that can't provide mutable rows."));
+		}
+	}
+
+	void FRowSortNode::SetupComparativeSort()
+	{
+		if (FRowHandleArray* MutableRows = Parent->GetMutableRows())
+		{
+			FRowHandleArray::EFlags Flags = MutableRows->IsUnique() ? FRowHandleArray::EFlags::IsUnique : FRowHandleArray::EFlags::None;
+			SortContext->Tasks.ComparativeSortRowHandleRanges(MutableRows->GetMutableRows(Flags), Storage, ColumnSorter, Private::FSortContext::SortBatchSize);
+			SortContext->Tasks.AddBarrier();
+
+			SortContext->RowHandleArray.AddUninitialized(MutableRows->Num() / 2); // Used only for the right portion of the merge for temporary storage.
+			SortContext->Tasks.MergeSortedAdjacentRanges(
+				MutableRows->GetMutableRows(Flags), SortContext->RowHandleArray, Storage, ColumnSorter, Private::FSortContext::SortBatchSize);
+		}
+		else
+		{
+			checkf(false, TEXT("A 64-bit fixed size sort was attempted on a parent node that can't provide mutable rows."));
+		}
+	}
+
+	void FRowSortNode::SetupHybridSort()
+	{
+		if (FRowHandleArray* MutableRows = Parent->GetMutableRows())
+		{
+			FRowHandleArrayView Rows = Parent->GetRows();
+			SortContext->PrefixArray.AddUninitialized(Rows.Num());
+			SortContext->Tasks.CollectPrefixes(
+				SortContext->PrefixArray, Rows, Storage, ColumnSorter, Private::FSortContext::PrefixBatchSize);
+			SortContext->Tasks.AddBarrier();
+
+			SortContext->PrefixShadowArray.AddUninitialized(Rows.Num());
+			FRowHandleArray::EFlags Flags = MutableRows->IsUnique() ? FRowHandleArray::EFlags::IsUnique : FRowHandleArray::EFlags::None;
+			SortContext->Tasks.BucketizeRange(MutableRows->GetMutableRows(Flags), SortContext->PrefixArray, 
+				SortContext->PrefixShadowArray, ColumnSorter, Storage, Private::FSortContext::PrefixBatchSize, 
+				Private::FSortContext::SortBatchSize);
+		}
+		else
+		{
+			checkf(false, TEXT("A hybrid sort was attempted on a parent node that can't provide mutable rows."));
+		}
+	}
+
+	FRowHandleArrayView FRowSortNode::GetRows() const
+	{
+		return Parent->GetRows();
+	}
+
+	FRowHandleArray* FRowSortNode::GetMutableRows()
+	{
+		return Parent->GetMutableRows();
+	}
+} // namespace UE::Editor::DataStorage::QueryStack
diff --git a/Engine/Plugins/Experimental/EditorDataStorageFeatures/Source/TedsQueryStack/Public/TedsRowSortNode.h b/Engine/Plugins/Experimental/EditorDataStorageFeatures/Source/TedsQueryStack/Public/TedsRowSortNode.h
new file mode 100644
--- /dev/null
+++ b/Engine/Plugins/Experimental/EditorDataStorageFeatures/Source/TedsQueryStack/Public/TedsRowSortNode.h
@@ -0,0 +1,67 @@
+// Copyright Epic Games, Inc. All Rights Reserved.
+
+#pragma once
+
+#include "Async/Mutex.h"
+#include "Elements/Interfaces/TypedElementDataStorageUiInterface.h"
+#include "Math/NumericLimits.h"
+#include "Misc/EnumClassFlags.h"
+#include "Misc/Timespan.h"
+#include "TedsQueryStackInterfaces.h"
+#include "Templates/FunctionFwd.h"
+#include "Templates/PimplPtr.h"
+
+namespace UE::Editor::DataStorage
+{
+	class ICoreProvider;
+}
+
+namespace UE::Editor::DataStorage::QueryStack
+{
+	namespace Private
+	{
+		// Information used only when sorting is active.
+		struct FSortContext;
+	}
+
+	/**
+	 * Sorts the rows in the parent node using a provided column sorter. Column sorters can be found on on TEDS UI widget constructors,
+	 * or custom sorters can be created.
+	 */
+	class FRowSortNode : public IRowNode
+	{
+	public:
+		TEDSQUERYSTACK_API FRowSortNode(ICoreProvider& InStorage, TSharedPtr<IRowNode> InParent, 
+			FTimespan InMaxFrameDuration = FTimespan::FromMilliseconds(10));
+		TEDSQUERYSTACK_API FRowSortNode(ICoreProvider& InStorge, TSharedPtr<IRowNode> InParent, 
+			TSharedPtr<const FColumnSorterInterface> InColumnSorter, FTimespan InMaxFrameDuration = FTimespan::FromMilliseconds(10));
+		virtual ~FRowSortNode() override = default;
+
+		TEDSQUERYSTACK_API void SetColumnSorter(TSharedPtr<const FColumnSorterInterface> InColumnSorter);
+		TEDSQUERYSTACK_API bool IsSorting() const;
+
+		TEDSQUERYSTACK_API virtual RevisionId GetRevision() const override;
+		TEDSQUERYSTACK_API virtual void Update() override;
+		TEDSQUERYSTACK_API virtual FRowHandleArrayView GetRows() const override;
+		TEDSQUERYSTACK_API virtual FRowHandleArray* GetMutableRows() override;
+
+	private:
+		void SetupFixedSize64();
+		void SetupFixedSizeOnly();
+		void SetupComparativeSort();
+		void SetupHybridSort();
+		
+		TSharedPtr<const FColumnSorterInterface> ColumnSorter;
+		TSharedPtr<IRowNode> Parent;
+		
+		TPimplPtr<Private::FSortContext> SortContext;
+
+		FTimespan MaxFrameDuration;
+		RevisionId LastUpdatedRevision = TNumericLimits<RevisionId>::Max();
+		RevisionId Revision = 0;
+		
+		ICoreProvider& Storage;
+		
+		bool bIsParentUnique = false;
+	};
+} // namespace UE::Editor::DataStorage::QueryStack
diff --git a/Engine/Plugins/Experimental/EditorDataStorageFeatures/Source/TedsQueryStack/Tests/TedsRowSortNodeTests.cpp b/Engine/Plugins/Experimental/EditorDataStorageFeatures/Source/TedsQueryStack/Tests/TedsRowSortNodeTests.cpp
new file mode 100644
--- /dev/null
+++ b/Engine/Plugins/Experimental/EditorDataStorageFeatures/Source/TedsQueryStack/Tests/TedsRowSortNodeTests.cpp
@@ -0,0 +1,349 @@
+// Copyright Epic Games, Inc. All Rights Reserved.
+
+#if WITH_TESTS
+
+#include "Elements/Common/EditorDataStorageFeatures.h"
+#include "Elements/Framework/TypedElementSorter.h"
+#include "Elements/Interfaces/TypedElementDataStorageInterface.h"
+#include "Math/UnrealMathUtility.h"
+#include "Sorting/SortTasks.h"
+#include "Tests/TestHarnessAdapter.h"
+#include "Templates/SharedPointer.h"
+
+namespace UE::Editor::DataStorage::Testing
+{
+	class FSimpleColumnSorter : public FColumnSorterInterface
+	{
+	public:
+		virtual ESortType GetSortType() const override
+		{
+			return ESortType::HybridSort;
+		}
+
+		virtual FText GetShortName() const override
+		{
+			return FText::GetEmpty();
+		}
+
+		virtual int32 Compare(const ICoreProvider& Storage, RowHandle Left, RowHandle Right) const override
+		{
+			return Left < Right ? -1 : (Left > Right ? 1 : 0);
+
+		}
+
+		virtual FPrefixInfo CalculatePrefix(const ICoreProvider& Storage, RowHandle Row, uint32 ByteIndex) const override
+		{
+			return FPrefixInfo
+			{
+				.Prefix = Row,
+				.bHasRemainingBytes = false
+			};
+		}
+	};
+
+	class FPrefixExpandedColumnSorter : public FSimpleColumnSorter
+	{
+	public:
+		virtual FPrefixInfo CalculatePrefix(const ICoreProvider& Storage, RowHandle Row, uint32 ByteIndex) const override
+		{
+			return FPrefixInfo
+			{
+				.Prefix = static_cast<uint8>(Row >> (56 - ByteIndex)),
+				.bHasRemainingBytes = (ByteIndex < 56)
+			};
+		}
+	};
+
+	class FPrefixVariableLengthColumnSorter : public FSimpleColumnSorter
+	{
+	public:
+		virtual FPrefixInfo CalculatePrefix(const ICoreProvider& Storage, RowHandle Row, uint32 ByteIndex) const override
+		{
+			bool bHasRemainingBytes = ByteIndex < static_cast<uint8>(Row & 255);
+			return FPrefixInfo
+			{
+				.Prefix = bHasRemainingBytes ? (Row >> 8) : 0,
+				.bHasRemainingBytes = bHasRemainingBytes
+			};
+		}
+	};
+
+	class FStringCompareColumnSorter : public FSimpleColumnSorter
+	{
+	public:
+		TArray<FString> Strings;
+
+		virtual int32 Compare(const ICoreProvider& Storage, RowHandle Left, RowHandle Right) const override
+		{
+			return Strings[static_cast<int32>(Left)].Compare(Strings[static_cast<int32>(Right)]);
+		}
+	};
+
+	void RunTasks(UE::Editor::DataStorage::QueryStack::Sorting::FSortTasks& Tasks, bool bIterative)
+	{
+		if (bIterative)
+		{
+			uint64 StartTime = FPlatformTime::Cycles64();
+			while (FTimespan(FPlatformTime::Cycles64() - StartTime) < FTimespan::FromSeconds(5) && Tasks.HasRemainingTasks())
+			{
+				Tasks.Update(FTimespan::FromMilliseconds(5));
+			}
+		}
+		else
+		{
+			Tasks.Update(FTimespan::FromSeconds(5));
+		}
+	}
+
+	bool CompareRows(TConstArrayView<RowHandle> Source, TConstArrayView<RowHandle> Compare)
+	{
+		int32 Count = Source.Num();
+		if (Count == Compare.Num())
+		{
+			const RowHandle* SourceIt = Source.GetData();
+			const RowHandle* CompareIt = Compare.GetData();
+			for (int32 Index = 0; Index < Count; ++Index)
+			{
+				if (*SourceIt++ != *CompareIt++)
+				{
+					return false;
+				}
+			}
+			return true;
+		}
+		return false;
+	}
+}
+
+TEST_CASE_NAMED(TEDS_QueryStack_RowSortNode_SortTasks_Tests, "Editor::DataStorage::QueryStack::FRowSortNode::FSortTasks", "[ApplicationContextMask][EngineFilter]")
+{
+	using namespace UE::Editor::DataStorage;
+	using namespace UE::Editor::DataStorage::Testing;
+	using namespace UE::Editor::DataStorage::QueryStack::Sorting;
+	
+	TSharedPtr<FSimpleColumnSorter> SimpleSorter = MakeShared<FSimpleColumnSorter>();
+	
+	if (const ICoreProvider* Storage = GetDataStorageFeature<ICoreProvider>(StorageFeatureName))
+	{
+		SECTION("Radix sort")
+		{
+			TArray<RowHandle> Array = { 1, 3, 5, 2, 4, 6 };
+			
+			TArray<FPrefix> Prefixes;
+			Prefixes.AddUninitialized(Array.Num());
+			
+			FSortTasks Tasks;
+			Tasks.CollectPrefixes(Prefixes, Array, *Storage, SimpleSorter, 6);
+			Tasks.AddBarrier();
+			Tasks.RadixSortPrefixes(Prefixes);
+			Tasks.AddBarrier();
+			Tasks.CopyPrefixesToRowHandleArray(Array, Prefixes);
+			RunTasks(Tasks, false);
+
+			CHECK_MESSAGE(TEXT("Sorted array didn't match expectation."), CompareRows(Array, { 1, 2, 3, 4, 5, 6 }));
+		}
+
+		SECTION("Comparative sort")
+		{
+			TSharedPtr<FStringCompareColumnSorter> StringSorter = MakeShared<FStringCompareColumnSorter>();
+			StringSorter->Strings.Emplace(TEXT("AAB"));
+			StringSorter->Strings.Emplace(TEXT("ABB"));
+			StringSorter->Strings.Emplace(TEXT("CBBEE"));
+			StringSorter->Strings.Emplace(TEXT("D"));
+			StringSorter->Strings.Emplace(TEXT("BDE"));
+			StringSorter->Strings.Emplace(TEXT("ABBA"));
+			TArray<RowHandle> Array = { 0, 1, 2, 3, 4, 5 };
+
+			FSortTasks Tasks;
+			Tasks.ComparativeSortRowHandleRanges(Array, *Storage, StringSorter, 6);
+			RunTasks(Tasks, false);
+
+			CHECK_MESSAGE(TEXT("Sorted array didn't match expectation."), CompareRows(Array, { 0, 1, 5, 4, 2, 3 }));
+		}
+
+		SECTION("Merge sort - Equal range")
+		{
+			TArray<RowHandle> Array = { 1, 3, 5, 2, 4, 6 };
+			TArray<RowHandle> Scratch;
+			Scratch.AddUninitialized(6);
+
+			FSortTasks Tasks;
+			Tasks.MergeSortedAdjacentRanges(Array, Scratch, *Storage, SimpleSorter, 3);
+			RunTasks(Tasks, false);
+
+			CHECK_MESSAGE(TEXT("Sorted array didn't match expectation."), CompareRows(Array, { 1, 2, 3, 4, 5, 6 }));
+		}
+
+		SECTION("Merge sort - Right fits in single left slot")
+		{
+			TArray<RowHandle> Array = { 1, 7, 9, 2, 4, 6 };
+			TArray<RowHandle> Scratch;
+			Scratch.AddUninitialized(6);
+
+			FSortTasks Tasks;
+			Tasks.MergeSortedAdjacentRanges(Array, Scratch, *Storage, SimpleSorter, 3);
+			RunTasks(Tasks, false);
+
+			CHECK_MESSAGE(TEXT("Sorted array didn't match expectation."), CompareRows(Array, { 1, 2, 4, 6, 7, 9 }));
+		}
+
+		SECTION("Merge sort - Uneven split")
+		{
+			TArray<RowHandle> Array = { 1, 3, 5, 7, 2, 4};
+			TArray<RowHandle> Scratch;
+			Scratch.AddUninitialized(6);
+
+			FSortTasks Tasks;
+			Tasks.MergeSortedAdjacentRanges(Array, Scratch, *Storage, SimpleSorter, 4);
+			RunTasks(Tasks, false);
+
+			CHECK_MESSAGE(TEXT("Sorted array didn't match expectation."), CompareRows(Array, { 1, 2, 3, 4, 5, 7 }));
+		}
+
+		SECTION("Bucketize Range - Single prefix")
+		{
+			TArray<RowHandle> Array = 
+			{ 
+				0x0101, 
+				0x0203, 
+				0x0105, 
+				0x0107, 
+				0x0202, 
+				0x0104
+			};
+			TArray<FPrefix> Prefixes;
+			TArray<FPrefix> PrefixesShadow;
+			Prefixes.AddUninitialized(Array.Num());
+			PrefixesShadow.AddUninitialized(Array.Num());
+
+			FSortTasks Tasks;
+			Tasks.CollectPrefixes(Prefixes, Array, *Storage, SimpleSorter, 6);
+			Tasks.AddBarrier();
+			Tasks.BucketizeRange(Array, Prefixes, PrefixesShadow, SimpleSorter, *Storage, 6, 0);
+			RunTasks(Tasks, true);
+
+			CHECK_MESSAGE(TEXT("Sorted array didn't match expectation."), CompareRows(Array, 
+				{ 
+					0x0101,
+					0x0104,
+					0x0105,
+					0x0107,
+					0x0202,
+					0x0203
+				}));
+		}
+
+		SECTION("Bucketize Range - Multiple prefixes")
+		{
+			TArray<RowHandle> Array =
+			{
+				0x0101,
+				0x0203,
+				0x0105,
+				0x0107,
+				0x0202,
+				0x0104
+			};
+			TArray<FPrefix> Prefixes;
+			TArray<FPrefix> PrefixesShadow;
+			Prefixes.AddUninitialized(Array.Num());
+			PrefixesShadow.AddUninitialized(Array.Num());
+
+			TSharedPtr<FPrefixExpandedColumnSorter> PrefixExpandedSorter = MakeShared<FPrefixExpandedColumnSorter>();
+			FSortTasks Tasks;
+			Tasks.CollectPrefixes(Prefixes, Array, *Storage, PrefixExpandedSorter, 6);
+			Tasks.AddBarrier();
+			Tasks.BucketizeRange(Array, Prefixes, PrefixesShadow, PrefixExpandedSorter, *Storage, 6, 0);
+			RunTasks(Tasks, true);
+
+			CHECK_MESSAGE(TEXT("Sorted array didn't match expectation."), CompareRows(Array,
+				{
+					0x0101,
+					0x0104,
+					0x0105,
+					0x0107,
+					0x0202,
+					0x0203
+				}));
+		}
+
+		SECTION("Bucketize Range - Variable prefix length")
+		{
+			TArray<RowHandle> Array =
+			{
+				0x1108,
+				0x2218,
+				0x1118,
+				0x1128,
+				0x2210,
+				0x1120
+			};
+			TArray<FPrefix> Prefixes;
+			TArray<FPrefix> PrefixesShadow;
+			Prefixes.AddUninitialized(Array.Num());
+			PrefixesShadow.AddUninitialized(Array.Num());
+
+			TSharedPtr<FPrefixVariableLengthColumnSorter> PrefixExpandedSorter = MakeShared<FPrefixVariableLengthColumnSorter>();
+			FSortTasks Tasks;
+			Tasks.CollectPrefixes(Prefixes, Array, *Storage, PrefixExpandedSorter, 6);
+			Tasks.AddBarrier();
+			Tasks.BucketizeRange(Array, Prefixes, PrefixesShadow, PrefixExpandedSorter, *Storage, 6, 0);
+			RunTasks(Tasks, true);
+
+			CHECK_MESSAGE(TEXT("Sorted array didn't match expectation."), CompareRows(Array,
+				{
+					0x1108,
+					0x1118,
+					0x1120,
+					0x1128,
+					0x2210,
+					0x2218
+				}));
+		}
+
+		auto BucketizeRangeLargeNumber = [this, Storage](bool bWithComparisons)
+			{
+				constexpr int32 Count = 25000;
+
+				TArray<RowHandle> Array;
+				Array.Reserve(Count);
+
+				for (int32 Random = 0; Random < Count; ++Random)
+				{
+					Array.Add(FMath::Rand32());
+				}
+
+				TArray<RowHandle> Compare = Array;
+				Algo::Sort(Compare);
+
+				TArray<FPrefix> Prefixes;
+				TArray<FPrefix> PrefixesShadow;
+				Prefixes.AddUninitialized(Array.Num());
+				PrefixesShadow.AddUninitialized(Array.Num());
+
+				TSharedPtr<FPrefixExpandedColumnSorter> PrefixExpandedSorter = MakeShared<FPrefixExpandedColumnSorter>();
+				FSortTasks Tasks;
+				Tasks.CollectPrefixes(Prefixes, Array, *Storage, PrefixExpandedSorter, 50);
+				Tasks.AddBarrier();
+				Tasks.BucketizeRange(Array, Prefixes, PrefixesShadow, PrefixExpandedSorter, *Storage, 50, bWithComparisons ? 50 : 0);
+				RunTasks(Tasks, true);
+
+				for (int32 Index = 0; Index < Count; ++Index)
+				{
+					CHECK(Array[Index] == Compare[Index]);
+				}
+			};
+
+		SECTION("Bucketize Range - Large number, no comparisons")
+		{
+			BucketizeRangeLargeNumber(false);
+		}
+
+		SECTION("Bucketize Range - Large number, with comparisons")
+		{
+			BucketizeRangeLargeNumber(true);
+		}
+	}
+}
+
+#endif // #if WITH_TESTS
