[SceneOutliner] WidgetEnterEditModeColumn is now added to WidgetRow instead of TargetRow, so that it can be used by multiple Widgets.

diff --git a/Engine/Plugins/Experimental/EditorDataStorageFeatures/Source/TedsOutliner/Private/TedsOutlinerImpl.cpp b/Engine/Plugins/Experimental/EditorDataStorageFeatures/Source/TedsOutliner/Private/TedsOutlinerImpl.cpp
--- a/Engine/Plugins/Experimental/EditorDataStorageFeatures/Source/TedsOutliner/Private/TedsOutlinerImpl.cpp
+++ b/Engine/Plugins/Experimental/EditorDataStorageFeatures/Source/TedsOutliner/Private/TedsOutlinerImpl.cpp
@@ -81,6 +81,84 @@ namespace QueryUtils
 }
 
 
+namespace LabelWidgetUtils
+{
+	static TUniquePtr<FTypedElementWidgetConstructor> CreateLabelWidgetConstructorInternal(
+			const DataStorage::ICoreProvider& Storage,
+			DataStorage::IUiProvider& StorageUi,
+			const DataStorage::IUiProvider::FPurposeID& Purpose,
+			const DataStorage::RowHandle TargetRow)
+	{
+		// Get all the columns on the given row
+		TArray<TWeakObjectPtr<const UScriptStruct>> ColumnTypes;
+		Storage.ListColumns(TargetRow, [&ColumnTypes](const UScriptStruct& ColumnType)
+		{
+			ColumnTypes.Add(&ColumnType);
+		});
+	
+		TUniquePtr<FTypedElementWidgetConstructor> WidgetConstructor;
+		StorageUi.CreateWidgetConstructors(StorageUi.FindPurpose(Purpose), DataStorage::IUiProvider::EMatchApproach::LongestMatch, ColumnTypes, {},
+			[&WidgetConstructor](TUniquePtr<FTypedElementWidgetConstructor> CreatedConstructor, TConstArrayView<TWeakObjectPtr<const UScriptStruct>>)
+			{
+				WidgetConstructor = MoveTemp(CreatedConstructor);
+				// Either this was the exact match so no need to search further or the longest possible chain didn't match so the next ones will 
+				// always be shorter in both cases just return.
+				return false;
+			});
+		return WidgetConstructor;
+	}
+	
+	TSharedPtr<SWidget> CreateLabelWidgetInternal(
+		FTypedElementWidgetConstructor& WidgetConstructor,
+		DataStorage::ICoreProvider& Storage,
+		DataStorage::IUiProvider& StorageUi,
+		const DataStorage::RowHandle TargetRow,
+		ISceneOutlinerTreeItem& TreeItem,
+		const STableRow<FSceneOutlinerTreeItemPtr>& RowItem,
+		const bool bInteractable)
+	{
+		using namespace DataStorage::Queries;
+		// Query description to pass as metadata to allow the label column to be writable
+		static const FQueryDescription MetaDataQueryReadWrite = Select().ReadWrite<FTypedElementLabelColumn>().Where().Compile();
+		static const FQueryDescription MetaDataQueryRead = Select().ReadOnly<FTypedElementLabelColumn>().Where().Compile();
+	
+		RowHandle WidgetRow = Storage.AddRow(Storage.FindTable(TableViewerUtils::GetWidgetTableName()));
+		checkf(WidgetRow != InvalidRowHandle, TEXT("Expected valid row handle!"));
+		if (FTypedElementRowReferenceColumn* RowReference = Storage.GetColumn<FTypedElementRowReferenceColumn>(WidgetRow))
+		{
+			RowReference->Row = TargetRow;
+		}
+		Storage.AddColumn(WidgetRow, FTedsOutlinerColumn{ .Outliner = TreeItem.WeakSceneOutliner });
+	
+		// Create metadata for the query
+		FQueryMetaDataView QueryMetaDataView = bInteractable ? FQueryMetaDataView(MetaDataQueryReadWrite) : FQueryMetaDataView(MetaDataQueryRead);
+		TSharedPtr<SWidget> Widget = StorageUi.ConstructWidget(WidgetRow, WidgetConstructor, QueryMetaDataView);
+		if (Widget)
+		{
+			if (FExternalWidgetSelectionColumn* ExternalWidgetSelectionColumn = Storage.GetColumn<FExternalWidgetSelectionColumn>(WidgetRow))
+			{
+				ExternalWidgetSelectionColumn->IsSelected = FIsSelected::CreateSP(&RowItem, &STableRow<FSceneOutlinerTreeItemPtr>::IsSelectedExclusively);
+			}
+
+			if (bInteractable)
+			{
+				TreeItem.RenameRequestEvent.BindLambda([StoragePtr = &Storage, WidgetRow]()
+					{
+						if (const FWidgetEnterEditModeColumn* Column = StoragePtr->GetColumn<FWidgetEnterEditModeColumn>(WidgetRow))
+						{
+							Column->OnEnterEditMode.ExecuteIfBound();
+						}
+					});
+			}
+		}
+		else
+		{
+			Storage.RemoveRow(WidgetRow);
+		}	
+		return Widget;
+	}
+}
+
 FTedsOutlinerParams::FTedsOutlinerParams(SSceneOutliner* InSceneOutliner)
 	: SceneOutliner(InSceneOutliner)
 	, QueryDescription()
@@ -220,93 +298,39 @@ void FTedsOutlinerImpl::SetSelection(const TArray<DataStorage::RowHandle>& InSel
 		Storage->AddColumn(Row, FTypedElementSelectionColumn{ .SelectionSet = SelectionSetName.GetValue() });
 	}
 }
-
+	
 TSharedRef<SWidget> FTedsOutlinerImpl::CreateLabelWidget(
 	DataStorage::ICoreProvider& Storage,
 	DataStorage::IUiProvider& StorageUi,
-	ISceneOutliner& Outliner,
-	DataStorage::IUiProvider::FPurposeID LabelWidgetPurpose,
-	DataStorage::RowHandle Row,
+	const DataStorage::IUiProvider::FPurposeID& Purpose,
+	DataStorage::RowHandle TargetRow,
+	ISceneOutlinerTreeItem& TreeItem,
 	const STableRow<FSceneOutlinerTreeItemPtr>& RowItem,
-	bool bIsInteractable)
+	bool bInteractable)
 {
-	// Find the best matching label widget for a given list of columns
-	auto CreateWidgetConstructor = [&StorageUi, LabelWidgetPurpose](TArray<TWeakObjectPtr<const UScriptStruct>>& ColumnTypes) -> TSharedPtr<FTypedElementWidgetConstructor>
-	{
-		TSharedPtr<FTypedElementWidgetConstructor> OutWidgetConstructorPtr = nullptr;
-
-		StorageUi.CreateWidgetConstructors(StorageUi.FindPurpose(LabelWidgetPurpose), DataStorage::IUiProvider::EMatchApproach::LongestMatch, ColumnTypes, {},
-					[&OutWidgetConstructorPtr, ColumnTypes](
-					TUniquePtr<FTypedElementWidgetConstructor> CreatedConstructor, 
-					TConstArrayView<TWeakObjectPtr<const UScriptStruct>> MatchedColumnTypes)
-					{
-						OutWidgetConstructorPtr = TSharedPtr<FTypedElementWidgetConstructor>(CreatedConstructor.Release());
-						// Either this was the exact match so no need to search further or the longest possible chain didn't match so the next ones will 
-						// always be shorter in both cases just return.
-						return false;
-					});
-		
-		return OutWidgetConstructorPtr;
-	};
-	
-	using namespace DataStorage::Queries;
-	// Query description to pass as metadata to allow the label column to be writable
-	static FQueryDescription MetaDataQueryReadWrite = Select().ReadWrite<FTypedElementLabelColumn>().Where().Compile();
-	static FQueryDescription MetaDataQueryRead = Select().ReadOnly<FTypedElementLabelColumn>().Where().Compile();
-
-	// Create a widget for this row using the given widget constructor
-	auto CreateWidget = [&Storage, &StorageUi, &Outliner, Row, &RowItem, bIsInteractable](const TSharedPtr<FTypedElementWidgetConstructor>& WidgetConstructor) -> TSharedPtr<SWidget>
-	{
-		// Create metadata for the query
-		FQueryMetaDataView QueryMetaDataView = bIsInteractable ?
-			FQueryMetaDataView(MetaDataQueryReadWrite) : FQueryMetaDataView(MetaDataQueryRead);
-
-		RowHandle UiRowHandle = Storage.AddRow(Storage.FindTable(DataStorage::TableViewerUtils::GetWidgetTableName()));
-
-		if (FTypedElementRowReferenceColumn* RowReference = Storage.GetColumn<FTypedElementRowReferenceColumn>(UiRowHandle))
-		{
-			RowReference->Row = Row;
-		}
-
-		Storage.AddColumn(UiRowHandle, FTedsOutlinerColumn{.Outliner = StaticCastSharedRef<ISceneOutliner>(Outliner.AsShared())});
-		
-		TSharedPtr<SWidget> Widget = StorageUi.ConstructWidget(UiRowHandle, *WidgetConstructor, QueryMetaDataView);
-		
-		if (FExternalWidgetSelectionColumn* ExternalWidgetSelectionColumn = Storage.GetColumn<FExternalWidgetSelectionColumn>(UiRowHandle))
-		{
-			ExternalWidgetSelectionColumn->IsSelected = FIsSelected::CreateSP(&RowItem, &STableRow<FSceneOutlinerTreeItemPtr>::IsSelectedExclusively);
-		}
-		return Widget;
-		
-	};
-
-	// Get all the columns on the given row
-	TArray<TWeakObjectPtr<const UScriptStruct>> Columns;
-	Storage.ListColumns(Row, [&Columns](const UScriptStruct& ColumnType)
+	TUniquePtr<FTypedElementWidgetConstructor> WidgetConstructor = LabelWidgetUtils::CreateLabelWidgetConstructorInternal(Storage, StorageUi, Purpose, TargetRow);
+	if (!WidgetConstructor)
 	{
-		Columns.Add(&ColumnType);
-	});
+		return SNullWidget::NullWidget;
+	}
 
-	// Create the label widget
-	if(const TSharedPtr<FTypedElementWidgetConstructor> WidgetConstructor = CreateWidgetConstructor(Columns))
+	TSharedPtr<SWidget> Widget = LabelWidgetUtils::CreateLabelWidgetInternal(*WidgetConstructor, Storage, StorageUi, TargetRow, TreeItem, RowItem, bInteractable);
+	if (!Widget)
 	{
-		if(const TSharedPtr<SWidget> Widget = CreateWidget(WidgetConstructor))
-		{
-			return SNew(SBox)
-					.HAlign(HAlign_Left)
-					.VAlign(VAlign_Center)
-					[
-						Widget.ToSharedRef()
-					];
-		}
+		return SNullWidget::NullWidget;
 	}
-
-	return SNullWidget::NullWidget;
+	
+	return SNew(SBox)
+		.HAlign(HAlign_Left)
+		.VAlign(VAlign_Center)
+		[
+			Widget.ToSharedRef()
+		];
 }
 
-TSharedRef<SWidget> FTedsOutlinerImpl::CreateLabelWidgetForItem(DataStorage::RowHandle InRowHandle, const STableRow<FSceneOutlinerTreeItemPtr>& InRow, bool bIsInteractable) const
+TSharedRef<SWidget>FTedsOutlinerImpl::CreateLabelWidgetForItem(DataStorage::RowHandle TargetRow, ISceneOutlinerTreeItem& TreeItem, const STableRow<FSceneOutlinerTreeItemPtr>& RowItem, bool bIsInteractable) const
 {
-	return CreateLabelWidget(*Storage, *StorageUi, *SceneOutliner, LabelWidgetPurpose, InRowHandle, InRow, bIsInteractable);
+	return CreateLabelWidget(*Storage, *StorageUi, LabelWidgetPurpose, TargetRow, TreeItem, RowItem, bIsInteractable);
 }
 
 void FTedsOutlinerImpl::AppendQuery(DataStorage::FQueryDescription& Query1, const DataStorage::FQueryDescription& Query2)
diff --git a/Engine/Plugins/Experimental/EditorDataStorageFeatures/Source/TedsOutliner/Private/TedsOutlinerItem.cpp b/Engine/Plugins/Experimental/EditorDataStorageFeatures/Source/TedsOutliner/Private/TedsOutlinerItem.cpp
--- a/Engine/Plugins/Experimental/EditorDataStorageFeatures/Source/TedsOutliner/Private/TedsOutlinerItem.cpp
+++ b/Engine/Plugins/Experimental/EditorDataStorageFeatures/Source/TedsOutliner/Private/TedsOutlinerItem.cpp
@@ -68,7 +68,7 @@ bool FTedsOutlinerTreeItem::CanInteract() const
 TSharedRef<SWidget> FTedsOutlinerTreeItem::GenerateLabelWidget(ISceneOutliner& Outliner,
 	const STableRow<FSceneOutlinerTreeItemPtr>& InRow)
 {
-	return TedsOutlinerImpl.IsValid() ? TedsOutlinerImpl.Pin()->CreateLabelWidgetForItem(RowHandle, InRow, CanInteract()) : SNullWidget::NullWidget;
+	return TedsOutlinerImpl.IsValid() ? TedsOutlinerImpl.Pin()->CreateLabelWidgetForItem(RowHandle, *this, InRow, CanInteract()) : SNullWidget::NullWidget;
 }
 
 void FTedsOutlinerTreeItem::GenerateContextMenu(UToolMenu* Menu, SSceneOutliner& Outliner)
diff --git a/Engine/Plugins/Experimental/EditorDataStorageFeatures/Source/TedsOutliner/Private/Widgets/OutlinerFolderLabelWidget.cpp b/Engine/Plugins/Experimental/EditorDataStorageFeatures/Source/TedsOutliner/Private/Widgets/OutlinerFolderLabelWidget.cpp
--- a/Engine/Plugins/Experimental/EditorDataStorageFeatures/Source/TedsOutliner/Private/Widgets/OutlinerFolderLabelWidget.cpp
+++ b/Engine/Plugins/Experimental/EditorDataStorageFeatures/Source/TedsOutliner/Private/Widgets/OutlinerFolderLabelWidget.cpp
@@ -156,7 +156,7 @@ TSharedRef<SWidget> FOutlinerFolderLabelWidgetConstructor::CreateLabel(UE::Edito
 				DataStorage->RemoveColumn<FIsInEditingModeTag>(WidgetRow);
 			});
 
-		DataStorage->AddColumn<FWidgetEnterEditModeColumn>(TargetRow, FWidgetEnterEditModeColumn{
+		DataStorage->AddColumn<FWidgetEnterEditModeColumn>(WidgetRow, FWidgetEnterEditModeColumn{
 				.OnEnterEditMode = FSimpleDelegate::CreateLambda([TextBlock]()
 				{
 					TextBlock->EnterEditingMode();
diff --git a/Engine/Plugins/Experimental/EditorDataStorageFeatures/Source/TedsOutliner/Private/Widgets/OutlinerTextWidget.cpp b/Engine/Plugins/Experimental/EditorDataStorageFeatures/Source/TedsOutliner/Private/Widgets/OutlinerTextWidget.cpp
--- a/Engine/Plugins/Experimental/EditorDataStorageFeatures/Source/TedsOutliner/Private/Widgets/OutlinerTextWidget.cpp
+++ b/Engine/Plugins/Experimental/EditorDataStorageFeatures/Source/TedsOutliner/Private/Widgets/OutlinerTextWidget.cpp
@@ -127,7 +127,7 @@ TSharedPtr<SInlineEditableTextBlock> FOutlinerTextWidgetConstructor::CreateEdita
 		.OverflowPolicy(bShouldUseMiddleEllipsis ? ETextOverflowPolicy::MiddleEllipsis : TOptional<ETextOverflowPolicy>())
 		.IsSelected(WidgetRowBinder.BindEvent(&FExternalWidgetSelectionColumn::IsSelected));
 
-	DataStorage->AddColumn<FWidgetEnterEditModeColumn>(TargetRow, FWidgetEnterEditModeColumn{
+	DataStorage->AddColumn<FWidgetEnterEditModeColumn>(WidgetRow, FWidgetEnterEditModeColumn{
 		.OnEnterEditMode = FSimpleDelegate::CreateLambda([TextBlock]()
 		{
 			TextBlock->EnterEditingMode();
diff --git a/Engine/Plugins/Experimental/EditorDataStorageFeatures/Source/TedsOutliner/Public/TedsOutlinerImpl.h b/Engine/Plugins/Experimental/EditorDataStorageFeatures/Source/TedsOutliner/Public/TedsOutlinerImpl.h
--- a/Engine/Plugins/Experimental/EditorDataStorageFeatures/Source/TedsOutliner/Public/TedsOutlinerImpl.h
+++ b/Engine/Plugins/Experimental/EditorDataStorageFeatures/Source/TedsOutliner/Public/TedsOutlinerImpl.h
@@ -133,11 +133,11 @@ public:
 	static TSharedRef<SWidget> CreateLabelWidget(
 		DataStorage::ICoreProvider& Storage,
 		DataStorage::IUiProvider& StorageUi,
-		ISceneOutliner& Outliner,
-		DataStorage::IUiProvider::FPurposeID LabelWidgetPurpose,
-		DataStorage::RowHandle Row,
+		const DataStorage::IUiProvider::FPurposeID& Purpose,
+		DataStorage::RowHandle TargetRow,
+		ISceneOutlinerTreeItem& TreeItem,
 		const STableRow<FSceneOutlinerTreeItemPtr>& RowItem,
-		bool bIsInteractable = true);
+		bool bInteractable = true);
 
 	FTedsOutlinerImpl(const FTedsOutlinerParams& InParams, ISceneOutlinerMode* InMode);
 	virtual ~FTedsOutlinerImpl();
@@ -166,7 +166,7 @@ public:
 	void SetSelection(const TArray<DataStorage::RowHandle>& InSelectedRows);
 
 	// Helper function to create a label widget for a given row
-	TSharedRef<SWidget> CreateLabelWidgetForItem(DataStorage::RowHandle InRowHandle, const STableRow<FSceneOutlinerTreeItemPtr>& InRow, bool bIsInteractable = true) const;
+	TSharedRef<SWidget> CreateLabelWidgetForItem(DataStorage::RowHandle TargetRow, ISceneOutlinerTreeItem& TreeItem, const STableRow<FSceneOutlinerTreeItemPtr>& RowItem, bool bIsInteractable = true) const;
 
 	// Get the hierarchy data associated with this table viewer
 	const TOptional<FTedsOutlinerHierarchyData>& GetHierarchyData();
