[TEDS] Fix for automated test failing on randomized sorting test.

The large iteration test was timing out on Mac. To address this the window has been increased from 5 to 30 seconds and the number of entries reduced, while still being large enough to have a significant amount of repetition.
Additional changes made:
- Added the ability to print the remaining tasks to a string. If the issue happens again in the future it will give an indication what's still left in task list.
- Removed `FCopyPrefixesToRowHandleArrayViewTask` as it was a duplicate of `FCopyPrefixesToRowHandleArrayTask`.
- Replaced `Algo::Sort` with `Algo::StableSort` in the large iteration test to make sure it matches what the sort tasks now produce.
- Don't report all mismatching results, but instead only report the first one. This previously generated a very large list of the same "Check failed" messages.

#rnx

diff --git a/Engine/Plugins/Experimental/EditorDataStorageFeatures/Source/TedsQueryStack/Private/Sorting/SortTasks.cpp b/Engine/Plugins/Experimental/EditorDataStorageFeatures/Source/TedsQueryStack/Private/Sorting/SortTasks.cpp
--- a/Engine/Plugins/Experimental/EditorDataStorageFeatures/Source/TedsQueryStack/Private/Sorting/SortTasks.cpp
+++ b/Engine/Plugins/Experimental/EditorDataStorageFeatures/Source/TedsQueryStack/Private/Sorting/SortTasks.cpp
@@ -159,6 +159,40 @@ namespace UE::Editor::DataStorage::QueryStack::Sorting
 		return !MainTaskList.IsEmpty() || !PendingTasks.IsEmpty() || !PendingTasksQueue.IsEmpty();
 	}
 
+	void FSortTasks::PrintRemainingTaskList(FString& Output)
+	{
+		checkf(!bIsUpdating, TEXT("It's not allowed to call CompileRemainingTaskList while updating."));
+
+		FTaskPrinter Printer{ .Output = Output };
+
+		// Collect pending tasks to make them accessible.
+		while (TaskChain* Next = PendingTasksQueue.Peek())
+		{
+			PendingTasks.Add(MoveTemp(*Next));
+			PendingTasksQueue.Dequeue();
+		}
+
+		auto ReportChain = [](FTaskPrinter& Printer, const TaskChain& Chain)
+			{
+				for (const FPrefixCollectionBatch& Batch : Chain)
+				{
+					Printer.Output.Append(TEXT("    "));
+					Visit(Printer, Batch.Task);
+					Printer.Output.Append(Batch.bCompleted ? TEXT(" (Completed)\n") : TEXT("\n"));
+				}
+			};
+
+		Output.Append(TEXT("  Main task list\n"));
+		ReportChain(Printer, MainTaskList);
+		
+		int32 BatchCount = 0;
+		for (TaskChain& Chain : PendingTasks)
+		{
+			Output.Appendf(TEXT("  Batch %i\n"), BatchCount++);
+			ReportChain(Printer, Chain);
+		}
+	}
+
 	void FSortTasks::FTaskVisitor::operator()(const FTaskBarrier& Task)
 	{
 		checkf(false, TEXT("Task barriers should never be processed."));
@@ -230,20 +264,6 @@ namespace UE::Editor::DataStorage::QueryStack::Sorting
 		}
 	}
 
-	void FSortTasks::FTaskVisitor::operator()(const FCopyPrefixesToRowHandleArrayViewTask& Task)
-	{
-		TRACE_CPUPROFILER_EVENT_SCOPE("[TEDS Sorting] Copy out prefixes to array view");
-
-		RowHandle* WriteIt = Task.Output.GetData();
-		const FPrefix* ReadIt = Task.Input.GetData();
-		const FPrefix* ReadEnd = ReadIt + Task.Input.Num();
-
-		for (; ReadIt < ReadEnd; ++ReadIt)
-		{
-			*WriteIt++ = ReadIt->Row;
-		}
-	}
-
 	void FSortTasks::FTaskVisitor::operator()(const FComparativeSortRowHandleRangeTask& Task)
 	{
 		TRACE_CPUPROFILER_EVENT_SCOPE("[TEDS Sorting] Comparative sort");
@@ -436,7 +456,7 @@ namespace UE::Editor::DataStorage::QueryStack::Sorting
 						CompareSortTask.Reserve(3);
 						Owner.ComparativeSortRowHandlePrefixRanges(CompareSortTask, SubPrefixView, Task.Storage, Task.Sorter);
 						Owner.AddBarrier(CompareSortTask);
-						Owner.CopyPrefixesToRowHandleArrayView(CompareSortTask, Task.Array.Slice(Offset, Count), SubPrefixView);
+						Owner.CopyPrefixesToRowHandleArray(CompareSortTask, Task.Array.Slice(Offset, Count), SubPrefixView);
 						Owner.PendingTasksQueue.Enqueue(MoveTemp(CompareSortTask));
 					}
 					else
@@ -458,6 +478,61 @@ namespace UE::Editor::DataStorage::QueryStack::Sorting
 		}
 	}
 
+	void FSortTasks::FTaskPrinter::operator()(const FTaskBarrier& Task)
+	{
+		Output += TEXT("- Barrier -");
+	}
+
+	void FSortTasks::FTaskPrinter::operator()(const FCollectPrefixesTask& Task)
+	{
+		Output.Appendf(TEXT("Collect prefixes - Input [0x%p, %i], Output [0x%p, %i]"), 
+			Task.Input.GetData(), Task.Input.Num(), Task.Output.GetData(), Task.Output.Num());
+	}
+
+	void FSortTasks::FTaskPrinter::operator()(const FRefreshPrefixesTask& Task)
+	{
+		Output.Appendf(TEXT("Refresh prefixes - Prefixes [0x%p, %i], Byte Index %u"), 
+			Task.Prefixes.GetData(), Task.Prefixes.Num(), Task.ByteIndex);
+	}
+
+	void FSortTasks::FTaskPrinter::operator()(const FRadixSortPrefixesTask& Task)
+	{
+		Output.Appendf(TEXT("Radix sort prefixes - Output [0x%p, %i]"), Task.Output.GetData(), Task.Output.Num());
+	}
+
+	void FSortTasks::FTaskPrinter::operator()(const FCopyPrefixesToRowHandleArrayTask& Task)
+	{
+		Output.Appendf(TEXT("Copy prefixes to row handle array - Input [0x%p, %i], Output [0x%p, %i]"), 
+			Task.Input.GetData(), Task.Input.Num(), Task.Output.GetData(), Task.Output.Num());
+	}
+
+	void FSortTasks::FTaskPrinter::operator()(const FComparativeSortRowHandleRangeTask& Task)
+	{
+		Output.Appendf(TEXT("Comparative sort - Rows [0x%p, %i]"), Task.Rows.GetData(), Task.Rows.Num());
+	}
+
+	void FSortTasks::FTaskPrinter::operator()(const FComparativeSortRowHandlePrefixRangeTask& Task)
+	{
+		Output.Appendf(TEXT("Comparative sort (prefixed) - Rows [0x%p, %i]"), Task.Rows.GetData(), Task.Rows.Num());
+	}
+
+	void FSortTasks::FTaskPrinter::operator()(const FMergeSortedAdjacentRangesTask& Task)
+	{
+		Output.Appendf(TEXT("Comparative sort (prefixed) - Array [0x%p, %i], Scratch [0x%p, %i], Split pos %i"), 
+			Task.Array.GetData(), Task.Array.Num(),
+			Task.Scratch.GetData(), Task.Scratch.Num(),
+			Task.SplitPosition);
+	}
+
+	void FSortTasks::FTaskPrinter::operator()(const FBucketizeRangesTask& Task)
+	{
+		Output.Appendf(TEXT("Bucketize range - Array [0x%p, %i], Prefixes [0x%p, %i], Shadow Prefixes [0x%p, %i] Prefix batch size %i, Sort batch size %i, Index %i"),
+			Task.Array.GetData(), Task.Array.Num(),
+			Task.Prefixes.GetData(), Task.Prefixes.Num(),
+			Task.PrefixesShadow.GetData(), Task.PrefixesShadow.Num(),
+			Task.PrefixBatchSize, Task.SortBatchSize, Task.Index);
+	}
+
 	void FSortTasks::AddBarrier(TaskChain& Tasks)
 	{
 		Tasks.Emplace(FPrefixCollectionBatch
@@ -539,20 +614,6 @@ namespace UE::Editor::DataStorage::QueryStack::Sorting
 			});
 	}
 
-	void FSortTasks::CopyPrefixesToRowHandleArrayView(TaskChain& Tasks, TArrayView<RowHandle> Output, TConstArrayView<FPrefix> Input)
-	{
-		checkf(Output.Num() == Input.Num(),
-			TEXT("The input sort prefix list doesn't have the same number of entries as the row handle output array view."));
-		Tasks.Emplace(FPrefixCollectionBatch
-			{
-				.Task = TaskVariant(TInPlaceType<FCopyPrefixesToRowHandleArrayViewTask>(), FCopyPrefixesToRowHandleArrayViewTask
-					{
-						.Output = Output,
-						.Input = Input
-					})
-			});
-	}
-
 	void FSortTasks::ComparativeSortRowHandleRanges(TaskChain& Tasks, TArrayView<RowHandle> Rows, const ICoreProvider& Storage,
 		TSharedPtr<const FColumnSorterInterface> Sorter, int32 BatchSize)
 	{
diff --git a/Engine/Plugins/Experimental/EditorDataStorageFeatures/Source/TedsQueryStack/Private/Sorting/SortTasks.h b/Engine/Plugins/Experimental/EditorDataStorageFeatures/Source/TedsQueryStack/Private/Sorting/SortTasks.h
--- a/Engine/Plugins/Experimental/EditorDataStorageFeatures/Source/TedsQueryStack/Private/Sorting/SortTasks.h
+++ b/Engine/Plugins/Experimental/EditorDataStorageFeatures/Source/TedsQueryStack/Private/Sorting/SortTasks.h
@@ -46,7 +46,13 @@ namespace UE::Editor::DataStorage::QueryStack::Sorting
 		
 		void Update(FTimespan AllottedTime);
 
+		/** Whether or not there are still any tasks remaining. If this is called during update it may not provide accurate results. */
 		bool HasRemainingTasks() const;
+		/** 
+		 * Collects all tasks pending on various threads and compile them together into a list of remaining tasks. This cannot be safely
+		 * called if Update is doing work.
+		 */
+		void PrintRemainingTaskList(FString& Output);
 
 	private:
 		struct FTaskBarrier {};
@@ -73,11 +79,6 @@ namespace UE::Editor::DataStorage::QueryStack::Sorting
 			TArrayView<RowHandle> Output;
 			TConstArrayView<FPrefix> Input;
 		};
-		struct FCopyPrefixesToRowHandleArrayViewTask
-		{
-			TArrayView<RowHandle> Output;
-			TConstArrayView<FPrefix> Input;
-		};
 		struct FComparativeSortRowHandleRangeTask
 		{
 			TArrayView<RowHandle> Rows;
@@ -115,7 +116,6 @@ namespace UE::Editor::DataStorage::QueryStack::Sorting
 			FRefreshPrefixesTask,
 			FRadixSortPrefixesTask,
 			FCopyPrefixesToRowHandleArrayTask,
-			FCopyPrefixesToRowHandleArrayViewTask,
 			FComparativeSortRowHandleRangeTask,
 			FComparativeSortRowHandlePrefixRangeTask,
 			FMergeSortedAdjacentRangesTask,
@@ -135,7 +135,6 @@ namespace UE::Editor::DataStorage::QueryStack::Sorting
 			void operator()(const FRefreshPrefixesTask& Task);
 			void operator()(const FRadixSortPrefixesTask& Task);
 			void operator()(const FCopyPrefixesToRowHandleArrayTask& Task);
-			void operator()(const FCopyPrefixesToRowHandleArrayViewTask& Task);
 			void operator()(const FComparativeSortRowHandleRangeTask& Task);
 			void operator()(const FComparativeSortRowHandlePrefixRangeTask& Task);
 			void operator()(const FMergeSortedAdjacentRangesTask& Task);
@@ -143,6 +142,21 @@ namespace UE::Editor::DataStorage::QueryStack::Sorting
 
 			FSortTasks& Owner;
 		};
+
+		struct FTaskPrinter
+		{
+			void operator()(const FTaskBarrier& Task);
+			void operator()(const FCollectPrefixesTask& Task);
+			void operator()(const FRefreshPrefixesTask& Task);
+			void operator()(const FRadixSortPrefixesTask& Task);
+			void operator()(const FCopyPrefixesToRowHandleArrayTask& Task);
+			void operator()(const FComparativeSortRowHandleRangeTask& Task);
+			void operator()(const FComparativeSortRowHandlePrefixRangeTask& Task);
+			void operator()(const FMergeSortedAdjacentRangesTask& Task);
+			void operator()(const FBucketizeRangesTask& Task);
+
+			FString& Output;
+		};
 		
 		using TaskChain = TArray<FPrefixCollectionBatch>;
 		using TaskReferenceChain = TArray<FPrefixCollectionBatch*>;
@@ -154,7 +168,6 @@ namespace UE::Editor::DataStorage::QueryStack::Sorting
 			TSharedPtr<const FColumnSorterInterface> Sorter, uint32 ByteIndex, int32 BatchSize);
 		void RadixSortPrefixes(TaskChain& Tasks, TArrayView<FPrefix> Output);
 		void CopyPrefixesToRowHandleArray(TaskChain& Tasks, TArrayView<RowHandle> Output, TConstArrayView<FPrefix> Input);
-		void CopyPrefixesToRowHandleArrayView(TaskChain& Tasks, TArrayView<RowHandle> Output, TConstArrayView<FPrefix> Input);
 		void ComparativeSortRowHandleRanges(TaskChain& Tasks, TArrayView<RowHandle> Rows, const ICoreProvider& Storage,
 			TSharedPtr<const FColumnSorterInterface> Sorter, int32 BatchSize);
 		void ComparativeSortRowHandlePrefixRanges(TaskChain& Tasks, TArrayView<FPrefix> Rows, const ICoreProvider& Storage,
diff --git a/Engine/Plugins/Experimental/EditorDataStorageFeatures/Source/TedsQueryStack/Tests/TedsRowSortNodeTests.cpp b/Engine/Plugins/Experimental/EditorDataStorageFeatures/Source/TedsQueryStack/Tests/TedsRowSortNodeTests.cpp
--- a/Engine/Plugins/Experimental/EditorDataStorageFeatures/Source/TedsQueryStack/Tests/TedsRowSortNodeTests.cpp
+++ b/Engine/Plugins/Experimental/EditorDataStorageFeatures/Source/TedsQueryStack/Tests/TedsRowSortNodeTests.cpp
@@ -2,6 +2,7 @@
 
 #if WITH_TESTS
 
+#include "Algo/StableSort.h"
 #include "Elements/Common/EditorDataStorageFeatures.h"
 #include "Elements/Framework/TypedElementSorter.h"
 #include "Elements/Interfaces/TypedElementDataStorageInterface.h"
@@ -84,14 +85,21 @@ namespace UE::Editor::DataStorage::Testing
 		if (bIterative)
 		{
 			uint64 StartTime = FPlatformTime::Cycles64();
-			while (FTimespan(FPlatformTime::Cycles64() - StartTime) < FTimespan::FromSeconds(5) && Tasks.HasRemainingTasks())
+			while (FTimespan(FPlatformTime::Cycles64() - StartTime) < FTimespan::FromSeconds(30) && Tasks.HasRemainingTasks())
 			{
-				Tasks.Update(FTimespan::FromMilliseconds(5));
+				Tasks.Update(FTimespan::FromMilliseconds(10));
 			}
 		}
 		else
 		{
-			Tasks.Update(FTimespan::FromSeconds(5));
+			Tasks.Update(FTimespan::FromSeconds(30));
+		}
+		
+		if (Tasks.HasRemainingTasks())
+		{
+			FString Message = TEXT("Tasks ran, but after 30 seconds there was still remaining work.\n");
+			Tasks.PrintRemainingTaskList(Message);
+			CHECK_MESSAGE(*Message, false);
 		}
 	}
 
@@ -303,7 +311,7 @@ TEST_CASE_NAMED(TEDS_QueryStack_RowSortNode_SortTasks_Tests, "Editor::DataStorag
 
 		auto BucketizeRangeLargeNumber = [this, Storage](bool bWithComparisons)
 			{
-				constexpr int32 Count = 25000;
+				constexpr int32 Count = 1000;
 
 				TArray<RowHandle> Array;
 				Array.Reserve(Count);
@@ -314,7 +322,7 @@ TEST_CASE_NAMED(TEDS_QueryStack_RowSortNode_SortTasks_Tests, "Editor::DataStorag
 				}
 
 				TArray<RowHandle> Compare = Array;
-				Algo::Sort(Compare);
+				Algo::StableSort(Compare);
 
 				TArray<FPrefix> Prefixes;
 				TArray<FPrefix> PrefixesShadow;
@@ -330,7 +338,13 @@ TEST_CASE_NAMED(TEDS_QueryStack_RowSortNode_SortTasks_Tests, "Editor::DataStorag
 
 				for (int32 Index = 0; Index < Count; ++Index)
 				{
-					CHECK(Array[Index] == Compare[Index]);
+					if (Array[Index] != Compare[Index])
+					{
+						CHECK_MESSAGE(
+							*FString::Printf(TEXT("Sorted array didn't match expectation. First conflict at index %i."), Index), 
+							false);
+						break;
+					}
 				}
 			};
 
