Alert Widget: Register the widget for the outliner purpose to work around an issue where the outliner is created before the purpose hierarchy is initialized in TEDS

diff --git a/Engine/Plugins/Experimental/EditorDataStorageFeatures/Source/TedsAlerts/Private/TedsAlertWidget.cpp b/Engine/Plugins/Experimental/EditorDataStorageFeatures/Source/TedsAlerts/Private/TedsAlertWidget.cpp
--- a/Engine/Plugins/Experimental/EditorDataStorageFeatures/Source/TedsAlerts/Private/TedsAlertWidget.cpp
+++ b/Engine/Plugins/Experimental/EditorDataStorageFeatures/Source/TedsAlerts/Private/TedsAlertWidget.cpp
@@ -157,6 +157,13 @@ void UAlertWidgetFactory::RegisterWidgetConstructors(
 
 	DataStorageUi.RegisterWidgetFactory<FAlertWidgetConstructor>(DataStorageUi.FindPurpose(DataStorageUi.GetGeneralWidgetPurposeID()),
 		TColumn<FAlertColumn>() || TColumn<FChildAlertColumn>());
+
+	// We are also explicitly registering the widget for the outliner purpose, because the outliner can be registered before TEDS has a chance
+	// to resolve the Outliner.Cell -> General.Cell parenting info - which means the alert widget can sometimes not get created on startup depending
+	// on init order
+	DataStorageUi.RegisterWidgetFactory<FAlertWidgetConstructor>(
+		DataStorageUi.FindPurpose(IUiProvider::FPurposeInfo("SceneOutliner", "Cell", NAME_None).GeneratePurposeID()),
+		TColumn<FAlertColumn>() || TColumn<FChildAlertColumn>());
 	
 	DataStorageUi.RegisterWidgetFactory<FAlertHeaderWidgetConstructor>(
 		DataStorageUi.FindPurpose(IUiProvider::FPurposeInfo("General", "Header", NAME_None).GeneratePurposeID()),
