Teds Outliner - use the query stack for row collection, addition, removal and filters instead of custom queries

diff --git a/Engine/Plugins/Experimental/EditorDataStorageFeatures/Source/TedsOutliner/Private/TedsOutlinerImpl.cpp b/Engine/Plugins/Experimental/EditorDataStorageFeatures/Source/TedsOutliner/Private/TedsOutlinerImpl.cpp
--- a/Engine/Plugins/Experimental/EditorDataStorageFeatures/Source/TedsOutliner/Private/TedsOutlinerImpl.cpp
+++ b/Engine/Plugins/Experimental/EditorDataStorageFeatures/Source/TedsOutliner/Private/TedsOutlinerImpl.cpp
@@ -16,11 +16,17 @@
 #include "Compatibility/SceneOutlinerRowHandleColumn.h"
 #include "TedsOutlinerFilter.h"
 #include "TedsOutlinerItem.h"
+#include "TedsQueryMergeNode.h"
+#include "TedsQueryNode.h"
+#include "TedsRowMergeNode.h"
+#include "TedsRowMonitorNode.h"
+#include "TedsRowQueryResultsNode.h"
 #include "Columns/TedsOutlinerColumns.h"
 #include "Elements/Columns/TypedElementCompatibilityColumns.h"
 #include "Elements/Columns/TypedElementOverrideColumns.h"
 #include "Elements/Columns/TypedElementSelectionColumns.h"
 #include "Elements/Columns/TypedElementSlateWidgetColumns.h"
+#include "Elements/Framework/TypedElementQueryBuilder.h"
 #include "Filters/FilterBase.h"
 #include "Widgets/SWidget.h"
 
@@ -198,7 +204,7 @@ FTedsOutlinerImpl::FTedsOutlinerImpl(const FTedsOutlinerParams& InParams, IScene
 	StorageUi = GetMutableDataStorageFeature<DataStorage::IUiProvider>(UiFeatureName);
 	StorageCompatibility = GetMutableDataStorageFeature<DataStorage::ICompatibilityProvider>(CompatibilityFeatureName);
 
-	bUsingQueryConditionsSyntax = InitialQueryDescription.Get().Conditions && !InitialQueryDescription.Get().Conditions->IsEmpty();
+	bUsingQueryConditionsSyntax = InitialQueryDescription.Conditions && !InitialQueryDescription.Conditions->IsEmpty();
 }
 
 void FTedsOutlinerImpl::CreateFilterQueries()
@@ -264,11 +270,34 @@ void FTedsOutlinerImpl::CreateFilterQueries()
 void FTedsOutlinerImpl::Init()
 {
 	CreateFilterQueries();
+	RecompileQueries();
 
 	// Tick post TEDS update to make sure all processors have run and the data is correct
 	Storage->OnUpdateCompleted().AddRaw(this, &FTedsOutlinerImpl::Tick);
 }
 
+void FTedsOutlinerImpl::OnMonitoredRowsAdded(DataStorage::FRowHandleArrayView InRows)
+{
+	for (DataStorage::RowHandle Row : InRows)
+	{
+		FSceneOutlinerHierarchyChangedData EventData;
+		EventData.Type = FSceneOutlinerHierarchyChangedData::Added;
+		EventData.Items.Add(SceneOutlinerMode->CreateItemFor<FTedsOutlinerTreeItem>(FTedsOutlinerTreeItem(Row, AsShared())));
+		HierarchyChangedEvent.Broadcast(EventData);
+	}
+}
+
+void FTedsOutlinerImpl::OnMonitoredRowsRemoved(DataStorage::FRowHandleArrayView InRows)
+{
+	for (DataStorage::RowHandle Row : InRows)
+	{
+		FSceneOutlinerHierarchyChangedData EventData;
+		EventData.Type = FSceneOutlinerHierarchyChangedData::Removed;
+		EventData.ItemIDs.Add(Row);
+		HierarchyChangedEvent.Broadcast(EventData);
+	}
+}
+
 FTedsOutlinerImpl::~FTedsOutlinerImpl()
 {
 	if (Storage)
@@ -333,46 +362,9 @@ TSharedRef<SWidget>FTedsOutlinerImpl::CreateLabelWidgetForItem(DataStorage::RowH
 	return CreateLabelWidget(*Storage, *StorageUi, LabelWidgetPurpose, TargetRow, TreeItem, RowItem, bIsInteractable);
 }
 
-void FTedsOutlinerImpl::AppendQuery(DataStorage::FQueryDescription& Query1, const DataStorage::FQueryDescription& Query2)
-{
-	// TEDS-Outliner TODO: We simply discard duplicate types for now but we probably want a more robust system to detect duplicates and conflicting conditions
-	for(int32 i = 0; i < Query2.ConditionOperators.Num(); ++i)
-	{
-		// Make sure we don't add duplicate conditions
-		DataStorage::FQueryDescription::FOperator* FoundCondition = Query1.ConditionOperators.FindByPredicate([&Query2, i](const DataStorage::FQueryDescription::FOperator& Op)
-		{
-			return Op.Type == Query2.ConditionOperators[i].Type;
-		});
-
-		// We also can't have a duplicate selection type and condition
-		TWeakObjectPtr<const UScriptStruct>* FoundSelection = Query1.SelectionTypes.FindByPredicate([&Query2, i](const TWeakObjectPtr<const UScriptStruct>& Selection)
-		{
-			return Selection == Query2.ConditionOperators[i].Type;
-		});
-
-		if (!FoundCondition && !FoundSelection)
-		{
-			Query1.ConditionOperators.Add(Query2.ConditionOperators[i]);
-			Query1.ConditionTypes.Add(Query2.ConditionTypes[i]);
-		}
-	}
-
-	if (Query2.Conditions)
-	{
-		if (Query1.Conditions)
-		{
-			Query1.Conditions = Query1.Conditions.GetValue() && Query2.Conditions.GetValue();
-		}
-		else
-		{
-			Query1.Conditions = Query2.Conditions;
-		}
-	}
-}
-
 void FTedsOutlinerImpl::AddExternalQuery(FName QueryName, const DataStorage::FQueryDescription& InQueryDescription)
 {
-	ExternalQueries.Emplace(QueryName, InQueryDescription);
+	ExternalQueries.Emplace(QueryName, MakeShared<DataStorage::QueryStack::FQueryNode>(*Storage, InQueryDescription));
 
 	RecompileQueries();
 }
@@ -382,15 +374,16 @@ void FTedsOutlinerImpl::RemoveExternalQuery(FName QueryName)
 	ExternalQueries.Remove(QueryName);
 }
 
+
+
 void FTedsOutlinerImpl::AppendExternalQueries(DataStorage::FQueryDescription& OutQuery)
 {
-	for(const TPair<FName, DataStorage::FQueryDescription>& ExternalQuery : ExternalQueries)
+	for(const TPair<FName, TSharedPtr<DataStorage::QueryStack::IQueryNode>>& ExternalQuery : ExternalQueries)
 	{
-		AppendQuery(OutQuery, ExternalQuery.Value);
+		DataStorage::Queries::MergeQueries(OutQuery, Storage->GetQueryDescription(ExternalQuery.Value->GetQuery()));
 	}
 }
 
-
 bool FTedsOutlinerImpl::CanDisplayRow(DataStorage::RowHandle ItemRowHandle) const
 {
 	/*
@@ -416,19 +409,13 @@ bool FTedsOutlinerImpl::CanDisplayRow(DataStorage::RowHandle ItemRowHandle) cons
 void FTedsOutlinerImpl::CreateItemsFromQuery(TArray<FSceneOutlinerTreeItemPtr>& OutItems, ISceneOutlinerMode* InMode) const
 {
 	using namespace UE::Editor::DataStorage::Queries;
-	
-	TArray<RowHandle> Rows;
-	
-	DirectQueryCallback RowCollector = CreateDirectQueryCallbackBinding(
-		[&Rows](IDirectQueryContext& Context, const RowHandle*)
-		{
-			TConstArrayView<RowHandle> ContextRows = Context.GetRowHandles();
-			Rows.Append(ContextRows);
-		});
 
-	Storage->RunQuery(RowHandleQuery, RowCollector);
+	// We only update the row collector node on demand, since this function is only called on a FullRefresh()
+	// Addition/Removal in between refreshes is handled by the row monitor node
+	RowCollectorNode->Update();
+	FRowHandleArrayView Rows = RowCollectorNode->GetRows();
 	
-	for (const RowHandle& Row : Rows)
+	for (RowHandle Row : Rows)
 	{
 		if (!CanDisplayRow(Row))
 		{
@@ -576,78 +563,41 @@ bool FTedsOutlinerImpl::ShouldForceShowParentRows()
 	return bForceShowParents;
 }
 
-void FTedsOutlinerImpl::OnItemAdded(DataStorage::RowHandle ItemRowHandle)
-{
-	if (!CanDisplayRow(ItemRowHandle))
-	{
-		return;
-	}
-	
-	RowsPendingAddition.Add(ItemRowHandle);
-}
-
-void FTedsOutlinerImpl::OnItemRemoved(DataStorage::RowHandle ItemRowHandle)
-{
-	FSceneOutlinerHierarchyChangedData EventData;
-	EventData.Type = FSceneOutlinerHierarchyChangedData::Removed;
-	EventData.ItemIDs.Add(ItemRowHandle);
-	HierarchyChangedEvent.Broadcast(EventData);
-}
-
 void FTedsOutlinerImpl::RecompileQueries()
 {
 	using namespace UE::Editor::DataStorage::Queries;
+	using namespace UE::Editor::DataStorage::QueryStack;
 
 	UnregisterQueries();
 
-	if (!InitialQueryDescription.IsSet())
+	// Get a list of all the query nodes (initial query + filters) to combine them into one composite query
+	TArray<TSharedPtr<IQueryNode>> QueryNodes;
+	QueryNodes.Add(MakeShared<FQueryNode>(*Storage, InitialQueryDescription));
+
+	for (const TPair<FName, TSharedPtr<IQueryNode>>& ExternalQuery : ExternalQueries)
 	{
-		return;
+		QueryNodes.Add(ExternalQuery.Value);
 	}
 
-	// Since InitialQueryDescription is an attribute we have to check this every time it changes 
-	bUsingQueryConditionsSyntax = InitialQueryDescription.Get().Conditions && !InitialQueryDescription.Get().Conditions->IsEmpty();
+	TSharedPtr<IQueryNode> CombinedQueryNode = MakeShared<FQueryMergeNode>(*Storage, QueryNodes);
 
-	// Our final query to collect rows to populate the Outliner - currently the same as the initial query the user provided
-	FQueryDescription FinalQueryDescription(InitialQueryDescription.Get());
-
-	// Add the filters the user has active to the query
-	AppendExternalQueries(FinalQueryDescription);
+	// Node to collect all rows matching the composite query
+	RowCollectorNode = MakeShared<FRowQueryResultsNode>(*Storage, CombinedQueryNode, FRowQueryResultsNode::ESyncFlags::None);
 
+	// Node to monitor row addition/removal
 	if (CreationParams.bUseDefaultObservers)
 	{
-		// Row to track addition of rows to the Outliner
-		FQueryDescription RowAdditionQueryDescription =
-			Select(
-				TEXT("Add Row to Outliner"),
-				FObserver::OnAdd<FTypedElementLabelColumn>().SetExecutionMode(EExecutionMode::GameThread),
-				[this](IQueryContext& Context, DataStorage::RowHandle Row)
-				{
-					OnItemAdded(Row);
-				})
-			.Compile();
+		RowMonitorNode = MakeShared<FRowMonitorNode>(*Storage, RowCollectorNode, CombinedQueryNode);
+		RowMonitorNode->OnMonitoredRowsAdded().AddRaw(this, &FTedsOutlinerImpl::OnMonitoredRowsAdded);
+		RowMonitorNode->OnMonitoredRowsRemoved().AddRaw(this, &FTedsOutlinerImpl::OnMonitoredRowsRemoved);
+	}
 
-		// Add the conditions from FinalQueryDescription to ensure we are tracking addition of the rows the user requested
-		AppendQuery(RowAdditionQueryDescription, FinalQueryDescription);
-	
-		// Row to track removal of rows from the Outliner
-		FQueryDescription RowRemovalQueryDescription =
-			Select(
-					TEXT("Remove Row from Outliner"),
-					FObserver::OnRemove<FTypedElementLabelColumn>().SetExecutionMode(EExecutionMode::GameThread),
-					[this](IQueryContext& Context, DataStorage::RowHandle Row)
-					{
-						OnItemRemoved(Row);
-					})
-				.Compile();
+	// Our final query to collect rows to populate the Outliner - currently the same as the initial query the user provided
+	FQueryDescription FinalQueryDescription(InitialQueryDescription);
+
+	// Add the filters the user has active to the query
+	AppendExternalQueries(FinalQueryDescription);
 
-		// Add the conditions from FinalQueryDescription to ensure we are tracking removal of the rows the user requested
-		AppendQuery(RowRemovalQueryDescription, FinalQueryDescription);
-		
-		RowAdditionQuery = Storage->RegisterQuery(MoveTemp(RowAdditionQueryDescription));
-		RowRemovalQuery = Storage->RegisterQuery(MoveTemp(RowRemovalQueryDescription));
-	}
-	
 	// Queries to track parent info, only required if we have hierarchy data
 	if (HierarchyData.IsSet())
 	{
@@ -672,7 +622,7 @@ void FTedsOutlinerImpl::RecompileQueries()
 		}
 		
 		// Add the conditions from FinalQueryDescription to ensure we are tracking removal of the rows the user requested
-		AppendQuery(ChildHandleQueryDescription, FinalQueryDescription);
+		MergeQueries(ChildHandleQueryDescription, FinalQueryDescription);
 
 		FQueryDescription UpdateParentQueryDescription =
 			Select(
@@ -721,8 +671,8 @@ void FTedsOutlinerImpl::RecompileQueries()
 		}
 		
 		// Add the conditions from FinalQueryDescription to ensure we are the rows the user requested
-		AppendQuery(UpdateParentQueryDescription, FinalQueryDescription);
-
+		MergeQueries(UpdateParentQueryDescription, FinalQueryDescription);
+		
 		ChildRowHandleQuery = Storage->RegisterQuery(MoveTemp(ChildHandleQueryDescription));
 		UpdateParentQuery = Storage->RegisterQuery(MoveTemp(UpdateParentQueryDescription));
 	}
@@ -759,7 +709,7 @@ void FTedsOutlinerImpl::RecompileQueries()
 							.Compile();
 
 		// Add the conditions from FinalQueryDescription to ensure we are the rows the user requested
-		AppendQuery(SelectionAddedQueryDescription, FinalQueryDescription);
+		MergeQueries(SelectionAddedQueryDescription, FinalQueryDescription);
 
 		// Query to track when a row gets deselected
 		FQueryDescription SelectionRemovedQueryDescription =
@@ -773,7 +723,7 @@ void FTedsOutlinerImpl::RecompileQueries()
 							.Compile();
 
 		// Add the conditions from FinalQueryDescription to ensure we are the rows the user requested
-		AppendQuery(SelectionRemovedQueryDescription, FinalQueryDescription);
+		MergeQueries(SelectionRemovedQueryDescription, FinalQueryDescription);
 
 		SelectedRowsQuery = Storage->RegisterQuery(MoveTemp(SelectedRowsQueryDescription));
 		SelectionAddedQuery = Storage->RegisterQuery(MoveTemp(SelectionAddedQueryDescription));
@@ -808,20 +758,15 @@ void FTedsOutlinerImpl::RecompileQueries()
 	}
 		
 	// Add the conditions from FinalQueryDescription to ensure we are the rows the user requested
-	AppendQuery(LabelUpdateQueryDescription, FinalQueryDescription);
+	MergeQueries(LabelUpdateQueryDescription, FinalQueryDescription);
 
 	LabelUpdateQuery = Storage->RegisterQuery(MoveTemp(LabelUpdateQueryDescription));
-	
-	RowHandleQuery = Storage->RegisterQuery(MoveTemp(FinalQueryDescription));
 }
 
 void FTedsOutlinerImpl::UnregisterQueries()
 {
 	if (Storage)
 	{
-		Storage->UnregisterQuery(RowHandleQuery);
-		Storage->UnregisterQuery(RowAdditionQuery);
-		Storage->UnregisterQuery(RowRemovalQuery);
 		Storage->UnregisterQuery(ChildRowHandleQuery);
 		Storage->UnregisterQuery(UpdateParentQuery);
 		Storage->UnregisterQuery(SelectedRowsQuery);
@@ -829,9 +774,6 @@ void FTedsOutlinerImpl::UnregisterQueries()
 		Storage->UnregisterQuery(SelectionRemovedQuery);
 		Storage->UnregisterQuery(LabelUpdateQuery);
 		
-		RowHandleQuery = DataStorage::InvalidQueryHandle;
-		RowAdditionQuery = DataStorage::InvalidQueryHandle;
-		RowRemovalQuery = DataStorage::InvalidQueryHandle;
 		ChildRowHandleQuery = DataStorage::InvalidQueryHandle;
 		UpdateParentQuery = DataStorage::InvalidQueryHandle;
 		SelectedRowsQuery = DataStorage::InvalidQueryHandle;
@@ -887,16 +829,11 @@ void FTedsOutlinerImpl::Tick()
 		bSelectionDirty = false;
 	}
 
-	// Process any new rows that need to be added
-	for (DataStorage::RowHandle Row : RowsPendingAddition)
+	if (RowMonitorNode)
 	{
-		FSceneOutlinerHierarchyChangedData EventData;
-		EventData.Type = FSceneOutlinerHierarchyChangedData::Added;
-		EventData.Items.Add(SceneOutlinerMode->CreateItemFor<FTedsOutlinerTreeItem>(FTedsOutlinerTreeItem(Row, AsShared())));
-		HierarchyChangedEvent.Broadcast(EventData);
+		RowMonitorNode->Update();
 	}
-	RowsPendingAddition.Empty();
-
+	
 	// Update the label for any rows that might need it
 	for (DataStorage::RowHandle Row : RowsPendingLabelUpdate)
 	{
diff --git a/Engine/Plugins/Experimental/EditorDataStorageFeatures/Source/TedsOutliner/Public/TedsOutlinerImpl.h b/Engine/Plugins/Experimental/EditorDataStorageFeatures/Source/TedsOutliner/Public/TedsOutlinerImpl.h
--- a/Engine/Plugins/Experimental/EditorDataStorageFeatures/Source/TedsOutliner/Public/TedsOutlinerImpl.h
+++ b/Engine/Plugins/Experimental/EditorDataStorageFeatures/Source/TedsOutliner/Public/TedsOutlinerImpl.h
@@ -6,10 +6,18 @@
 #include "ISceneOutlinerHierarchy.h"
 #include "ISceneOutlinerMode.h"
 #include "Elements/Columns/TypedElementHiearchyColumns.h"
+#include "Elements/Framework/TypedElementRowHandleArrayView.h"
 #include "Elements/Interfaces/TypedElementDataStorageUiInterface.h"
 
 namespace UE::Editor::DataStorage
 {
+	namespace QueryStack
+	{
+		class IQueryNode;
+		class FRowMonitorNode;
+		class IRowNode;
+	}
+
 	class IUiProvider;
 	class ICompatibilityProvider;
 } // namespace UE::Editor::DataStorage
@@ -85,7 +93,7 @@ struct FTedsOutlinerParams
 	SSceneOutliner* SceneOutliner;
 
 	// The query description that will be used to populate rows in the TEDS-Outliner
-	TAttribute<DataStorage::FQueryDescription> QueryDescription;
+	DataStorage::FQueryDescription QueryDescription;
 	
 	// TEDS queries that will be used to create filters in this Outliner
 	// TEDS-Outliner TODO: Can we consolidate this with the SceneOutliner API to create filters? Currently has to be separate because FTEDSOutlinerFilter
@@ -143,6 +151,8 @@ public:
 	virtual ~FTedsOutlinerImpl();
 
 	void Init();
+	void OnMonitoredRowsAdded(DataStorage::FRowHandleArrayView InRows);
+	void OnMonitoredRowsRemoved(DataStorage::FRowHandleArrayView InRows);
 
 	// TEDS construct getters
 	DataStorage::ICoreProvider* GetStorage() const;
@@ -177,11 +187,7 @@ public:
 
 	// Append all external queries into the given query
 	void AppendExternalQueries(DataStorage::FQueryDescription& OutQuery);
-
-	// TEDS-Outliner TODO: This should live in TEDS long term
-	// Funtion to combine 2 queries (adds to second query to the first)
-	static void AppendQuery(DataStorage::FQueryDescription& Query1, const DataStorage::FQueryDescription& Query2);
-
+	
 	// Outliner specific functionality
 	void CreateItemsFromQuery(TArray<FSceneOutlinerTreeItemPtr>& OutItems, ISceneOutlinerMode* InMode) const;
 	void CreateChildren(const FSceneOutlinerTreeItemPtr& Item, TArray<FSceneOutlinerTreeItemPtr>& OutChildren) const;
@@ -199,8 +205,6 @@ public:
 
 protected:
 	
-	void OnItemAdded(DataStorage::RowHandle ItemRowHandle);
-	void OnItemRemoved(DataStorage::RowHandle ItemRowHandle);
 	void ClearSelection() const;
 	void Tick();
 
@@ -224,18 +228,19 @@ protected:
 	DataStorage::IUiProvider::FPurposeID LabelWidgetPurpose;
 	
 	// Initial query provided by user
-	TAttribute<DataStorage::FQueryDescription> InitialQueryDescription;
+	DataStorage::FQueryDescription InitialQueryDescription;
 
 	// External queries that are currently active (e.g Filters)
-	TMap<FName, DataStorage::FQueryDescription> ExternalQueries;
+	TMap<FName, TSharedPtr<DataStorage::QueryStack::IQueryNode>> ExternalQueries;
 
 	// Optional Hierarchy Data
 	TOptional<FTedsOutlinerHierarchyData> HierarchyData;
 
-	// Querys to track row handle collection, addition and removal
-	DataStorage::QueryHandle RowHandleQuery = DataStorage::InvalidQueryHandle;
-	DataStorage::QueryHandle RowAdditionQuery = DataStorage::InvalidQueryHandle;
-	DataStorage::QueryHandle RowRemovalQuery = DataStorage::InvalidQueryHandle;
+	// The query stack node responsible for collecting all rows that match the composite query on FullRefresh()
+	TSharedPtr<DataStorage::QueryStack::IRowNode> RowCollectorNode;
+
+	// The query stack node responsible for tracking and row addition/removals that happen in between refreshes
+	TSharedPtr<DataStorage::QueryStack::FRowMonitorNode> RowMonitorNode;
 
 	// Query to get all child rows
 	DataStorage::QueryHandle ChildRowHandleQuery = DataStorage::InvalidQueryHandle;
@@ -277,7 +282,6 @@ protected:
 
 	// Addition and Label Updates are deferred because they access data storage implicitly
 	// Deletion is currently not deferred to work nicely with object lifecycles in some cases - but can be once everything goes through the query stack
-	TSet<DataStorage::RowHandle> RowsPendingAddition;
 	TSet<DataStorage::RowHandle> RowsPendingLabelUpdate;
 };
 } // namsepace UE::Editor::Outliner
diff --git a/Engine/Plugins/Experimental/EditorDataStorageFeatures/Source/TedsOutliner/TedsOutliner.Build.cs b/Engine/Plugins/Experimental/EditorDataStorageFeatures/Source/TedsOutliner/TedsOutliner.Build.cs
--- a/Engine/Plugins/Experimental/EditorDataStorageFeatures/Source/TedsOutliner/TedsOutliner.Build.cs
+++ b/Engine/Plugins/Experimental/EditorDataStorageFeatures/Source/TedsOutliner/TedsOutliner.Build.cs
@@ -26,7 +26,7 @@ public class TedsOutliner : ModuleRules
 					"SceneOutliner",
 					"Slate",
 					"SlateCore",
-					"TypedElementFramework",
+					"TypedElementFramework"
 				});
 			PrivateDependencyModuleNames.AddRange(
 				new string[]
@@ -39,7 +39,8 @@ public class TedsOutliner : ModuleRules
 					"TedsTableViewer",
 					"ToolWidgets",
 					"InputCore",
-					"LevelEditor"
+					"LevelEditor",
+					"TedsQueryStack"
 				});
 
 			DynamicallyLoadedModuleNames.AddRange(new string[] {});
