[ID Name] Add the ID Name column so it can be used in the Outliner

diff --git a/Engine/Plugins/Experimental/EditorDataStorage/Source/TedsCore/Private/Compatibility/TypedElementUObjectNameQueries.cpp b/Engine/Plugins/Experimental/EditorDataStorage/Source/TedsCore/Private/Compatibility/TypedElementUObjectNameQueries.cpp
new file mode 100644
--- /dev/null
+++ b/Engine/Plugins/Experimental/EditorDataStorage/Source/TedsCore/Private/Compatibility/TypedElementUObjectNameQueries.cpp
@@ -0,0 +1,35 @@
+// Copyright Epic Games, Inc. All Rights Reserved.
+
+#include "TypedElementUObjectNameQueries.h"
+
+#include "Elements/Columns/TypedElementCompatibilityColumns.h"
+#include "Elements/Columns/TypedElementMiscColumns.h"
+#include "Elements/Common/TypedElementQueryTypes.h"
+#include "Elements/Framework/TypedElementQueryBuilder.h"
+
+void UObjectNameDataStorageFactory::RegisterQueries(UE::Editor::DataStorage::ICoreProvider& DataStorage)
+{
+	RegisterUObjectAddNameColumn(DataStorage);
+}
+
+void UObjectNameDataStorageFactory::RegisterUObjectAddNameColumn(UE::Editor::DataStorage::ICoreProvider& DataStorage) const
+{
+	using namespace UE::Editor::DataStorage::Queries;
+
+	DataStorage.RegisterQuery(
+		Select(
+			TEXT("Sync UObject Name to column"),
+			FProcessor(EQueryTickPhase::PrePhysics, DataStorage.GetQueryTickGroupName(EQueryTickGroups::SyncExternalToDataStorage))
+			.SetExecutionMode(EExecutionMode::GameThread),
+			[](const FTypedElementUObjectColumn& Object, FNameColumn& Name)
+			{
+				if (UObject* ObjectInstance = Object.Object.Get())
+				{
+					Name.Name = ObjectInstance->GetFName();
+				}
+			}
+		)
+		.Where()
+			.All<FTypedElementSyncFromWorldTag>()
+		.Compile());
+}
diff --git a/Engine/Plugins/Experimental/EditorDataStorage/Source/TedsCore/Private/Compatibility/TypedElementUObjectNameQueries.h b/Engine/Plugins/Experimental/EditorDataStorage/Source/TedsCore/Private/Compatibility/TypedElementUObjectNameQueries.h
new file mode 100644
--- /dev/null
+++ b/Engine/Plugins/Experimental/EditorDataStorage/Source/TedsCore/Private/Compatibility/TypedElementUObjectNameQueries.h
@@ -0,0 +1,28 @@
+// Copyright Epic Games, Inc. All Rights Reserved.
+
+#pragma once
+
+#include "Elements/Interfaces/TypedElementDataStorageFactory.h"
+#include "UObject/ObjectMacros.h"
+
+#include "TypedElementUObjectNameQueries.generated.h"
+
+UCLASS()
+class UObjectNameDataStorageFactory : public UEditorDataStorageFactory
+{
+	GENERATED_BODY()
+
+public:
+	~UObjectNameDataStorageFactory() override = default;
+
+	void RegisterQueries(UE::Editor::DataStorage::ICoreProvider& DataStorage) override;
+
+private:
+
+	/**
+	 * Adds the Name columns to new actors that do not have one already.
+	 */
+	void RegisterUObjectAddNameColumn(UE::Editor::DataStorage::ICoreProvider& DataStorage) const;
+
+};
+
diff --git a/Engine/Plugins/Experimental/EditorDataStorage/Source/TedsCore/Private/TypedElementDatabaseCompatibility.cpp b/Engine/Plugins/Experimental/EditorDataStorage/Source/TedsCore/Private/TypedElementDatabaseCompatibility.cpp
--- a/Engine/Plugins/Experimental/EditorDataStorage/Source/TedsCore/Private/TypedElementDatabaseCompatibility.cpp
+++ b/Engine/Plugins/Experimental/EditorDataStorage/Source/TedsCore/Private/TypedElementDatabaseCompatibility.cpp
@@ -353,7 +353,7 @@ void UEditorDataStorageCompatibility::CreateStandardArchetypes()
 	using namespace UE::Editor::DataStorage;
 	StandardActorTable = Storage->RegisterTable(TTypedElementColumnTypeList<
 			FTypedElementUObjectColumn, FTypedElementUObjectIdColumn, FTypedElementClassTypeInfoColumn,
-			FTypedElementLabelColumn, FTypedElementLabelHashColumn, FTypedElementActorTag,
+			FTypedElementLabelColumn, FTypedElementLabelHashColumn, FNameColumn, FTypedElementActorTag,
 			FTypedElementSyncFromWorldTag>(), 
 		FName("Editor_StandardActorTable"));
 
@@ -363,7 +363,7 @@ void UEditorDataStorageCompatibility::CreateStandardArchetypes()
 
 	StandardUObjectTable = Storage->RegisterTable(TTypedElementColumnTypeList<
 			FTypedElementUObjectColumn, FTypedElementUObjectIdColumn, FTypedElementClassTypeInfoColumn,
-			FTypedElementSyncFromWorldTag>(), 
+			FNameColumn, FTypedElementSyncFromWorldTag>(), 
 		FName("Editor_StandardUObjectTable"));
 
 	StandardExternalObjectTable = Storage->RegisterTable(TTypedElementColumnTypeList<
diff --git a/Engine/Plugins/Experimental/EditorDataStorageFeatures/Source/TedsOutliner/Private/Compatibility/SceneOutlinerTedsBridge.cpp b/Engine/Plugins/Experimental/EditorDataStorageFeatures/Source/TedsOutliner/Private/Compatibility/SceneOutlinerTedsBridge.cpp
--- a/Engine/Plugins/Experimental/EditorDataStorageFeatures/Source/TedsOutliner/Private/Compatibility/SceneOutlinerTedsBridge.cpp
+++ b/Engine/Plugins/Experimental/EditorDataStorageFeatures/Source/TedsOutliner/Private/Compatibility/SceneOutlinerTedsBridge.cpp
@@ -395,6 +395,7 @@ void FSceneOutlinerTedsQueryBinder::SetupDefaultColumnMapping()
 	// Map the type column from the TEDS to the default Outliner type column, so we can show type info for objects not in TEDS
 	TEDSToOutlinerDefaultColumnMapping.Add(FTypedElementClassTypeInfoColumn::StaticStruct(), FSceneOutlinerBuiltInColumnTypes::ActorInfo());
 	TEDSToOutlinerDefaultColumnMapping.Add(FVisibleInEditorColumn::StaticStruct(), FSceneOutlinerBuiltInColumnTypes::Gutter());
+	TEDSToOutlinerDefaultColumnMapping.Add(FNameColumn::StaticStruct(), FSceneOutlinerBuiltInColumnTypes::IDName());
 }
 
 FName FSceneOutlinerTedsQueryBinder::FindOutlinerColumnFromTEDSColumns(TConstArrayView<TWeakObjectPtr<const UScriptStruct>> TEDSColumns) const
diff --git a/Engine/Plugins/Experimental/EditorDataStorageFeatures/Source/TedsOutliner/Private/TedsOutlinerModule.cpp b/Engine/Plugins/Experimental/EditorDataStorageFeatures/Source/TedsOutliner/Private/TedsOutlinerModule.cpp
--- a/Engine/Plugins/Experimental/EditorDataStorageFeatures/Source/TedsOutliner/Private/TedsOutlinerModule.cpp
+++ b/Engine/Plugins/Experimental/EditorDataStorageFeatures/Source/TedsOutliner/Private/TedsOutlinerModule.cpp
@@ -92,7 +92,6 @@ TSharedRef<ISceneOutliner> FTedsOutlinerModule::CreateTedsOutliner(const FSceneO
 	}
 	
 	InitOptions.ColumnMap.Add(FSceneOutlinerBuiltInColumnTypes::Label(), FSceneOutlinerColumnInfo(ESceneOutlinerColumnVisibility::Visible, 10));
-
 	
 	TSharedRef<ISceneOutliner> TedsOutlinerShared = SNew(SSceneOutliner, InitOptions);
 	
@@ -125,13 +124,21 @@ DataStorage::QueryHandle FTedsOutlinerModule::GetLevelEditorTedsOutlinerColumnQu
 
 	static DataStorage::QueryHandle ColumnQuery = Storage->RegisterQuery(
 		Select()
-			.ReadOnly<FTypedElementClassTypeInfoColumn, FAlertColumn, FChildAlertColumn, FVisibleInEditorColumn>()
+			.ReadOnly<FTypedElementClassTypeInfoColumn,
+					  FAlertColumn, 
+					  FChildAlertColumn,
+					  FNameColumn,
+					  FVisibleInEditorColumn>()
 		.Compile());
 
 	// Query to also include revision control info
 	static DataStorage::QueryHandle RevisionControlQuery = Storage->RegisterQuery(
 		Select()
-			.ReadOnly<FTypedElementClassTypeInfoColumn, FTypedElementPackageReference, FAlertColumn, FVisibleInEditorColumn>()
+			.ReadOnly<FTypedElementClassTypeInfoColumn,  
+					  FTypedElementPackageReference, 
+					  FAlertColumn, 
+					  FNameColumn,
+					  FVisibleInEditorColumn>()
 		.Compile());
 
 	return Private::bUseNewRevisionControlWidgets ? RevisionControlQuery : ColumnQuery;
diff --git a/Engine/Plugins/Experimental/EditorDataStorageFeatures/Source/TedsOutliner/Private/Widgets/OutlinerIdNameWidget.cpp b/Engine/Plugins/Experimental/EditorDataStorageFeatures/Source/TedsOutliner/Private/Widgets/OutlinerIdNameWidget.cpp
new file mode 100644
--- /dev/null
+++ b/Engine/Plugins/Experimental/EditorDataStorageFeatures/Source/TedsOutliner/Private/Widgets/OutlinerIdNameWidget.cpp
@@ -0,0 +1,72 @@
+// Copyright Epic Games, Inc. All Rights Reserved.
+
+#include "Widgets/OutlinerIdNameWidget.h"
+
+#include "SceneOutlinerHelpers.h"
+#include "Columns/SlateHeaderColumns.h"
+#include "Elements/Columns/TypedElementCompatibilityColumns.h"
+#include "Elements/Columns/TypedElementMiscColumns.h"
+#include "Elements/Columns/TypedElementTypeInfoColumns.h"
+#include "Elements/Columns/TypedElementSlateWidgetColumns.h"
+#include "Elements/Framework/TypedElementAttributeBinding.h"
+#include "TedsTableViewerUtils.h"
+#include "Widgets/Text/STextBlock.h"
+
+#define LOCTEXT_NAMESPACE "OutlinerOutlinerIdNameWidget"
+
+//
+// Cell Factory
+//
+void UOutlinerIdNameWidgetFactory::RegisterWidgetConstructors(UE::Editor::DataStorage::ICoreProvider& DataStorage,
+	UE::Editor::DataStorage::IUiProvider& DataStorageUi) const
+{
+	using namespace UE::Editor::DataStorage::Queries;
+
+	DataStorageUi.RegisterWidgetFactory<FOutlinerIdNameWidgetConstructor>(
+	DataStorageUi.FindPurpose(UE::Editor::DataStorage::IUiProvider::FPurposeInfo("SceneOutliner", "Cell", NAME_None).GeneratePurposeID()),
+		TColumn<FNameColumn>());
+}
+
+FOutlinerIdNameWidgetConstructor::FOutlinerIdNameWidgetConstructor()
+	: FSimpleWidgetConstructor(StaticStruct())
+{
+}
+
+TSharedPtr<SWidget> FOutlinerIdNameWidgetConstructor::CreateWidget(
+	UE::Editor::DataStorage::ICoreProvider* DataStorage,
+	UE::Editor::DataStorage::IUiProvider* DataStorageUi,
+	UE::Editor::DataStorage::RowHandle TargetRow,
+	UE::Editor::DataStorage::RowHandle WidgetRow,
+	const UE::Editor::DataStorage::FMetaDataView& Arguments)
+{
+	using namespace UE::Editor::DataStorage;
+	if (!DataStorage->IsRowAvailable(TargetRow))
+	{
+		return SNew(STextBlock)
+			.Text(LOCTEXT("MissingRowReferenceColumn", "Unable to retrieve row reference."));
+	}
+
+	FAttributeBinder TargetRowBinder(TargetRow, DataStorage);
+
+	TSharedPtr<SWidget> Widget = SNew(SHorizontalBox)
+		+ SHorizontalBox::Slot()
+		.AutoWidth()
+		.VAlign(VAlign_Center)
+		.Padding(8, 0, 0, 0)
+		[
+			SNew(STextBlock)
+				.ColorAndOpacity(FSlateColor::UseSubduedForeground())
+				.Text(TargetRowBinder.BindText(&FNameColumn::Name))
+				.ToolTipText(TargetRowBinder.BindText(&FNameColumn::Name))
+		];
+
+	return Widget;
+}
+
+FText FOutlinerIdNameWidgetConstructor::CreateWidgetDisplayNameText(UE::Editor::DataStorage::ICoreProvider* DataStorage, 
+	UE::Editor::DataStorage::RowHandle Row) const
+{
+	return LOCTEXT("SceneOutlinerColumnIdName", "ID Name");
+}
+
+#undef LOCTEXT_NAMESPACE
diff --git a/Engine/Plugins/Experimental/EditorDataStorageFeatures/Source/TedsOutliner/Public/Widgets/OutlinerIdNameWidget.h b/Engine/Plugins/Experimental/EditorDataStorageFeatures/Source/TedsOutliner/Public/Widgets/OutlinerIdNameWidget.h
new file mode 100644
--- /dev/null
+++ b/Engine/Plugins/Experimental/EditorDataStorageFeatures/Source/TedsOutliner/Public/Widgets/OutlinerIdNameWidget.h
@@ -0,0 +1,41 @@
+// Copyright Epic Games, Inc. All Rights Reserved.
+
+#pragma once
+
+#include "Elements/Interfaces/TypedElementDataStorageFactory.h"
+#include "Elements/Interfaces/TypedElementDataStorageInterface.h"
+#include "Elements/Interfaces/TypedElementDataStorageUiInterface.h"
+
+#include "OutlinerIdNameWidget.generated.h"
+
+UCLASS()
+class UOutlinerIdNameWidgetFactory : public UEditorDataStorageFactory
+{
+	GENERATED_BODY()
+
+public:
+	~UOutlinerIdNameWidgetFactory() override = default;
+
+	void RegisterWidgetConstructors(UE::Editor::DataStorage::ICoreProvider& DataStorage,
+		UE::Editor::DataStorage::IUiProvider& DataStorageUi) const override;
+};
+
+USTRUCT()
+struct TEDSOUTLINER_API FOutlinerIdNameWidgetConstructor : public FSimpleWidgetConstructor
+{
+	GENERATED_BODY()
+
+public:
+	FOutlinerIdNameWidgetConstructor();
+	~FOutlinerIdNameWidgetConstructor() override = default;
+
+	virtual TSharedPtr<SWidget> CreateWidget(
+		UE::Editor::DataStorage::ICoreProvider* DataStorage,
+		UE::Editor::DataStorage::IUiProvider* DataStorageUi,
+		UE::Editor::DataStorage::RowHandle TargetRow,
+		UE::Editor::DataStorage::RowHandle WidgetRow,
+		const UE::Editor::DataStorage::FMetaDataView& Arguments) override;
+
+	virtual FText CreateWidgetDisplayNameText(UE::Editor::DataStorage::ICoreProvider* DataStorage, 
+		UE::Editor::DataStorage::RowHandle Row = UE::Editor::DataStorage::InvalidRowHandle) const override;
+};
\ No newline at end of file
