[TEDS] TEDS UI now generates sorters if none are provided based on type information.

Properties in TEDS columns that have UPROPERTY meta data "Sortable" will now have a sorter automatically generated, if there's a handler registered. Default TEDS UI support sorting various strings and numeric values, but users can provide their own sorter for specific properties. Sorters for strings can be further configured by providing arguments, e.g. `UPROPERTY(meta = (Sortable = "FName_WithNone, FName_RemoveLeadingSlash"))`. The options available are:
- FString_CaseSensitive - Moved to a case sensitive compare for FString instead of the default case insensitive sort.
- FText_CaseSensitive - Moved  to a case sensitive compare for FText instead of the default case insensitive sort.
- FName_WithNone - Instead of treating an empty FName as an empty string, use the word "None" instead.
- FName_RemoveLeadingSlash - If an FName starts with a slash, the slash is ignored.
As an example a small number of columns have "Sortable" added to them.
This change also required an additional argument to the sorters and reorganized the `IUiProvider` so extra documentation could be provided.

This changelist is combined with the CR found at 42509024.
[TEDS] Addressed bug while sorting FNames.
The code used to locate the bug has been repurposed to add options to the FName sorter. It also prevents an intermediate `TSortNameView` if not needed, which is a minor performance improvement.

diff --git a/Engine/Plugins/Experimental/EditorDataStorage/Source/TedsCore/Private/PropertySorters/DefaultPropertySorterFactory.cpp b/Engine/Plugins/Experimental/EditorDataStorage/Source/TedsCore/Private/PropertySorters/DefaultPropertySorterFactory.cpp
new file mode 100644
--- /dev/null
+++ b/Engine/Plugins/Experimental/EditorDataStorage/Source/TedsCore/Private/PropertySorters/DefaultPropertySorterFactory.cpp
@@ -0,0 +1,57 @@
+// Copyright Epic Games, Inc. All Rights Reserved.
+
+#include "PropertySorters/DefaultPropertySorterFactory.h"
+
+#include "Elements/Interfaces/TypedElementDataStorageUiInterface.h"
+#include "PropertySorters/NumericSorters.h"
+#include "PropertySorters/StringSorters.h"
+#include "UObject/StrProperty.h"
+#include "UObject/TextProperty.h"
+#include "UObject/UnrealType.h"
+
+void UDefaultPropertySorterFactory::RegisterPropertySorters(UE::Editor::DataStorage::IUiProvider& DataStorageUi) const
+{
+	using namespace UE::Editor::DataStorage::Sorters;
+
+	auto RegisterCallback = [&DataStorageUi]<typename Sorter, typename PropertyType>()
+	{
+		DataStorageUi.RegisterSorterGeneratorForProperty(PropertyType::StaticClass(),
+			[](TWeakObjectPtr<const UScriptStruct> ColumnType, const FProperty& Property)
+			{
+				return MakeShared<Sorter>(ColumnType, CastField<PropertyType>(&Property));
+			});
+	};
+
+	auto RegisterNumericCallback = [&DataStorageUi]<typename ValueType, typename PropertyType>()
+	{
+		DataStorageUi.RegisterSorterGeneratorForProperty(PropertyType::StaticClass(),
+			[](TWeakObjectPtr<const UScriptStruct> ColumnType, const FProperty& Property)
+			{
+				return MakeShared<TNumericSorter<ValueType, PropertyType>>(ColumnType, CastField<PropertyType>(&Property));
+			});
+	};
+	
+	// Strings
+	RegisterCallback.operator()<FStringSorter, FStrProperty>();
+	RegisterCallback.operator()<FTextSorter, FTextProperty>();
+	RegisterCallback.operator()<FNameSorter, FNameProperty>();
+	
+	// Signed number properties
+	RegisterNumericCallback.operator()<int64, FInt64Property>();
+	RegisterNumericCallback.operator()<int32, FIntProperty>();
+	RegisterNumericCallback.operator()<int16, FInt16Property>();
+	RegisterNumericCallback.operator()<int8,  FInt8Property>();
+	
+	// Unsigned number properties
+	RegisterNumericCallback.operator()<uint64, FUInt64Property>();
+	RegisterNumericCallback.operator()<uint32, FUInt32Property>();
+	RegisterNumericCallback.operator()<uint16, FUInt16Property>();
+	RegisterNumericCallback.operator()<uint8,  FByteProperty>();
+
+	// Floating point
+	RegisterNumericCallback.operator()<float, FFloatProperty>();
+	RegisterNumericCallback.operator()<double, FDoubleProperty>();
+
+	// Boolean
+	RegisterCallback.operator()<FBooleanSorter, FBoolProperty>();
+}
\ No newline at end of file
diff --git a/Engine/Plugins/Experimental/EditorDataStorage/Source/TedsCore/Private/PropertySorters/DefaultPropertySorterFactory.h b/Engine/Plugins/Experimental/EditorDataStorage/Source/TedsCore/Private/PropertySorters/DefaultPropertySorterFactory.h
new file mode 100644
--- /dev/null
+++ b/Engine/Plugins/Experimental/EditorDataStorage/Source/TedsCore/Private/PropertySorters/DefaultPropertySorterFactory.h
@@ -0,0 +1,19 @@
+// Copyright Epic Games, Inc. All Rights Reserved.
+
+#pragma once
+
+#include "Elements/Interfaces/TypedElementDataStorageFactory.h"
+
+#include "DefaultPropertySorterFactory.generated.h"
+
+/**
+ * Factory used to register sorter for numeric and text properties.
+ */
+UCLASS(Transient)
+class UDefaultPropertySorterFactory : public UEditorDataStorageFactory
+{
+	GENERATED_BODY()
+
+public:
+	virtual void RegisterPropertySorters(UE::Editor::DataStorage::IUiProvider& DataStorageUi) const override;
+};
diff --git a/Engine/Plugins/Experimental/EditorDataStorage/Source/TedsCore/Private/PropertySorters/NumericSorters.cpp b/Engine/Plugins/Experimental/EditorDataStorage/Source/TedsCore/Private/PropertySorters/NumericSorters.cpp
new file mode 100644
--- /dev/null
+++ b/Engine/Plugins/Experimental/EditorDataStorage/Source/TedsCore/Private/PropertySorters/NumericSorters.cpp
@@ -0,0 +1,22 @@
+// Copyright Epic Games, Inc. All Rights Reserved.
+
+#include "PropertySorters/NumericSorters.h"
+#include "UObject/UnrealType.h"
+
+namespace UE::Editor::DataStorage::Sorters
+{
+	FBooleanSorter::FBooleanSorter(TWeakObjectPtr<const UScriptStruct> ColumnType, const FBoolProperty* Property)
+		: BaseType(MoveTemp(ColumnType), Property)
+	{
+	}
+
+	int32 FBooleanSorter::Compare(const bool& Left, const bool& Right) const
+	{
+		return static_cast<int32>(Left) - static_cast<int32>(Right);
+	}
+
+	FPrefixInfo FBooleanSorter::CalculatePrefix(const bool& Value, uint32 ByteIndex) const
+	{
+		return CreateSortPrefix(ByteIndex, Value);
+	}
+} // namespace UE::Editor::DataStorage::Sorters
\ No newline at end of file
diff --git a/Engine/Plugins/Experimental/EditorDataStorage/Source/TedsCore/Private/PropertySorters/NumericSorters.h b/Engine/Plugins/Experimental/EditorDataStorage/Source/TedsCore/Private/PropertySorters/NumericSorters.h
new file mode 100644
--- /dev/null
+++ b/Engine/Plugins/Experimental/EditorDataStorage/Source/TedsCore/Private/PropertySorters/NumericSorters.h
@@ -0,0 +1,67 @@
+// Copyright Epic Games, Inc. All Rights Reserved.
+
+#pragma once
+
+#include <type_traits>
+#include "Elements/Interfaces/TypedElementDataStorageUiInterface.h"
+#include "PropertySorters/PropertySorterBase.h"
+#include "UObject/WeakObjectPtrTemplates.h"
+
+class FBoolProperty;
+class UScriptStruct;
+
+namespace UE::Editor::DataStorage
+{
+	class ICoreProvider;
+}
+
+namespace UE::Editor::DataStorage::Sorters
+{
+	template<typename NumericType, typename NumericPropertyType>
+	class TNumericSorter final : 
+		public TPropertySorterBase<FColumnSorterInterface::ESortType::FixedSize64, NumericType, NumericPropertyType>
+	{
+		using BaseType = TPropertySorterBase<FColumnSorterInterface::ESortType::FixedSize64, NumericType, NumericPropertyType>;
+	public:
+		TNumericSorter(TWeakObjectPtr<const UScriptStruct> ColumnType, const NumericPropertyType* Property);
+
+	protected:
+		virtual int32 Compare(const NumericType& Left, const NumericType& Right) const override;
+		virtual FPrefixInfo CalculatePrefix(const NumericType& Value, uint32 ByteIndex) const override;
+	};
+
+	class FBooleanSorter final : public TPropertySorterBase<FColumnSorterInterface::ESortType::FixedSize64, bool, FBoolProperty>
+	{
+		using BaseType = TPropertySorterBase<FColumnSorterInterface::ESortType::FixedSize64, bool, FBoolProperty>;
+	public:
+		FBooleanSorter(TWeakObjectPtr<const UScriptStruct> ColumnType, const FBoolProperty* Property);
+
+	protected:
+		virtual int32 Compare(const bool& Left, const bool& Right) const override;
+		virtual FPrefixInfo CalculatePrefix(const bool& Value, uint32 ByteIndex) const override;
+	};
+
+
+	//
+	// Implementations
+	//
+
+	template<typename NumericType, typename NumericPropertyType>
+	TNumericSorter<NumericType, NumericPropertyType>::TNumericSorter(
+		TWeakObjectPtr<const UScriptStruct> ColumnType, const NumericPropertyType* Property)
+		: BaseType(MoveTemp(ColumnType), Property)
+	{
+	}
+
+	template<typename NumericType, typename NumericPropertyType>
+	int32 TNumericSorter<NumericType, NumericPropertyType>::Compare(const NumericType& Left, const NumericType& Right) const
+	{
+		return Left == Right ? 0 : (Left < Right ? -1 : 1);
+	}
+
+	template<typename NumericType, typename NumericPropertyType>
+	FPrefixInfo TNumericSorter<NumericType, NumericPropertyType>::CalculatePrefix(const NumericType& Value, uint32 ByteIndex) const
+	{
+		return CreateSortPrefix(ByteIndex, Value);
+	}
+} // namespace UE::Editor::DataStorage::Sorters
\ No newline at end of file
diff --git a/Engine/Plugins/Experimental/EditorDataStorage/Source/TedsCore/Private/PropertySorters/PropertySorterBase.h b/Engine/Plugins/Experimental/EditorDataStorage/Source/TedsCore/Private/PropertySorters/PropertySorterBase.h
new file mode 100644
--- /dev/null
+++ b/Engine/Plugins/Experimental/EditorDataStorage/Source/TedsCore/Private/PropertySorters/PropertySorterBase.h
@@ -0,0 +1,40 @@
+// Copyright Epic Games, Inc. All Rights Reserved.
+
+#pragma once
+
+#include "Elements/Interfaces/TypedElementDataStorageUiInterface.h"
+#include "UObject/WeakObjectPtrTemplates.h"
+
+class UScriptStruct;
+
+namespace UE::Editor::DataStorage
+{
+	class ICoreProvider;
+}
+
+namespace UE::Editor::DataStorage::Sorters
+{
+	template<FColumnSorterInterface::ESortType SortType, typename CompareType, typename PropertyType>
+	class TPropertySorterBase : public FColumnSorterInterface
+	{
+	public:
+		TPropertySorterBase(TWeakObjectPtr<const UScriptStruct> ColumnType, const PropertyType* Property);
+
+		virtual int32 Compare(const ICoreProvider& Storage, RowHandle Left, RowHandle Right) const override;
+		virtual FPrefixInfo CalculatePrefix(const ICoreProvider& Storage, RowHandle Row, uint32 ByteIndex) const override;
+		ESortType GetSortType() const override final;
+		FText GetShortName() const override final;
+
+	protected:
+		virtual int32 Compare(const CompareType& Left, const CompareType& Right) const = 0;
+		virtual FPrefixInfo CalculatePrefix(const CompareType& Value, uint32 ByteIndex) const = 0;
+
+	private:
+		const CompareType* GetValue(const ICoreProvider& Storage, const UScriptStruct* Column, RowHandle Row) const;
+
+		TWeakObjectPtr<const UScriptStruct> ColumnType;
+		const PropertyType* Property;
+	};
+} // namespace UE::Editor::DataStorage::Sorters
+
+#include "PropertySorters/PropertySorterBase.inl"
\ No newline at end of file
diff --git a/Engine/Plugins/Experimental/EditorDataStorage/Source/TedsCore/Private/PropertySorters/PropertySorterBase.inl b/Engine/Plugins/Experimental/EditorDataStorage/Source/TedsCore/Private/PropertySorters/PropertySorterBase.inl
new file mode 100644
--- /dev/null
+++ b/Engine/Plugins/Experimental/EditorDataStorage/Source/TedsCore/Private/PropertySorters/PropertySorterBase.inl
@@ -0,0 +1,71 @@
+// Copyright Epic Games, Inc. All Rights Reserved.
+
+#include "Elements/Interfaces/TypedElementDataStorageInterface.h"
+
+namespace UE::Editor::DataStorage::Sorters
+{
+	template<FColumnSorterInterface::ESortType SortType, typename CompareType, typename PropertyType>
+	TPropertySorterBase<SortType, CompareType, PropertyType>::TPropertySorterBase(
+		TWeakObjectPtr<const UScriptStruct> ColumnType, const PropertyType* Property)
+		: ColumnType(MoveTemp(ColumnType))
+		, Property(Property)
+	{
+	}
+
+	template<FColumnSorterInterface::ESortType SortType, typename CompareType, typename PropertyType>
+	int32 TPropertySorterBase<SortType, CompareType, PropertyType>::Compare(
+		const ICoreProvider& Storage, RowHandle Left, RowHandle Right) const
+	{
+		if (const UScriptStruct* Column = ColumnType.Get())
+		{
+			const CompareType* LeftValue = GetValue(Storage, Column, Left);
+			const CompareType* RightValue = GetValue(Storage, Column, Right);
+
+			return (LeftValue && RightValue)
+				? Compare(*LeftValue, *RightValue)
+				: ((LeftValue == nullptr) - (RightValue == nullptr));
+		}
+		return 0;
+	}
+
+	template<FColumnSorterInterface::ESortType SortType, typename CompareType, typename PropertyType>
+	FPrefixInfo TPropertySorterBase<SortType, CompareType, PropertyType>::CalculatePrefix(
+		const ICoreProvider& Storage, RowHandle Row, uint32 ByteIndex) const
+	{
+		if (const UScriptStruct* Column = ColumnType.Get())
+		{
+			if (const void* Data = Storage.GetColumnData(Row, Column))
+			{
+				if (const CompareType* Value = Property->template ContainerPtrToValuePtr<CompareType>(Data))
+				{
+					return CalculatePrefix(*Value, ByteIndex);
+				}
+			}
+		}
+		return FPrefixInfo{};
+	}
+
+	template<FColumnSorterInterface::ESortType SortType, typename CompareType, typename PropertyType>
+	FColumnSorterInterface::ESortType TPropertySorterBase<SortType, CompareType, PropertyType>::GetSortType() const
+	{
+		return SortType;
+	}
+
+	template<FColumnSorterInterface::ESortType SortType, typename CompareType, typename PropertyType>
+	FText TPropertySorterBase<SortType, CompareType, PropertyType>::GetShortName() const
+	{
+		// The property is part of the column type, so if the column type is valid then the property pointer is save to use.
+		return ColumnType.IsValid() ? Property->GetDisplayNameText() : FText::GetEmpty();
+	}
+
+	template<FColumnSorterInterface::ESortType SortType, typename CompareType, typename PropertyType>
+	const CompareType* TPropertySorterBase<SortType, CompareType, PropertyType>::GetValue(
+		const ICoreProvider& Storage, const UScriptStruct* Column, RowHandle Row) const
+	{
+		if (const void* Data = Storage.GetColumnData(Row, Column))
+		{
+			return Property->template ContainerPtrToValuePtr<CompareType>(Data);
+		}
+		return nullptr;
+	}
+} // namespace UE::Editor::DataStorage::Sorters
\ No newline at end of file
diff --git a/Engine/Plugins/Experimental/EditorDataStorage/Source/TedsCore/Private/PropertySorters/StringSorters.cpp b/Engine/Plugins/Experimental/EditorDataStorage/Source/TedsCore/Private/PropertySorters/StringSorters.cpp
new file mode 100644
--- /dev/null
+++ b/Engine/Plugins/Experimental/EditorDataStorage/Source/TedsCore/Private/PropertySorters/StringSorters.cpp
@@ -0,0 +1,118 @@
+// Copyright Epic Games, Inc. All Rights Reserved.
+
+#include "PropertySorters/StringSorters.h"
+
+namespace UE::Editor::DataStorage::Sorters
+{
+	//
+	// FStringSorter
+	//
+
+	FStringSorter::FStringSorter(TWeakObjectPtr<const UScriptStruct> ColumnType, const FStrProperty* Property)
+		: BaseType(MoveTemp(ColumnType), Property)
+	{
+		const FString& Sortable = Property->GetMetaData(TEXT("Sortable"));
+		if (Sortable.Contains(TEXT("FString_CaseSensitive")))
+		{
+			bCaseSensitive = true;
+		}
+	}
+
+	int32 FStringSorter::Compare(const FString& Left, const FString& Right) const
+	{
+		return Left.Compare(Right, bCaseSensitive ? ESearchCase::CaseSensitive : ESearchCase::IgnoreCase);
+	}
+
+	FPrefixInfo FStringSorter::CalculatePrefix(const FString& Value, uint32 ByteIndex) const
+	{
+		FPrefixInfo Result = bCaseSensitive
+			? CreateSortPrefix(ByteIndex, TSortStringView(FSortCaseSensitive{}, Value))
+			: CreateSortPrefix(ByteIndex, TSortStringView(FSortCaseInsensitive{}, Value));
+		return Result;
+	}
+
+
+
+	//
+	// FTextSorter
+	//
+
+	FTextSorter::FTextSorter(TWeakObjectPtr<const UScriptStruct> ColumnType, const FTextProperty* Property)
+		: BaseType(MoveTemp(ColumnType), Property)
+	{
+		const FString& Sortable = Property->GetMetaData(TEXT("Sortable"));
+		if (Sortable.Contains(TEXT("FText_CaseSensitive")))
+		{
+			bCaseSensitive = true;
+		}
+	}
+
+	int32 FTextSorter::Compare(const FText& Left, const FText& Right) const
+	{
+		return bCaseSensitive
+			? Left.CompareTo(Right)
+			: Left.CompareToCaseIgnored(Right);
+	}
+
+	FPrefixInfo FTextSorter::CalculatePrefix(const FText& Value, uint32 ByteIndex) const
+	{
+		return bCaseSensitive 
+			? CreateSortPrefix(ByteIndex, TSortStringView(FSortCaseSensitive{}, Value.ToString()))
+			: CreateSortPrefix(ByteIndex, TSortStringView(FSortCaseInsensitive{}, Value.ToString()));
+	}
+
+
+
+	//
+	// FNameSorter
+	//
+
+	FNameSorter::FNameSorter(TWeakObjectPtr<const UScriptStruct> ColumnType, const FNameProperty* Property)
+		: BaseType(MoveTemp(ColumnType), Property)
+	{
+		const FString& Sortable = Property->GetMetaData(TEXT("Sortable"));
+		if (Sortable.Contains(TEXT("FName_WithNone")))
+		{
+			SortByFlags |= ESortByNameFlags::WithNone;
+		}
+		if (Sortable.Contains(TEXT("FName_RemoveLeadingSlash")))
+		{
+			SortByFlags |= ESortByNameFlags::RemoveLeadingSlash;
+		}
+	}
+
+	int32 FNameSorter::Compare(const FName& Left, const FName& Right) const
+	{
+		switch (SortByFlags)
+		{
+		default:
+			// fall through
+		case ESortByNameFlags::Default:
+			return TSortNameView(TSortByName<>{}, Left).Compare(Right);
+		case ESortByNameFlags::WithNone:
+			return TSortNameView(TSortByName<ESortByNameFlags::WithNone>{}, Left).Compare(Right);
+		case ESortByNameFlags::RemoveLeadingSlash:
+			return TSortNameView(TSortByName<ESortByNameFlags::RemoveLeadingSlash>{}, Left).Compare(Right);
+		case ESortByNameFlags::WithNone | ESortByNameFlags::RemoveLeadingSlash:
+			return TSortNameView(TSortByName<ESortByNameFlags::WithNone | ESortByNameFlags::RemoveLeadingSlash>{}, Left).Compare(Right);
+		}
+	}
+
+	FPrefixInfo FNameSorter::CalculatePrefix(const FName& Value, uint32 ByteIndex) const
+	{
+		switch (SortByFlags)
+		{
+		default:
+			// fall through
+		case ESortByNameFlags::Default:
+			return CreateSortPrefix(ByteIndex, TSortNameView(TSortByName<>{}, Value));
+		case ESortByNameFlags::WithNone:
+			return CreateSortPrefix(ByteIndex, TSortNameView(TSortByName<ESortByNameFlags::WithNone>{}, Value));
+		case ESortByNameFlags::RemoveLeadingSlash:
+			return CreateSortPrefix(ByteIndex, TSortNameView(TSortByName<ESortByNameFlags::RemoveLeadingSlash>{}, Value));
+		case ESortByNameFlags::WithNone | ESortByNameFlags::RemoveLeadingSlash:
+			return CreateSortPrefix(ByteIndex, 
+				TSortNameView(TSortByName<ESortByNameFlags::WithNone | ESortByNameFlags::RemoveLeadingSlash>{}, Value));
+		}
+	}
+} // namespace UE::Editor::DataStorage::Sorters
diff --git a/Engine/Plugins/Experimental/EditorDataStorage/Source/TedsCore/Private/PropertySorters/StringSorters.h b/Engine/Plugins/Experimental/EditorDataStorage/Source/TedsCore/Private/PropertySorters/StringSorters.h
new file mode 100644
--- /dev/null
+++ b/Engine/Plugins/Experimental/EditorDataStorage/Source/TedsCore/Private/PropertySorters/StringSorters.h
@@ -0,0 +1,62 @@
+// Copyright Epic Games, Inc. All Rights Reserved.
+
+#pragma once
+
+#include "Elements/Interfaces/TypedElementDataStorageUiInterface.h"
+#include "PropertySorters/PropertySorterBase.h"
+#include "UObject/StrProperty.h"
+#include "UObject/TextProperty.h"
+#include "UObject/UnrealType.h"
+#include "UObject/WeakObjectPtrTemplates.h"
+
+class UScriptStruct;
+
+namespace UE::Editor::DataStorage
+{
+	class ICoreProvider;
+}
+
+namespace UE::Editor::DataStorage::Sorters
+{
+	class FStringSorter final : public TPropertySorterBase<FColumnSorterInterface::ESortType::HybridSort, FString, FStrProperty>
+	{
+		using BaseType = TPropertySorterBase<FColumnSorterInterface::ESortType::HybridSort, FString, FStrProperty>;
+	public:
+		FStringSorter(TWeakObjectPtr<const UScriptStruct> ColumnType, const FStrProperty* Property);
+
+	protected:
+		virtual int32 Compare(const FString& Left, const FString& Right) const override;
+		virtual FPrefixInfo CalculatePrefix(const FString& Value, uint32 ByteIndex) const override;
+
+		bool bCaseSensitive = false;
+	};
+
+
+	class FTextSorter final : public TPropertySorterBase<FColumnSorterInterface::ESortType::HybridSort, FText, FTextProperty>
+	{
+		using BaseType = TPropertySorterBase<FColumnSorterInterface::ESortType::HybridSort, FText, FTextProperty>;
+	public:
+		FTextSorter(TWeakObjectPtr<const UScriptStruct> ColumnType, const FTextProperty* Property);
+
+	protected:
+		virtual int32 Compare(const FText& Left, const FText& Right) const override;
+		virtual FPrefixInfo CalculatePrefix(const FText& Value, uint32 ByteIndex) const override;
+
+		bool bCaseSensitive = false;
+	};
+
+
+	class FNameSorter final : public TPropertySorterBase<FColumnSorterInterface::ESortType::HybridSort, FName, FNameProperty>
+	{
+		using BaseType = TPropertySorterBase<FColumnSorterInterface::ESortType::HybridSort, FName, FNameProperty>;
+	public:
+		FNameSorter(TWeakObjectPtr<const UScriptStruct> ColumnType, const FNameProperty* Property);
+
+	protected:
+		virtual int32 Compare(const FName& Left, const FName& Right) const override;
+		virtual FPrefixInfo CalculatePrefix(const FName& Value, uint32 ByteIndex) const override;
+
+		ESortByNameFlags SortByFlags = ESortByNameFlags::Default;
+	};
+	
+} // namespace UE::Editor::DataStorage::Sorters
\ No newline at end of file
diff --git a/Engine/Plugins/Experimental/EditorDataStorage/Source/TedsCore/Private/TypedElementDatabaseUI.cpp b/Engine/Plugins/Experimental/EditorDataStorage/Source/TedsCore/Private/TypedElementDatabaseUI.cpp
--- a/Engine/Plugins/Experimental/EditorDataStorage/Source/TedsCore/Private/TypedElementDatabaseUI.cpp
+++ b/Engine/Plugins/Experimental/EditorDataStorage/Source/TedsCore/Private/TypedElementDatabaseUI.cpp
@@ -693,6 +693,49 @@ UE::Editor::DataStorage::TableHandle UEditorDataStorageUi::GetWidgetTable() cons
 	return WidgetTable;
 }
 
+void UEditorDataStorageUi::GeneratePropertySorters(TArray<TSharedPtr<const UE::Editor::DataStorage::FColumnSorterInterface>>& Results,
+	TArrayView<TWeakObjectPtr<const UScriptStruct>> Columns) const
+{
+	using namespace UE::Editor::DataStorage;
+
+	for (TWeakObjectPtr<const UScriptStruct>& ColumnTypePointer : Columns)
+	{
+		if (const UScriptStruct* ColumnType = ColumnTypePointer.Get())
+		{
+			for (TFieldIterator<FProperty> PropIt(ColumnType); PropIt; ++PropIt)
+			{
+				FProperty* Property = *PropIt;
+				if (Property->HasMetaData(TEXT("Sortable")))
+				{
+					if (const PropertySorterConstructorCallback* Callback = PropertySorterConstructors.Find(Property->GetClass()))
+					{
+						Results.Add((*Callback)(ColumnTypePointer, *Property));
+					}
+					else
+					{
+						UE_LOG(LogEditorDataStorageUI, Warning, 
+							TEXT("Property '%s' in class '%s' is marked as sortable, but no property sorter for type '%s' was registered."),
+							*Property->GetName(), *Property->GetFullGroupName(true), *Property->GetClass()->GetName());
+					}
+				}
+			}
+		}
+	}
+}
+
+void UEditorDataStorageUi::RegisterSorterGeneratorForProperty(const FFieldClass* PropertyType, 
+	PropertySorterConstructorCallback PropertySorterConstructor)
+{
+	checkf(!PropertySorterConstructors.Contains(PropertyType), 
+		TEXT("Double registration of property sorter for type '%s'."), *PropertyType->GetName());
+	PropertySorterConstructors.Add(PropertyType, MoveTemp(PropertySorterConstructor));
+}
+
+void UEditorDataStorageUi::UnregisterSorterGeneratorForProperty(const FFieldClass* PropertyType)
+{
+	PropertySorterConstructors.Remove(PropertyType);
+}
+
 UE::Editor::DataStorage::IUiProvider::FPurposeID UEditorDataStorageUi::GetDefaultWidgetPurposeID() const
 {
 	return Internal::DefaultWidgetPurposeID;	
diff --git a/Engine/Plugins/Experimental/EditorDataStorage/Source/TedsCore/Private/TypedElementDatabaseUI.h b/Engine/Plugins/Experimental/EditorDataStorage/Source/TedsCore/Private/TypedElementDatabaseUI.h
--- a/Engine/Plugins/Experimental/EditorDataStorage/Source/TedsCore/Private/TypedElementDatabaseUI.h
+++ b/Engine/Plugins/Experimental/EditorDataStorage/Source/TedsCore/Private/TypedElementDatabaseUI.h
@@ -2,6 +2,7 @@
 
 #pragma once
 
+#include <type_traits>
 #include "Containers/Map.h"
 #include "Elements/Common/TypedElementQueryConditions.h"
 #include "Elements/Interfaces/TypedElementDataStorageUiInterface.h"
@@ -81,6 +82,14 @@ public:
 	
 	/** Get the table where TEDS UI widgets are stored */
 	virtual UE::Editor::DataStorage::TableHandle GetWidgetTable() const override;
+
+	virtual void GeneratePropertySorters(TArray<TSharedPtr<const UE::Editor::DataStorage::FColumnSorterInterface>>& Results,
+		TArrayView<TWeakObjectPtr<const UScriptStruct>> Columns) const override;
+
+	virtual void RegisterSorterGeneratorForProperty(const FFieldClass* PropertyType, 
+		PropertySorterConstructorCallback PropertySorterConstructor) override;
+
+	virtual void UnregisterSorterGeneratorForProperty(const FFieldClass* PropertyType) override;
 	
 	/** Get the ID of the default TEDS UI widget purpose used to register default widgets for different types of data (e.g FText -> STextBlock) */
 	virtual FPurposeID GetDefaultWidgetPurposeID() const;
@@ -94,7 +103,7 @@ public:
 private:
 	void CreateStandardArchetypes();
 	void RegisterQueries();
-
+	
 	bool CreateSingleWidgetConstructor(
 		UE::Editor::DataStorage::RowHandle FactoryRow,
 		const UE::Editor::DataStorage::FMetaDataView& Arguments,
@@ -136,8 +145,10 @@ private:
 	// Get all factories for the provided purpose
 	void GetFactories(UE::Editor::DataStorage::RowHandle PurposeRowHandle, TArray<UE::Editor::DataStorage::RowHandle>& OutFactories) const;
 
-	// The the query conditions for the provided factory
-	const UE::Editor::DataStorage::Queries::FConditions& GetFactoryConditions(UE::Editor::DataStorage::RowHandle FactoryRow) const ;
+	// The query conditions for the provided factory
+	const UE::Editor::DataStorage::Queries::FConditions& GetFactoryConditions(UE::Editor::DataStorage::RowHandle FactoryRow) const;
+
+	TMap<const FFieldClass*, PropertySorterConstructorCallback> PropertySorterConstructors;
 
 	UE::Editor::DataStorage::TableHandle WidgetTable{ UE::Editor::DataStorage::InvalidTableHandle };
 	UE::Editor::DataStorage::TableHandle WidgetPurposeTable{ UE::Editor::DataStorage::InvalidTableHandle };
@@ -145,4 +156,4 @@ private:
 	
 	UE::Editor::DataStorage::ICoreProvider* Storage{ nullptr };
 	UE::Editor::DataStorage::ICompatibilityProvider* StorageCompatibility{ nullptr };
-};
\ No newline at end of file
+};
diff --git a/Engine/Plugins/Experimental/EditorDataStorage/Source/TedsCore/Private/TypedElementsDataStorage.cpp b/Engine/Plugins/Experimental/EditorDataStorage/Source/TedsCore/Private/TypedElementsDataStorage.cpp
--- a/Engine/Plugins/Experimental/EditorDataStorage/Source/TedsCore/Private/TypedElementsDataStorage.cpp
+++ b/Engine/Plugins/Experimental/EditorDataStorage/Source/TedsCore/Private/TypedElementsDataStorage.cpp
@@ -118,10 +118,14 @@ void FEditorDataStorageModule::StartupModule()
 					
 					Factory->RegisterTables(*DataStorage);
 					Factory->RegisterTables(*DataStorage, *DataStorageCompatibility);
+					
 					Factory->RegisterTickGroups(*DataStorage);
+					
 					Factory->RegisterRegistrationFilters(*DataStorageCompatibility);
 					Factory->RegisterDealiaser(*DataStorageCompatibility);
+					
 					Factory->RegisterWidgetPurposes(*DataStorageUi);
+					Factory->RegisterPropertySorters(*DataStorageUi);
 				}
 
 				// Second pass to call all registration that would benefit or need the registration in the previous pass.
diff --git a/Engine/Plugins/Experimental/EditorDataStorage/Source/TedsUI/Private/Widgets/NameWidget.cpp b/Engine/Plugins/Experimental/EditorDataStorage/Source/TedsUI/Private/Widgets/NameWidget.cpp
--- a/Engine/Plugins/Experimental/EditorDataStorage/Source/TedsUI/Private/Widgets/NameWidget.cpp
+++ b/Engine/Plugins/Experimental/EditorDataStorage/Source/TedsUI/Private/Widgets/NameWidget.cpp
@@ -22,20 +22,12 @@ namespace UE::Editor::DataStorage
 	
 	int32 FNameWidgetSorter_NoSlash::Compare(const FNameColumn& Left, const FNameColumn& Right) const
 	{
-		FString LeftString = Left.Name.ToString();
-		LeftString.RemoveFromStart(TEXT("/"));
-		
-		FString RightString = Right.Name.ToString();
-		RightString.RemoveFromStart(TEXT("/"));
-		
-		return LeftString.Compare(RightString, ESearchCase::IgnoreCase);
+		return TSortNameView(TSortByName<ESortByNameFlags::RemoveLeadingSlash>{}, Left.Name).Compare(Right.Name);
 	}
 
 	FPrefixInfo FNameWidgetSorter_NoSlash::CalculatePrefix(const FNameColumn& Column, uint32 ByteIndex) const
 	{
-		FString String = Column.Name.ToString();
-		String.RemoveFromStart(TEXT("/"));
-		return CreateSortPrefix(ByteIndex, TSortStringView(FSortCaseInsensitive{}, String));
+		return CreateSortPrefix(ByteIndex, TSortNameView(TSortByName<ESortByNameFlags::RemoveLeadingSlash>{}, Column.Name));
 	}
 
 
@@ -52,12 +44,12 @@ namespace UE::Editor::DataStorage
 	
 	int32 FNameWidgetSorter_WithSlash::Compare(const FNameColumn& Left, const FNameColumn& Right) const
 	{
-		return Left.Name.Compare(Right.Name);
+		return TSortNameView(TSortByName<>{}, Left.Name).Compare(Right.Name);
 	}
 
 	FPrefixInfo FNameWidgetSorter_WithSlash::CalculatePrefix(const FNameColumn& Column, uint32 ByteIndex) const
 	{
-		return CreateSortPrefix(ByteIndex, TSortNameView(FSortByName{}, Column.Name));
+		return CreateSortPrefix(ByteIndex, TSortNameView(TSortByName<>{}, Column.Name));
 	}
 } // namespace UE::Editor::DataStorage
 
@@ -88,6 +80,7 @@ TSharedPtr<SWidget> FNameWidgetConstructor::CreateWidget(UE::Editor::DataStorage
 
 TArray<TSharedPtr<const UE::Editor::DataStorage::FColumnSorterInterface>> FNameWidgetConstructor::ConstructColumnSorters(
 	UE::Editor::DataStorage::ICoreProvider* DataStorage,
+	UE::Editor::DataStorage::IUiProvider* DataStorageUi,
 	const UE::Editor::DataStorage::FMetaDataView& Arguments)
 {
 	using namespace UE::Editor::DataStorage;
diff --git a/Engine/Plugins/Experimental/EditorDataStorage/Source/TedsUI/Public/Widgets/NameWidget.h b/Engine/Plugins/Experimental/EditorDataStorage/Source/TedsUI/Public/Widgets/NameWidget.h
--- a/Engine/Plugins/Experimental/EditorDataStorage/Source/TedsUI/Public/Widgets/NameWidget.h
+++ b/Engine/Plugins/Experimental/EditorDataStorage/Source/TedsUI/Public/Widgets/NameWidget.h
@@ -66,5 +66,6 @@ public:
 
 	virtual TArray<TSharedPtr<const UE::Editor::DataStorage::FColumnSorterInterface>> ConstructColumnSorters(
 		UE::Editor::DataStorage::ICoreProvider* DataStorage,
+		UE::Editor::DataStorage::IUiProvider* DataStorageUi,
 		const UE::Editor::DataStorage::FMetaDataView& Arguments) override;
 };
diff --git a/Engine/Plugins/Experimental/EditorDataStorageFeatures/Source/TedsAssetData/Private/Widgets/AssetDataVirtualPathWidget.cpp b/Engine/Plugins/Experimental/EditorDataStorageFeatures/Source/TedsAssetData/Private/Widgets/AssetDataVirtualPathWidget.cpp
--- a/Engine/Plugins/Experimental/EditorDataStorageFeatures/Source/TedsAssetData/Private/Widgets/AssetDataVirtualPathWidget.cpp
+++ b/Engine/Plugins/Experimental/EditorDataStorageFeatures/Source/TedsAssetData/Private/Widgets/AssetDataVirtualPathWidget.cpp
@@ -41,8 +41,11 @@ namespace UE::Editor::DataStorage
 		{
 			if (const FVirtualPathColumn_Experimental* VirtualPathColumn = Storage.GetColumn<FVirtualPathColumn_Experimental>(Row))
 			{
-				Result = VirtualPathColumn->VirtualPath.ToString();
-				Result = TedsAssetDataHelper::RemoveSlashFromStart(Result);
+				if (!VirtualPathColumn->VirtualPath.IsNone())
+				{
+					Result = VirtualPathColumn->VirtualPath.ToString();
+					Result = TedsAssetDataHelper::RemoveSlashFromStart(Result);
+				}
 				return true;
 			}
 		}
@@ -50,9 +53,12 @@ namespace UE::Editor::DataStorage
 		{
 			if (const FAssetPathColumn_Experimental* AssetPathColumn = Storage.GetColumn<FAssetPathColumn_Experimental>(Row))
 			{
-				Result = AssetPathColumn->Path.ToString();
-				Result = TedsAssetDataHelper::RemoveSlashFromStart(Result);
-				Result = TedsAssetDataHelper::RemoveAllFromLastSlash(Result);
+				if (!AssetPathColumn->Path.IsNone())
+				{
+					Result = AssetPathColumn->Path.ToString();
+					Result = TedsAssetDataHelper::RemoveSlashFromStart(Result);
+					Result = TedsAssetDataHelper::RemoveAllFromLastSlash(Result);
+				}
 				return true;
 			}
 		}
@@ -123,6 +129,7 @@ TSharedPtr<SWidget> FAssetDataVirtualPathWidgetConstructor::CreateWidget(UE::Edi
 
 TArray<TSharedPtr<const UE::Editor::DataStorage::FColumnSorterInterface>> FAssetDataVirtualPathWidgetConstructor::ConstructColumnSorters(
 	UE::Editor::DataStorage::ICoreProvider* DataStorage,
+	UE::Editor::DataStorage::IUiProvider* DataStorageUi,
 	const UE::Editor::DataStorage::FMetaDataView& Arguments)
 {
 	using namespace UE::Editor::DataStorage;
diff --git a/Engine/Plugins/Experimental/EditorDataStorageFeatures/Source/TedsAssetData/Private/Widgets/AssetDataVirtualPathWidget.h b/Engine/Plugins/Experimental/EditorDataStorageFeatures/Source/TedsAssetData/Private/Widgets/AssetDataVirtualPathWidget.h
--- a/Engine/Plugins/Experimental/EditorDataStorageFeatures/Source/TedsAssetData/Private/Widgets/AssetDataVirtualPathWidget.h
+++ b/Engine/Plugins/Experimental/EditorDataStorageFeatures/Source/TedsAssetData/Private/Widgets/AssetDataVirtualPathWidget.h
@@ -53,5 +53,6 @@ public:
 
 	virtual TArray<TSharedPtr<const UE::Editor::DataStorage::FColumnSorterInterface>> ConstructColumnSorters(
 		UE::Editor::DataStorage::ICoreProvider* DataStorage,
+		UE::Editor::DataStorage::IUiProvider* DataStorageUi,
 		const UE::Editor::DataStorage::FMetaDataView& Arguments) override;
 };
diff --git a/Engine/Plugins/Experimental/EditorDataStorageFeatures/Source/TedsAssetData/Public/TedsAssetDataColumns.h b/Engine/Plugins/Experimental/EditorDataStorageFeatures/Source/TedsAssetData/Public/TedsAssetDataColumns.h
--- a/Engine/Plugins/Experimental/EditorDataStorageFeatures/Source/TedsAssetData/Public/TedsAssetDataColumns.h
+++ b/Engine/Plugins/Experimental/EditorDataStorageFeatures/Source/TedsAssetData/Public/TedsAssetDataColumns.h
@@ -89,7 +89,7 @@ struct FDiskSizeColumn final : public FEditorDataStorageColumn
 	GENERATED_BODY()
 
 	// Size on disk (in bytes)
-	UPROPERTY()
+	UPROPERTY(meta = (Sortable))
 	int64 DiskSize = 0;
 };
 
diff --git a/Engine/Plugins/Experimental/EditorDataStorageFeatures/Source/TedsQueryStack/Private/TedsRowSortNode.cpp b/Engine/Plugins/Experimental/EditorDataStorageFeatures/Source/TedsQueryStack/Private/TedsRowSortNode.cpp
--- a/Engine/Plugins/Experimental/EditorDataStorageFeatures/Source/TedsQueryStack/Private/TedsRowSortNode.cpp
+++ b/Engine/Plugins/Experimental/EditorDataStorageFeatures/Source/TedsQueryStack/Private/TedsRowSortNode.cpp
@@ -50,6 +50,8 @@ namespace UE::Editor::DataStorage::QueryStack
 
 	void FRowSortNode::SetColumnSorter(TSharedPtr<const FColumnSorterInterface> InColumnSorter)
 	{
+		// Always set this even if it can cause some redundant extra work. It can be that child nodes, such as the node that reverses order
+		// or a secondary sort, was called and a new pass is needed to sort everything back in the correct order.
 		ColumnSorter = MoveTemp(InColumnSorter);
 
 		// Restart sort on next update.
@@ -192,16 +194,27 @@ namespace UE::Editor::DataStorage::QueryStack
 		if (FRowHandleArray* MutableRows = Parent->GetMutableRows())
 		{
 			FRowHandleArrayView Rows = Parent->GetRows();
-			SortContext->PrefixArray.AddUninitialized(Rows.Num());
-			SortContext->Tasks.CollectPrefixes(
-				SortContext->PrefixArray, Rows, Storage, ColumnSorter, Private::FSortContext::PrefixBatchSize);
-			SortContext->Tasks.AddBarrier();
-
-			SortContext->PrefixShadowArray.AddUninitialized(Rows.Num());
-			FRowHandleArray::EFlags Flags = MutableRows->IsUnique() ? FRowHandleArray::EFlags::IsUnique : FRowHandleArray::EFlags::None;
-			SortContext->Tasks.BucketizeRange(MutableRows->GetMutableRows(Flags), SortContext->PrefixArray, 
-				SortContext->PrefixShadowArray, ColumnSorter, Storage, Private::FSortContext::PrefixBatchSize, 
-				Private::FSortContext::SortBatchSize);
+			
+			if (Rows.Num() >= Private::FSortContext::SortBatchSize)
+			{
+				SortContext->PrefixArray.AddUninitialized(Rows.Num());
+				SortContext->Tasks.CollectPrefixes(
+					SortContext->PrefixArray, Rows, Storage, ColumnSorter, Private::FSortContext::PrefixBatchSize);
+				SortContext->Tasks.AddBarrier();
+
+				SortContext->PrefixShadowArray.AddUninitialized(Rows.Num());
+				FRowHandleArray::EFlags Flags = MutableRows->IsUnique() ? FRowHandleArray::EFlags::IsUnique : FRowHandleArray::EFlags::None;
+				SortContext->Tasks.BucketizeRange(MutableRows->GetMutableRows(Flags), SortContext->PrefixArray,
+					SortContext->PrefixShadowArray, ColumnSorter, Storage, Private::FSortContext::PrefixBatchSize,
+					Private::FSortContext::SortBatchSize);
+			}
+			else
+			{
+				// If there are less than the sort batch size of entries then there's no need for a prefix pass as it'll add extra work to sort since each
+				// bucket is going to be smaller than the batch size and instantly go to comparative sort.
+				FRowHandleArray::EFlags Flags = MutableRows->IsUnique() ? FRowHandleArray::EFlags::IsUnique : FRowHandleArray::EFlags::None;
+				SortContext->Tasks.ComparativeSortRowHandleRanges(MutableRows->GetMutableRows(Flags), Storage, ColumnSorter, Private::FSortContext::SortBatchSize);
+			}
 		}
 		else
 		{
diff --git a/Engine/Plugins/Experimental/EditorDataStorageFeatures/Source/TedsQueryStack/Tests/TedsRowSortNodeTests.cpp b/Engine/Plugins/Experimental/EditorDataStorageFeatures/Source/TedsQueryStack/Tests/TedsRowSortNodeTests.cpp
--- a/Engine/Plugins/Experimental/EditorDataStorageFeatures/Source/TedsQueryStack/Tests/TedsRowSortNodeTests.cpp
+++ b/Engine/Plugins/Experimental/EditorDataStorageFeatures/Source/TedsQueryStack/Tests/TedsRowSortNodeTests.cpp
@@ -8,6 +8,8 @@
 #include "Elements/Interfaces/TypedElementDataStorageInterface.h"
 #include "Math/UnrealMathUtility.h"
 #include "Sorting/SortTasks.h"
+#include "TedsRowArrayNode.h"
+#include "TedsRowSortNode.h"
 #include "Tests/TestHarnessAdapter.h"
 #include "Templates/SharedPointer.h"
 
@@ -78,10 +80,35 @@ namespace UE::Editor::DataStorage::Testing
 		{
 			return Strings[static_cast<int32>(Left)].Compare(Strings[static_cast<int32>(Right)]);
 		}
+
+		virtual FPrefixInfo CalculatePrefix(const ICoreProvider& Storage, RowHandle Row, uint32 ByteIndex) const override
+		{
+			return CreateSortPrefix(ByteIndex, TSortStringView(FSortCaseInsensitive{}, Strings[static_cast<int32>(Row)]));
+		}
+	};
+
+	template<typename SortByType>
+	class FNameCompareColumnSorter : public FSimpleColumnSorter
+	{
+	public:
+		TArray<FName> Names;
+
+		virtual int32 Compare(const ICoreProvider& Storage, RowHandle Left, RowHandle Right) const override
+		{
+			return TSortNameView(SortByType{}, Names[static_cast<int32>(Left)]).Compare(Names[static_cast<int32>(Right)]);
+		}
+
+		virtual FPrefixInfo CalculatePrefix(const ICoreProvider& Storage, RowHandle Row, uint32 ByteIndex) const override
+		{
+			return CreateSortPrefix(ByteIndex, TSortNameView(SortByType{}, Names[static_cast<int32>(Row)]));
+		}
 	};
 
 	void RunTasks(UE::Editor::DataStorage::QueryStack::Sorting::FSortTasks& Tasks, bool bIterative)
 	{
+		// Normally this function should return in milliseconds, but if there are bugs it could lead to tasks with infinite
+		// regression. To prevent users and/or the build machine from being stuck forever or until some large time out, this
+		// test fails if it hasn't been completed in 30 seconds.
 		if (bIterative)
 		{
 			uint64 StartTime = FPlatformTime::Cycles64();
@@ -101,8 +128,20 @@ namespace UE::Editor::DataStorage::Testing
 			Tasks.PrintRemainingTaskList(Message);
 			CHECK_MESSAGE(*Message, false);
 		}
+		CHECK_MESSAGE(TEXT("Not all tasks completed in time"), !Tasks.HasRemainingTasks());
+	}
+
+	void UpdateNode(const TSharedPtr<UE::Editor::DataStorage::QueryStack::FRowSortNode>& Node)
+	{
+		uint64 StartTime = FPlatformTime::Cycles64();
+		do
+		{
+			Node->Update();
+		} while (FTimespan(FPlatformTime::Cycles64() - StartTime) < FTimespan::FromSeconds(5) && Node->IsSorting());
+		CHECK_MESSAGE(TEXT("Sorting didn't complete in time"), !Node->IsSorting());
 	}
 
+
 	bool CompareRows(TConstArrayView<RowHandle> Source, TConstArrayView<RowHandle> Compare)
 	{
 		int32 Count = Source.Num();
@@ -169,6 +208,106 @@ TEST_CASE_NAMED(TEDS_QueryStack_RowSortNode_SortTasks_Tests, "Editor::DataStorag
 			CHECK_MESSAGE(TEXT("Sorted array didn't match expectation."), CompareRows(Array, { 0, 1, 5, 4, 2, 3 }));
 		}
 
+		SECTION("Comparative sort with empty strings")
+		{
+			TSharedPtr<FStringCompareColumnSorter> StringSorter = MakeShared<FStringCompareColumnSorter>();
+
+			/* 00 */ StringSorter->Strings.Emplace(TEXT(""));
+			/* 01 */ StringSorter->Strings.Emplace(TEXT(""));
+			/* 02 */ StringSorter->Strings.Emplace(TEXT("Toolbar"));
+			/* 03 */ StringSorter->Strings.Emplace(TEXT(""));
+			/* 04 */ StringSorter->Strings.Emplace(TEXT(""));
+			/* 05 */ StringSorter->Strings.Emplace(TEXT("Text"));
+			/* 06 */ StringSorter->Strings.Emplace(TEXT(""));
+			/* 07 */ StringSorter->Strings.Emplace(TEXT("Icon"));
+			/* 08 */ StringSorter->Strings.Emplace(TEXT(""));
+			/* 09 */ StringSorter->Strings.Emplace(TEXT(""));
+			/* 10 */ StringSorter->Strings.Emplace(TEXT("RowHandle"));
+			/* 11 */ StringSorter->Strings.Emplace(TEXT(""));
+			/* 12 */ StringSorter->Strings.Emplace(TEXT("Large"));
+			/* 13 */ StringSorter->Strings.Emplace(TEXT("Default"));
+			/* 14 */ StringSorter->Strings.Emplace(TEXT(""));
+			/* 15 */ StringSorter->Strings.Emplace(TEXT("Large"));
+			/* 16 */ StringSorter->Strings.Emplace(TEXT(""));
+			/* 17 */ StringSorter->Strings.Emplace(TEXT(""));
+
+			TArray<RowHandle> Array = { 0, 1, 2, 3, 4, 5, 6, 7, 8, 9, 10, 11, 12, 13, 14, 15, 16, 17 };
+
+			FSortTasks Tasks;
+			Tasks.ComparativeSortRowHandleRanges(Array, *Storage, StringSorter, 2500);
+			RunTasks(Tasks, false);
+
+			CHECK_MESSAGE(TEXT("Sorted array didn't match expectation."), 
+				CompareRows(Array, { 0, 1, 3, 4, 6, 8, 9, 11, 14, 16, 17, 13, 7, 12, 15, 10, 5, 2 }));
+		}
+
+		SECTION("Comparative sort with empty names")
+		{
+			TSharedPtr<FNameCompareColumnSorter<TSortByName<>>> NameSorter = MakeShared<FNameCompareColumnSorter<TSortByName<>>>();
+
+			/* 00 */ NameSorter->Names.Emplace(TEXT(""));
+			/* 01 */ NameSorter->Names.Emplace(TEXT(""));
+			/* 02 */ NameSorter->Names.Emplace(TEXT("Toolbar"));
+			/* 03 */ NameSorter->Names.Emplace(TEXT(""));
+			/* 04 */ NameSorter->Names.Emplace(TEXT(""));
+			/* 05 */ NameSorter->Names.Emplace(TEXT("Text"));
+			/* 06 */ NameSorter->Names.Emplace(TEXT(""));
+			/* 07 */ NameSorter->Names.Emplace(TEXT("Icon"));
+			/* 08 */ NameSorter->Names.Emplace(TEXT(""));
+			/* 09 */ NameSorter->Names.Emplace(TEXT(""));
+			/* 10 */ NameSorter->Names.Emplace(TEXT("RowHandle"));
+			/* 11 */ NameSorter->Names.Emplace(TEXT(""));
+			/* 12 */ NameSorter->Names.Emplace(TEXT("Large"));
+			/* 13 */ NameSorter->Names.Emplace(TEXT("Default"));
+			/* 14 */ NameSorter->Names.Emplace(TEXT(""));
+			/* 15 */ NameSorter->Names.Emplace(TEXT("Large"));
+			/* 16 */ NameSorter->Names.Emplace(TEXT(""));
+			/* 17 */ NameSorter->Names.Emplace(TEXT(""));
+
+			TArray<RowHandle> Array = { 0, 1, 2, 3, 4, 5, 6, 7, 8, 9, 10, 11, 12, 13, 14, 15, 16, 17 };
+
+			FSortTasks Tasks;
+			Tasks.ComparativeSortRowHandleRanges(Array, *Storage, NameSorter, 2500);
+			RunTasks(Tasks, false);
+
+			CHECK_MESSAGE(TEXT("Sorted array didn't match expectation."),
+				CompareRows(Array, { 0, 1, 3, 4, 6, 8, 9, 11, 14, 16, 17, 13, 7, 12, 15, 10, 5, 2 }));
+		}
+
+		SECTION("Comparative sort with empty names with none")
+		{
+			TSharedPtr<FNameCompareColumnSorter<TSortByName<ESortByNameFlags::WithNone>>> NameSorter =
+				MakeShared<FNameCompareColumnSorter<TSortByName<ESortByNameFlags::WithNone>>>();
+
+			/* 00 */ NameSorter->Names.Emplace(TEXT(""));
+			/* 01 */ NameSorter->Names.Emplace(TEXT(""));
+			/* 02 */ NameSorter->Names.Emplace(TEXT("Toolbar"));
+			/* 03 */ NameSorter->Names.Emplace(TEXT(""));
+			/* 04 */ NameSorter->Names.Emplace(TEXT(""));
+			/* 05 */ NameSorter->Names.Emplace(TEXT("Text"));
+			/* 06 */ NameSorter->Names.Emplace(TEXT(""));
+			/* 07 */ NameSorter->Names.Emplace(TEXT("Icon"));
+			/* 08 */ NameSorter->Names.Emplace(TEXT(""));
+			/* 09 */ NameSorter->Names.Emplace(TEXT(""));
+			/* 10 */ NameSorter->Names.Emplace(TEXT("RowHandle"));
+			/* 11 */ NameSorter->Names.Emplace(TEXT(""));
+			/* 12 */ NameSorter->Names.Emplace(TEXT("Large"));
+			/* 13 */ NameSorter->Names.Emplace(TEXT("Default"));
+			/* 14 */ NameSorter->Names.Emplace(TEXT(""));
+			/* 15 */ NameSorter->Names.Emplace(TEXT("Large"));
+			/* 16 */ NameSorter->Names.Emplace(TEXT(""));
+			/* 17 */ NameSorter->Names.Emplace(TEXT(""));
+
+			TArray<RowHandle> Array = { 0, 1, 2, 3, 4, 5, 6, 7, 8, 9, 10, 11, 12, 13, 14, 15, 16, 17 };
+
+			FSortTasks Tasks;
+			Tasks.ComparativeSortRowHandleRanges(Array, *Storage, NameSorter, 2500);
+			RunTasks(Tasks, false);
+
+			CHECK_MESSAGE(TEXT("Sorted array didn't match expectation."),
+				CompareRows(Array, { 13, 7, 12, 15, 0, 1, 3, 4, 6, 8, 9, 11, 14, 16, 17, 10, 5, 2 }));
+		}
+
 		SECTION("Merge sort - Equal range")
 		{
 			TArray<RowHandle> Array = { 1, 3, 5, 2, 4, 6 };
@@ -360,4 +499,124 @@ TEST_CASE_NAMED(TEDS_QueryStack_RowSortNode_SortTasks_Tests, "Editor::DataStorag
 	}
 }
 
+TEST_CASE_NAMED(TEDS_QueryStack_RowSortNode_SortNode_Integration_Tests, "Editor::DataStorage::QueryStack::FRowSortNode::Integration", "[ApplicationContextMask][EngineFilter]")
+{
+	using namespace UE::Editor::DataStorage;
+	using namespace UE::Editor::DataStorage::Testing;
+	using namespace UE::Editor::DataStorage::QueryStack;
+	using namespace UE::Editor::DataStorage::QueryStack::Sorting;
+
+	if (ICoreProvider* Storage = GetMutableDataStorageFeature<ICoreProvider>(StorageFeatureName))
+	{
+		SECTION("Full sort single strings")
+		{
+			TSharedPtr<FStringCompareColumnSorter> StringSorter = MakeShared<FStringCompareColumnSorter>();
+
+			/* 00 */ StringSorter->Strings.Emplace(TEXT(""));
+			/* 01 */ StringSorter->Strings.Emplace(TEXT(""));
+			/* 02 */ StringSorter->Strings.Emplace(TEXT("Toolbar"));
+			/* 03 */ StringSorter->Strings.Emplace(TEXT(""));
+			/* 04 */ StringSorter->Strings.Emplace(TEXT(""));
+			/* 05 */ StringSorter->Strings.Emplace(TEXT("Text"));
+			/* 06 */ StringSorter->Strings.Emplace(TEXT(""));
+			/* 07 */ StringSorter->Strings.Emplace(TEXT("Icon"));
+			/* 08 */ StringSorter->Strings.Emplace(TEXT(""));
+			/* 09 */ StringSorter->Strings.Emplace(TEXT(""));
+			/* 10 */ StringSorter->Strings.Emplace(TEXT("RowHandle"));
+			/* 11 */ StringSorter->Strings.Emplace(TEXT(""));
+			/* 12 */ StringSorter->Strings.Emplace(TEXT("Large"));
+			/* 13 */ StringSorter->Strings.Emplace(TEXT("Default"));
+			/* 14 */ StringSorter->Strings.Emplace(TEXT(""));
+			/* 15 */ StringSorter->Strings.Emplace(TEXT("Large"));
+			/* 16 */ StringSorter->Strings.Emplace(TEXT(""));
+			/* 17 */ StringSorter->Strings.Emplace(TEXT(""));
+
+			FRowHandleArray Array(FRowHandleArrayView(
+				{ 0, 1, 2, 3, 4, 5, 6, 7, 8, 9, 10, 11, 12, 13, 14, 15, 16, 17 },
+				FRowHandleArrayView::EFlags::IsUnique));
+			
+			TSharedPtr<FRowArrayNode> ArrayNode = MakeShared<FRowArrayNode>(Array);
+			TSharedPtr<FRowSortNode> SortNode = MakeShared<FRowSortNode>(*Storage, ArrayNode, StringSorter);
+
+			UpdateNode(SortNode);
+			
+			CHECK_MESSAGE(TEXT("Sorted array didn't match expectation."),
+				CompareRows(SortNode->GetRows(), {0, 1, 3, 4, 6, 8, 9, 11, 14, 16, 17, 13, 7, 12, 15, 10, 5, 2}));
+		}
+
+		SECTION("Full sort single names")
+		{
+			TSharedPtr<FNameCompareColumnSorter<TSortByName<>>> NameSorter = MakeShared<FNameCompareColumnSorter<TSortByName<>>>();
+
+			/* 00 */ NameSorter->Names.Emplace(TEXT(""));
+			/* 01 */ NameSorter->Names.Emplace(TEXT(""));
+			/* 02 */ NameSorter->Names.Emplace(TEXT("Toolbar"));
+			/* 03 */ NameSorter->Names.Emplace(TEXT(""));
+			/* 04 */ NameSorter->Names.Emplace(TEXT(""));
+			/* 05 */ NameSorter->Names.Emplace(TEXT("Text"));
+			/* 06 */ NameSorter->Names.Emplace(TEXT(""));
+			/* 07 */ NameSorter->Names.Emplace(TEXT("Icon"));
+			/* 08 */ NameSorter->Names.Emplace(TEXT(""));
+			/* 09 */ NameSorter->Names.Emplace(TEXT(""));
+			/* 10 */ NameSorter->Names.Emplace(TEXT("RowHandle"));
+			/* 11 */ NameSorter->Names.Emplace(TEXT(""));
+			/* 12 */ NameSorter->Names.Emplace(TEXT("Large"));
+			/* 13 */ NameSorter->Names.Emplace(TEXT("Default"));
+			/* 14 */ NameSorter->Names.Emplace(TEXT(""));
+			/* 15 */ NameSorter->Names.Emplace(TEXT("Large"));
+			/* 16 */ NameSorter->Names.Emplace(TEXT(""));
+			/* 17 */ NameSorter->Names.Emplace(TEXT(""));
+
+			FRowHandleArray Array(FRowHandleArrayView(
+				{ 0, 1, 2, 3, 4, 5, 6, 7, 8, 9, 10, 11, 12, 13, 14, 15, 16, 17 },
+				FRowHandleArrayView::EFlags::IsUnique));
+
+			TSharedPtr<FRowArrayNode> ArrayNode = MakeShared<FRowArrayNode>(Array);
+			TSharedPtr<FRowSortNode> SortNode = MakeShared<FRowSortNode>(*Storage, ArrayNode, NameSorter);
+
+			UpdateNode(SortNode);
+			
+			CHECK_MESSAGE(TEXT("Sorted array didn't match expectation."),
+				CompareRows(SortNode->GetRows(), { 0, 1, 3, 4, 6, 8, 9, 11, 14, 16, 17, 13, 7, 12, 15, 10, 5, 2 }));
+		}
+
+		SECTION("Full sort single names with none")
+		{
+			TSharedPtr<FNameCompareColumnSorter<TSortByName<ESortByNameFlags::WithNone>>> NameSorter =
+				MakeShared<FNameCompareColumnSorter<TSortByName<ESortByNameFlags::WithNone>>>();
+
+			/* 00 */ NameSorter->Names.Emplace(TEXT(""));
+			/* 01 */ NameSorter->Names.Emplace(TEXT(""));
+			/* 02 */ NameSorter->Names.Emplace(TEXT("Toolbar"));
+			/* 03 */ NameSorter->Names.Emplace(TEXT(""));
+			/* 04 */ NameSorter->Names.Emplace(TEXT(""));
+			/* 05 */ NameSorter->Names.Emplace(TEXT("Text"));
+			/* 06 */ NameSorter->Names.Emplace(TEXT(""));
+			/* 07 */ NameSorter->Names.Emplace(TEXT("Icon"));
+			/* 08 */ NameSorter->Names.Emplace(TEXT(""));
+			/* 09 */ NameSorter->Names.Emplace(TEXT(""));
+			/* 10 */ NameSorter->Names.Emplace(TEXT("RowHandle"));
+			/* 11 */ NameSorter->Names.Emplace(TEXT(""));
+			/* 12 */ NameSorter->Names.Emplace(TEXT("Large"));
+			/* 13 */ NameSorter->Names.Emplace(TEXT("Default"));
+			/* 14 */ NameSorter->Names.Emplace(TEXT(""));
+			/* 15 */ NameSorter->Names.Emplace(TEXT("Large"));
+			/* 16 */ NameSorter->Names.Emplace(TEXT(""));
+			/* 17 */ NameSorter->Names.Emplace(TEXT(""));
+
+			FRowHandleArray Array(FRowHandleArrayView(
+				{ 0, 1, 2, 3, 4, 5, 6, 7, 8, 9, 10, 11, 12, 13, 14, 15, 16, 17 },
+				FRowHandleArrayView::EFlags::IsUnique));
+
+			TSharedPtr<FRowArrayNode> ArrayNode = MakeShared<FRowArrayNode>(Array);
+			TSharedPtr<FRowSortNode> SortNode = MakeShared<FRowSortNode>(*Storage, ArrayNode, NameSorter);
+
+			UpdateNode(SortNode);
+
+			CHECK_MESSAGE(TEXT("Sorted array didn't match expectation."),
+				CompareRows(SortNode->GetRows(), { 13, 7, 12, 15, 0, 1, 3, 4, 6, 8, 9, 11, 14, 16, 17, 10, 5, 2 }));
+		}
+	}
+}
+
 #endif // #if WITH_TESTS
diff --git a/Engine/Plugins/Experimental/EditorDataStorageFeatures/Source/TedsTableViewer/Private/Widgets/RowHandleWidget.cpp b/Engine/Plugins/Experimental/EditorDataStorageFeatures/Source/TedsTableViewer/Private/Widgets/RowHandleWidget.cpp
--- a/Engine/Plugins/Experimental/EditorDataStorageFeatures/Source/TedsTableViewer/Private/Widgets/RowHandleWidget.cpp
+++ b/Engine/Plugins/Experimental/EditorDataStorageFeatures/Source/TedsTableViewer/Private/Widgets/RowHandleWidget.cpp
@@ -109,6 +109,7 @@ bool FRowHandleWidgetConstructor::FinalizeWidget(UE::Editor::DataStorage::ICoreP
 
 TArray<TSharedPtr<const UE::Editor::DataStorage::FColumnSorterInterface>> FRowHandleWidgetConstructor::ConstructColumnSorters(
 	UE::Editor::DataStorage::ICoreProvider* DataStorage,
+	UE::Editor::DataStorage::IUiProvider* DataStorageUi,
 	const UE::Editor::DataStorage::FMetaDataView& Arguments)
 {
 	using namespace UE::Editor::DataStorage;
diff --git a/Engine/Plugins/Experimental/EditorDataStorageFeatures/Source/TedsTableViewer/Public/Widgets/RowHandleWidget.h b/Engine/Plugins/Experimental/EditorDataStorageFeatures/Source/TedsTableViewer/Public/Widgets/RowHandleWidget.h
--- a/Engine/Plugins/Experimental/EditorDataStorageFeatures/Source/TedsTableViewer/Public/Widgets/RowHandleWidget.h
+++ b/Engine/Plugins/Experimental/EditorDataStorageFeatures/Source/TedsTableViewer/Public/Widgets/RowHandleWidget.h
@@ -61,5 +61,6 @@ protected:
 
 	virtual TArray<TSharedPtr<const UE::Editor::DataStorage::FColumnSorterInterface>> ConstructColumnSorters(
 		UE::Editor::DataStorage::ICoreProvider* DataStorage,
+		UE::Editor::DataStorage::IUiProvider* DataStorageUi,
 		const UE::Editor::DataStorage::FMetaDataView& Arguments) override;
 };
diff --git a/Engine/Source/Runtime/TypedElementFramework/Private/Elements/Interfaces/TypedElementDataStorageUiInterface.cpp b/Engine/Source/Runtime/TypedElementFramework/Private/Elements/Interfaces/TypedElementDataStorageUiInterface.cpp
--- a/Engine/Source/Runtime/TypedElementFramework/Private/Elements/Interfaces/TypedElementDataStorageUiInterface.cpp
+++ b/Engine/Source/Runtime/TypedElementFramework/Private/Elements/Interfaces/TypedElementDataStorageUiInterface.cpp
@@ -219,9 +219,20 @@ TSharedPtr<SWidget> FTypedElementWidgetConstructor::Construct(
 
 TArray<TSharedPtr<const UE::Editor::DataStorage::FColumnSorterInterface>> FTypedElementWidgetConstructor::ConstructColumnSorters(
 	UE::Editor::DataStorage::ICoreProvider* DataStorage,
+	UE::Editor::DataStorage::IUiProvider* DataStorageUi,
 	const UE::Editor::DataStorage::FMetaDataView& Arguments)
 {
-	return TArray<TSharedPtr<const UE::Editor::DataStorage::FColumnSorterInterface>>();
+	using namespace UE::Editor::DataStorage;
+
+	TArray<TSharedPtr<const FColumnSorterInterface>> Results;
+
+	DataStorageUi->GeneratePropertySorters(Results, MatchedColumnTypes);
+	
+	TArray<TWeakObjectPtr<const UScriptStruct>, TInlineAllocator<32>> DynamicKeys;
+	MatchedDynamicTemplates.GetKeys(DynamicKeys);
+	DataStorageUi->GeneratePropertySorters(Results, DynamicKeys);
+
+	return Results;
 }
 
 TSharedPtr<SWidget> FTypedElementWidgetConstructor::CreateWidget(const UE::Editor::DataStorage::FMetaDataView& Arguments)
diff --git a/Engine/Source/Runtime/TypedElementFramework/Public/Elements/Columns/TypedElementLabelColumns.h b/Engine/Source/Runtime/TypedElementFramework/Public/Elements/Columns/TypedElementLabelColumns.h
--- a/Engine/Source/Runtime/TypedElementFramework/Public/Elements/Columns/TypedElementLabelColumns.h
+++ b/Engine/Source/Runtime/TypedElementFramework/Public/Elements/Columns/TypedElementLabelColumns.h
@@ -16,7 +16,7 @@ struct FTypedElementLabelColumn final : public FEditorDataStorageColumn
 {
 	GENERATED_BODY()
 
-	UPROPERTY(meta = (Searchable))
+	UPROPERTY(meta = (Searchable, Sortable))
 	FString Label;
 };
 
diff --git a/Engine/Source/Runtime/TypedElementFramework/Public/Elements/Columns/WidgetPurposeColumns.h b/Engine/Source/Runtime/TypedElementFramework/Public/Elements/Columns/WidgetPurposeColumns.h
--- a/Engine/Source/Runtime/TypedElementFramework/Public/Elements/Columns/WidgetPurposeColumns.h
+++ b/Engine/Source/Runtime/TypedElementFramework/Public/Elements/Columns/WidgetPurposeColumns.h
@@ -25,13 +25,13 @@ struct FWidgetPurposeNameColumn : public FEditorDataStorageColumn
 {
 	GENERATED_BODY()
 
-	UPROPERTY(meta = (Searchable))
+	UPROPERTY(meta = (Searchable, Sortable))
 	FName Namespace;
 
-	UPROPERTY(meta = (Searchable))
+	UPROPERTY(meta = (Searchable, Sortable))
 	FName Name;
 
-	UPROPERTY(meta = (Searchable))
+	UPROPERTY(meta = (Searchable, Sortable))
 	FName Frame;
 };
 
diff --git a/Engine/Source/Runtime/TypedElementFramework/Public/Elements/Framework/TypedElementSorter.h b/Engine/Source/Runtime/TypedElementFramework/Public/Elements/Framework/TypedElementSorter.h
--- a/Engine/Source/Runtime/TypedElementFramework/Public/Elements/Framework/TypedElementSorter.h
+++ b/Engine/Source/Runtime/TypedElementFramework/Public/Elements/Framework/TypedElementSorter.h
@@ -7,6 +7,7 @@
 #include "Elements/Common/TypedElementHandles.h"
 #include "Internationalization/Text.h"
 #include "Math/NumericLimits.h"
+#include "Misc/EnumClassFlags.h"
 
 namespace UE::Editor::DataStorage
 {
@@ -233,10 +234,27 @@ namespace UE::Editor::DataStorage
 
 	enum class ENameSortBy : bool { Id, String };
 
+	enum class ESortByNameFlags
+	{
+		Default = 0,
+		WithNone = 1 << 0, /** If an empty name is found, use the string "None" to represent it. */
+		RemoveLeadingSlash = 1 << 1, /** Ignore any leading slashes, if they exist. */
+	};
+	ENUM_CLASS_FLAGS(ESortByNameFlags)
+
+	/** Sort FNames by their name. If an empty name is provided it will be treated as an empty string. */
+	template<ESortByNameFlags InFlags = ESortByNameFlags::Default>
+	struct TSortByName { static const ESortByNameFlags Flags = InFlags; };
+	/** Sort FNames by their id. */
 	struct FSortById {};
-	struct FSortByName {};
+	
+	template<typename>
+	struct IsSortByName : std::false_type {};
+	template<ESortByNameFlags Flags>
+	struct IsSortByName<TSortByName<Flags>> : std::true_type {};
+
 	template<typename T>
-	concept SortBy = std::is_same_v<T, FSortById> || std::is_same_v<T, FSortByName>;
+	concept SortBy = std::is_same_v<T, FSortById> || IsSortByName<T>::value;
 
 	template<SortBy By>
 	struct TSortNameView
@@ -246,7 +264,7 @@ namespace UE::Editor::DataStorage
 		constexpr static bool bIsFixedSize = bIsById;
 
 		TSortNameView() = default;
-		TSortNameView(const FName& Name) : View(&Name) {}
+		explicit TSortNameView(const FName& Name) : View(&Name) {}
 		TSortNameView& operator=(const FName& Name);
 
 		template<SortBy TargetBy>
@@ -256,6 +274,7 @@ namespace UE::Editor::DataStorage
 		constexpr static uint32 GetElementSize();
 		FPrefixInfo CalculatePrefix(int32 CurrentIndex, int32 ByteIndex) const;
 
+		int32 Compare(const FName& Rhs) const;
 		int32 Compare(const TSortNameView& Rhs) const;
 		bool operator==(const TSortNameView& Rhs) const;
 		bool operator!=(const TSortNameView& Rhs) const;
@@ -269,6 +288,7 @@ namespace UE::Editor::DataStorage
 		mutable bool bIsCached = false;
 
 	private:
+		int32 Compare(const FName& Lhs, const FName& Rhs) const;
 		void CacheCompareType() const;
 	};
 
diff --git a/Engine/Source/Runtime/TypedElementFramework/Public/Elements/Framework/TypedElementSorter.inl b/Engine/Source/Runtime/TypedElementFramework/Public/Elements/Framework/TypedElementSorter.inl
--- a/Engine/Source/Runtime/TypedElementFramework/Public/Elements/Framework/TypedElementSorter.inl
+++ b/Engine/Source/Runtime/TypedElementFramework/Public/Elements/Framework/TypedElementSorter.inl
@@ -293,26 +293,59 @@ namespace UE::Editor::DataStorage
 			return TSortTypeInfo<SortStringView>::CalculatePrefix(CurrentIndex, ByteIndex, StringView(Cache));
 		}
 	}
-
+	
 	template<SortBy By>
-	int32 TSortNameView<By>::Compare(const TSortNameView& Rhs) const
+	int32 TSortNameView<By>::Compare(const FName& Lhs, const FName& Rhs) const
 	{
-		if (View && Rhs.View)
+		if constexpr (bIsById)
 		{
-			if constexpr (bIsById)
+			return View->CompareIndexes(Rhs);
+		}
+		else
+		{
+			auto Compare = [this](const FName& To)
+				{
+					if constexpr (EnumHasAllFlags(By::Flags, ESortByNameFlags::RemoveLeadingSlash))
+					{
+						CacheCompareType();
+						FString ToString = To.ToString();
+						ToString.RemoveFromStart(TEXT("/"));
+						return Cache.Compare(ToString, ESearchCase::IgnoreCase);
+					}
+					else
+					{
+						return View->Compare(To);
+					}
+				};
+
+			if constexpr (EnumHasAllFlags(By::Flags, ESortByNameFlags::WithNone))
 			{
-				return View->CompareIndexes(*Rhs.View);
+				return Compare(Rhs);
 			}
 			else
 			{
-				return View->Compare(*Rhs.View);
+				bool LeftIsAssigned = !View->IsNone();
+				bool RightIsAssigned = !Rhs.IsNone();
+				return (LeftIsAssigned && RightIsAssigned)
+					? Compare(Rhs)
+					// Note that this is in reversed order then elsewhere because in this case the empty value should be at the top, not the bottom.
+					: static_cast<int32>(LeftIsAssigned) - static_cast<int32>(RightIsAssigned);
 			}
 		}
-		else
-		{
-			int32 ViewIsNull = View == nullptr ? 1 : -1;
-			return View == Rhs.View ? 0 : ViewIsNull;
-		}
+	}
+
+	template<SortBy By>
+	int32 TSortNameView<By>::Compare(const FName& Rhs) const
+	{
+		return View ? Compare(*View, Rhs) : 1;
+	}
+
+	template<SortBy By>
+	int32 TSortNameView<By>::Compare(const TSortNameView& Rhs) const
+	{
+		return (View && Rhs.View)
+			? Compare(*View, *Rhs.View)
+			: static_cast<int32>(View == nullptr) - static_cast<int32>(Rhs.View == nullptr);
 	}
 
 	template<SortBy By>
@@ -362,8 +395,25 @@ namespace UE::Editor::DataStorage
 			}
 			else
 			{
-				Cache = View->ToString();
+				Cache.Reset();
+				if constexpr (EnumHasAllFlags(By::Flags, ESortByNameFlags::WithNone))
+				{
+					View->ToString(Cache);
+				}
+				else
+				{
+					if (!View->IsNone())
+					{
+						View->ToString(Cache);
+					}
+				}
+
+				if constexpr (EnumHasAllFlags(By::Flags, ESortByNameFlags::RemoveLeadingSlash))
+				{
+					Cache.RemoveFromStart(TEXT("/"));
+				}
 			}
+			bIsCached = true;
 		}
 	}
 
diff --git a/Engine/Source/Runtime/TypedElementFramework/Public/Elements/Interfaces/TypedElementDataStorageFactory.h b/Engine/Source/Runtime/TypedElementFramework/Public/Elements/Interfaces/TypedElementDataStorageFactory.h
--- a/Engine/Source/Runtime/TypedElementFramework/Public/Elements/Interfaces/TypedElementDataStorageFactory.h
+++ b/Engine/Source/Runtime/TypedElementFramework/Public/Elements/Interfaces/TypedElementDataStorageFactory.h
@@ -46,6 +46,13 @@ public:
 	virtual void RegisterWidgetConstructors(UE::Editor::DataStorage::ICoreProvider& DataStorage,
 		UE::Editor::DataStorage::IUiProvider& DataStorageUi) const {}
 
+	/** 
+	 * Used to register sorters for a specific property type to be used in the default generated property sorter. The widget constructor
+	 * can used to provide a custom sorter if more complex sorting logic is needed. Implementations of this function should mostly be
+	 * calling `RegisterSorterGeneratorForProperty` on the UI provider.
+	 */
+	virtual void RegisterPropertySorters(UE::Editor::DataStorage::IUiProvider& DataStorageUi) const {}
+
 	/**
 	 * Called in reverse order before the DataStorage object is shut down
 	 */
diff --git a/Engine/Source/Runtime/TypedElementFramework/Public/Elements/Interfaces/TypedElementDataStorageUiInterface.h b/Engine/Source/Runtime/TypedElementFramework/Public/Elements/Interfaces/TypedElementDataStorageUiInterface.h
--- a/Engine/Source/Runtime/TypedElementFramework/Public/Elements/Interfaces/TypedElementDataStorageUiInterface.h
+++ b/Engine/Source/Runtime/TypedElementFramework/Public/Elements/Interfaces/TypedElementDataStorageUiInterface.h
@@ -2,6 +2,8 @@
 
 #pragma once
 
+#include "Containers/Array.h"
+#include "Containers/ArrayView.h"
 #include "Containers/ContainersFwd.h"
 #include "Containers/StringFwd.h"
 #include "Elements/Common/TypedElementQueryConditions.h"
@@ -13,6 +15,7 @@
 #include "Templates/UnrealTypeTraits.h"
 #include "Templates/SharedPointer.h"
 #include "UObject/Interface.h"
+#include "UObject/WeakObjectPtrTemplates.h"
 
 #include "TypedElementDataStorageUiInterface.generated.h"
 
@@ -102,6 +105,7 @@ public:
 	 */
 	TYPEDELEMENTFRAMEWORK_API virtual TArray<TSharedPtr<const UE::Editor::DataStorage::FColumnSorterInterface>> ConstructColumnSorters(
 		UE::Editor::DataStorage::ICoreProvider* DataStorage,
+		UE::Editor::DataStorage::IUiProvider* DataStorageUi,
 		const UE::Editor::DataStorage::FMetaDataView& Arguments);
 
 	/** Initializes a new constructor based on the provided arguments.. */
@@ -333,6 +337,16 @@ namespace UE::Editor::DataStorage
 		using WidgetCreatedCallback = TFunctionRef<void(const TSharedRef<SWidget>& NewWidget, UE::Editor::DataStorage::RowHandle Row)>;
 		using WidgetConstructorCallback = TFunctionRef<bool(TUniquePtr<FTypedElementWidgetConstructor>, TConstArrayView<TWeakObjectPtr<const UScriptStruct>>)>;
 		using WidgetPurposeCallback = TFunctionRef<void(FName, EPurposeType, const FText&)>;
+		using PropertySorterConstructorCallback = 
+			TFunction<TSharedPtr<const FColumnSorterInterface>(TWeakObjectPtr<const UScriptStruct> ColumnType, const FProperty& Property)>;
+
+		/**
+		 * @section Widget purposes
+		 * Widget purposes indicates how widgets can be used and categorizes/organizes the available widget factories. Purposes as used to 
+		 * match columns to widget constructors and provide contextual information on how the widgets get setup. In practice this allows 
+		 * tools to provide what purposes they support and TEDS UI will find the closest matching widget constructors that are compatible
+		 * with the purpose.
+		 */
 
 		/**
 		 * Register a widget purpose. Widget purposes indicates how widgets can be used and categorizes/organizes the available 
@@ -351,7 +365,27 @@ namespace UE::Editor::DataStorage
 		 * @param InPurposeInfo Struct with init params for the purpose
 		 */
 		virtual RowHandle RegisterWidgetPurpose(const FPurposeID& PurposeID, const FPurposeInfo& InPurposeInfo) = 0;
+
+		/** Find the row handle for a purpose by looking it up using the purpose ID */
+		virtual RowHandle FindPurpose(const FPurposeID& PurposeID) const = 0;
+
+		/** Calls the provided callback for all known registered widget purposes. */
+		virtual void ListWidgetPurposes(const WidgetPurposeCallback& Callback) const = 0;
+
+
+		/** Get the name of the default TEDS UI widget purpose used to register default widgets for different types of data (e.g FText -> STextBlock) */
+		virtual FPurposeID GetDefaultWidgetPurposeID() const = 0;
+
+		/** Get the name of the general TEDS UI purpose used to register general purpose widgets for columns */
+		virtual FPurposeID GetGeneralWidgetPurposeID() const = 0;
 	
+		/**
+		 * @section Widget constructors.
+		 * Widget constructors are the glue between TEDS data and Slate. They are associated with specific columns and provide ways
+		 * to generate Slate widgets for those columns. They also facilitate communication back-and-forth between TEDS and Slate.
+		 * Widget purposes can be used to specify what widget constructor to use when, allowing multiple widget constructors to
+		 * work on the same columns. Through the meta data object individual widgets can be further specialized on an individual baseis.
+		 */
 	
 		/** 
 		 * Registers a widget factory that will be called when the purpose it's registered under is requested.
@@ -454,29 +488,55 @@ namespace UE::Editor::DataStorage
 		virtual TSharedPtr<SWidget> ConstructWidget(RowHandle Row, FTypedElementWidgetConstructor& Constructor,
 			const FMetaDataView& Arguments) = 0;
 
-		/** Calls the provided callback for all known registered widget purposes. */
-		virtual void ListWidgetPurposes(const WidgetPurposeCallback& Callback) const = 0;
-
-		/** Check if a custom extension is supported. This can be used to check for in-development features, custom extensions, etc. */
-		virtual bool SupportsExtension(FName Extension) const = 0;
-		/** Provides a list of all extensions that are enabled. */
-		virtual void ListExtensions(TFunctionRef<void(FName)> Callback) const = 0;
-
 		/** Create the container widget that every TEDS UI widget is stored in */
 		virtual TSharedPtr<ITedsWidget> CreateContainerTedsWidget(RowHandle UiRowHandle) const = 0;
 
 		/** Get the table where TEDS UI widgets are stored */
 		virtual TableHandle GetWidgetTable() const = 0;
 
-		/** Find the row handle for a purpose by looking it up using the purpose ID */
-		virtual RowHandle FindPurpose(const FPurposeID& PurposeID) const = 0;
-	
-		/** Get the name of the default TEDS UI widget purpose used to register default widgets for different types of data (e.g FText -> STextBlock) */
-		virtual FPurposeID GetDefaultWidgetPurposeID() const = 0;
-	
-		/** Get the name of the general TEDS UI purpose used to register general purpose widgets for columns */
-		virtual FPurposeID GetGeneralWidgetPurposeID() const = 0;
-	
+		/**
+		 * @section Sorter generation
+		 * Widget constructors can provide a custom sorter object that will allows Slate columns to sort by the data that the sorter
+		 * supports. If none is provided however, TEDS UI will generate one based on the properties in a column that are marked with
+		 * the meta data "Sortable". Out of the box TEDS UI supports a variety of property types, but additional properties can be 
+		 * registered with the functions below as well. Note that this is mechanic is exclusively for properties that are used in
+		 * the dynamic generation of sorters. For complex sorting that involves data specific knowledge or combines multiple values
+		 * it's highly recommended to create a custom sorter and attach it to the appropriate widget constructor.
+		 */
+
+		/** Build sorters for the provided list of columns */
+		virtual void GeneratePropertySorters(TArray<TSharedPtr<const FColumnSorterInterface>>& Results,
+			TArrayView<TWeakObjectPtr<const UScriptStruct>> Columns) const = 0;
+
+		/** Register a function to handle converting a property into a sorter. */
+		virtual void RegisterSorterGeneratorForProperty(const FFieldClass* PropertyType, 
+			PropertySorterConstructorCallback PropertySorterConstructor) = 0;
+
+		/** Register a function to handle converting a property into a sorter. */
+		template<typename PropertyClassType> requires(std::is_base_of_v<FProperty, PropertyClassType>)
+			void RegisterSorterGeneratorForProperty(PropertySorterConstructorCallback PropertySorterConstructor);
+
+		/** Unregister a function to handle converting a property into a sorter. */
+		virtual void UnregisterSorterGeneratorForProperty(const FFieldClass* PropertyType) = 0;
+
+		/** Unregister a function to handle converting a property into a sorter. */
+		template<typename PropertyClassType> requires(std::is_base_of_v<FProperty, PropertyClassType>)
+			void UnregisterSorterGeneratorForProperty();
+
+		/**
+		 * @section Miscellaneous
+		 */
+
+		/** Check if a custom extension is supported. This can be used to check for in-development features, custom extensions, etc. */
+		virtual bool SupportsExtension(FName Extension) const = 0;
+		/** Provides a list of all extensions that are enabled. */
+		virtual void ListExtensions(TFunctionRef<void(FName)> Callback) const = 0;
+
+		/**
+		 * @section Deprecated
+		 * Please replace use of the functions below with the recommended replacements.
+		 */
+		
 		UE_DEPRECATED(5.6, "Use the version that takes in an FPurposeInfo struct instead")
 		virtual void RegisterWidgetPurpose(FName Purpose, EPurposeType Type, FText Description) = 0;
 	
@@ -547,4 +607,16 @@ namespace UE::Editor::DataStorage
 	{
 		return this->RegisterWidgetFactory(PurposeRow, ConstructorType::StaticStruct(), Columns);
 	}
+
+	template<typename PropertyClassType> requires(std::is_base_of_v<FProperty, PropertyClassType>)
+		void IUiProvider::RegisterSorterGeneratorForProperty(PropertySorterConstructorCallback PropertySorterConstructor)
+	{
+		RegisterSorterGeneratorForProperty(PropertyClassType::StaticClass(), PropertySorterConstructor);
+	}
+
+	template<typename PropertyClassType> requires(std::is_base_of_v<FProperty, PropertyClassType>)
+		void IUiProvider::UnregisterSorterGeneratorForProperty()
+	{
+		UnregisterSorterGeneratorForProperty(PropertyClassType::StaticClass());
+	}
 }
\ No newline at end of file
diff --git a/Engine/Source/Runtime/TypedElementFramework/Tests/Elements/Framework/TypedElementDataStorageUiSorterTests.cpp b/Engine/Source/Runtime/TypedElementFramework/Tests/Elements/Framework/TypedElementDataStorageUiSorterTests.cpp
--- a/Engine/Source/Runtime/TypedElementFramework/Tests/Elements/Framework/TypedElementDataStorageUiSorterTests.cpp
+++ b/Engine/Source/Runtime/TypedElementFramework/Tests/Elements/Framework/TypedElementDataStorageUiSorterTests.cpp
@@ -761,7 +761,7 @@ namespace UE::Editor::DataStorage::Tests
 		if (ICoreProvider* DataStorage = GetMutableDataStorageFeature<ICoreProvider>(StorageFeatureName))
 		{
 			TArray<FName> Names{ "a", "bba", "Aba", "cba", "ABac", "LongName" };
-			FSingleValueTestSorter<TSortNameView<FSortByName>> Sorter;
+			FSingleValueTestSorter<TSortNameView<TSortByName<>>> Sorter;
 			Sorter.Values.Append(Names);
 			Sorter.SortByPrefix(*DataStorage, 0);
 			CHECK(Sorter.Verify(*DataStorage));
