[TEDS] Updated the TEDS Table Viewer to support sorting and the TEDS Debugger to use the new Query Stack nodes for monitoring queries.

diff --git a/Engine/Plugins/Experimental/EditorDataStorageFeatures/Source/TedsDebugger/Private/Widgets/QueryEditor/TedsQueryEditorResultsView.cpp b/Engine/Plugins/Experimental/EditorDataStorageFeatures/Source/TedsDebugger/Private/Widgets/QueryEditor/TedsQueryEditorResultsView.cpp
--- a/Engine/Plugins/Experimental/EditorDataStorageFeatures/Source/TedsDebugger/Private/Widgets/QueryEditor/TedsQueryEditorResultsView.cpp
+++ b/Engine/Plugins/Experimental/EditorDataStorageFeatures/Source/TedsDebugger/Private/Widgets/QueryEditor/TedsQueryEditorResultsView.cpp
@@ -6,13 +6,18 @@
 #include "TedsOutlinerModule.h"
 #include "TedsTableViewerColumn.h"
 #include "Elements/Common/EditorDataStorageFeatures.h"
+#include "Elements/Common/TypedElementDataStorageLog.h"
 #include "Elements/Framework/TypedElementQueryBuilder.h"
 #include "Elements/Interfaces/TypedElementDataStorageInterface.h"
 #include "Elements/Interfaces/TypedElementDataStorageUiInterface.h"
 #include "Modules/ModuleManager.h"
 #include "QueryEditor/TedsQueryEditorModel.h"
+#include "TedsRowMonitorNode.h"
+#include "TedsRowQueryResultsNode.h"
 #include "TedsRowViewNode.h"
 #include "Elements/Columns/TypedElementHiearchyColumns.h"
+#include "TedsQueryNode.h"
+#include "TedsRowSortNode.h"
 #include "Widgets/STedsTableViewer.h"
 #include "Widgets/SRowDetails.h"
 #include "Widgets/Text/STextBlock.h"
@@ -28,19 +33,23 @@ namespace UE::Editor::DataStorage::Debug::QueryEditor
 
 	void SResultsView::Construct(const FArguments& InArgs, FTedsQueryEditorModel& InModel)
 	{
+		using namespace UE::Editor::DataStorage::QueryStack;
+		
 		Model = &InModel;
 		ModelChangedDelegateHandle = Model->GetModelChangedDelegate().AddRaw(this, &SResultsView::OnModelChanged);
 		UiProvider = GetMutableDataStorageFeature<IUiProvider>(UiFeatureName);
 		checkf(UiProvider, TEXT("Cannot create SResultsView without TEDS UI"));
 
-		// RowQueryHandle = InvalidQueryHandle;
-
+		
 		// Create a custom column for the table viewer to display row handles
 		CreateRowHandleColumn();
 	
-		RowQueryStack = MakeShared<QueryStack::FRowViewNode>(FRowHandleArrayView(TableViewerRows.GetData(), TableViewerRows.Num(),
-			FRowHandleArrayView::EFlags::IsSorted | FRowHandleArrayView::EFlags::IsUnique));
+		QueryNode = MakeShared<FQueryNode>(Model->GetTedsInterface());
 
+		RowQueryStack = MakeShared<FRowQueryResultsNode>(Model->GetTedsInterface(), QueryNode,
+			FRowQueryResultsNode::ESyncFlags::RefreshOnQueryChange | FRowQueryResultsNode::ESyncFlags::IncrementWhenDifferent);
+		RowQueryStack = MakeShared<FRowMonitorNode>(Model->GetTedsInterface(), RowQueryStack, QueryNode);
+		
 		ChildSlot
 		[
 			SNew(SVerticalBox)
@@ -66,7 +75,6 @@ namespace UE::Editor::DataStorage::Debug::QueryEditor
 									{
 										RowDetailsWidget->ClearRow();
 									}
-							
 								}
 							}))
 				]
@@ -145,47 +153,19 @@ namespace UE::Editor::DataStorage::Debug::QueryEditor
 					TableViewer->AddCustomColumn(RowHandleColumn.ToSharedRef());
 				}
 		
-				if (TableViewerQueryHandle != InvalidQueryHandle)
-				{
-					TedsInterface.UnregisterQuery(TableViewerQueryHandle);
-					TableViewerQueryHandle = InvalidQueryHandle;
-				}
-
 				// Mass doesn't like empty queries, so we only set it if there are actual conditions
 				if(TableViewerQueryDescription.ConditionTypes.Num() || TableViewerQueryDescription.SelectionTypes.Num())
 				{
-					TableViewerQueryHandle = TedsInterface.RegisterQuery(MoveTemp(TableViewerQueryDescription));
+					QueryNode->SetQuery(MoveTemp(TableViewerQueryDescription));
+				}
+				else
+				{
+					QueryNode->ClearQuery();
 				}
 			}
 	
 			bModelDirty = false;
 		}
-
-		TSet<RowHandle> NewTableViewerRows_Set;
-	
-		// Every frame we re-run the query to update the rows the table viewer is showing
-		if(TableViewerQueryHandle != InvalidQueryHandle)
-		{
-			NewTableViewerRows_Set.Reserve(TableViewerRows_Set.Num());
-
-			FQueryResult QueryResult = Model->GetTedsInterface().RunQuery(TableViewerQueryHandle,
-				CreateDirectQueryCallbackBinding([&NewTableViewerRows_Set](const IDirectQueryContext& Context, const RowHandle*)
-				{
-					NewTableViewerRows_Set.Append(Context.GetRowHandles());
-				}));
-		}
-	
-		// Check if the two sets are equal, i.e not changes and no need to update the table viewer
-		const bool bSetsEqual = (TableViewerRows_Set.Num() == NewTableViewerRows_Set.Num()) && TableViewerRows_Set.Includes(NewTableViewerRows_Set);
-
-		if(!bSetsEqual)
-		{
-			Swap(TableViewerRows_Set, NewTableViewerRows_Set);
-			TableViewerRows = TableViewerRows_Set.Array();
-
-			RowQueryStack->ResetView(FRowHandleArrayView(TableViewerRows.GetData(), TableViewerRows.Num(),
-				FRowHandleArrayView::EFlags::IsSorted | FRowHandleArrayView::EFlags::IsUnique));
-		}
 	}
 
 	void SResultsView::OnModelChanged()
diff --git a/Engine/Plugins/Experimental/EditorDataStorageFeatures/Source/TedsDebugger/Private/Widgets/QueryEditor/TedsQueryEditorResultsView.h b/Engine/Plugins/Experimental/EditorDataStorageFeatures/Source/TedsDebugger/Private/Widgets/QueryEditor/TedsQueryEditorResultsView.h
--- a/Engine/Plugins/Experimental/EditorDataStorageFeatures/Source/TedsDebugger/Private/Widgets/QueryEditor/TedsQueryEditorResultsView.h
+++ b/Engine/Plugins/Experimental/EditorDataStorageFeatures/Source/TedsDebugger/Private/Widgets/QueryEditor/TedsQueryEditorResultsView.h
@@ -19,7 +19,8 @@ namespace UE::Editor::DataStorage
 
 	namespace QueryStack
 	{
-		class FRowViewNode;
+		class FQueryNode;
+		class IRowNode;
 	}
 
 	namespace Debug::QueryEditor
@@ -47,14 +48,12 @@ namespace UE::Editor::DataStorage
 
 
 			QueryHandle CountQueryHandle = InvalidQueryHandle;
-			QueryHandle TableViewerQueryHandle = InvalidQueryHandle;
-
+			
 			TArray<RowHandle> TableViewerRows;
-			// We have to keep a TSet copy to have a sorted order for the rows for now
-			TSet<RowHandle> TableViewerRows_Set;
 			TSharedPtr<STedsTableViewer> TableViewer;
-			TSharedPtr<QueryStack::FRowViewNode> RowQueryStack;
-
+			TSharedPtr<QueryStack::FQueryNode> QueryNode;
+			TSharedPtr<QueryStack::IRowNode> RowQueryStack;
+			
 			// Custom column for the table viewer to display row handles
 			TSharedPtr<FTedsTableViewerColumn> RowHandleColumn;
 
diff --git a/Engine/Plugins/Experimental/EditorDataStorageFeatures/Source/TedsTableViewer/Private/TedsTableViewerColumn.cpp b/Engine/Plugins/Experimental/EditorDataStorageFeatures/Source/TedsTableViewer/Private/TedsTableViewerColumn.cpp
--- a/Engine/Plugins/Experimental/EditorDataStorageFeatures/Source/TedsTableViewer/Private/TedsTableViewerColumn.cpp
+++ b/Engine/Plugins/Experimental/EditorDataStorageFeatures/Source/TedsTableViewer/Private/TedsTableViewerColumn.cpp
@@ -33,15 +33,26 @@ namespace UE::Editor::DataStorage
 		StorageCompatibility = GetMutableDataStorageFeature<ICompatibilityProvider>(CompatibilityFeatureName);
 
 		// Store the matched columns as a query condition that requires all of them (i.e AND's them)
-		for(const TWeakObjectPtr<const UScriptStruct>& Column : MatchedColumns)
+		for (const TWeakObjectPtr<const UScriptStruct>& Column : MatchedColumns)
 		{
 			MatchedColumnConditions = MatchedColumnConditions && Queries::TColumn(Column);
 		}
 		MatchedColumnConditions.Compile(Queries::FEditorStorageQueryConditionCompileContext(Storage));
 
+		// Retrieve sort functors to allow the content of the column to be used for sorting
+		if (HeaderWidgetConstructor)
+		{
+			ColumnSorters = HeaderWidgetConstructor->ConstructColumnSorters(Storage, StorageUi, InWidgetMetaData);
+		}
+		// Prefer to use the sort function in the header, but if it doesn't exist, use the one from the cell.
+		if (CellWidgetConstructor && ColumnSorters.IsEmpty())
+		{
+			ColumnSorters = CellWidgetConstructor->ConstructColumnSorters(Storage, StorageUi, InWidgetMetaData);
+		}
+
 		RegisterQueries();
 	}
-	
+
 	FTedsTableViewerColumn::~FTedsTableViewerColumn()
 	{
 		UnRegisterQueries();
@@ -50,23 +61,23 @@ namespace UE::Editor::DataStorage
 	TSharedPtr<SWidget> FTedsTableViewerColumn::ConstructRowWidget(RowHandle InRowHandle, TFunction<void(ICoreProvider&, const RowHandle&)> WidgetRowSetupDelegate/* = nullptr*/) const
 	{
 		TSharedPtr<SWidget> RowWidget;
-		
-		if(Storage->IsRowAssigned(InRowHandle))
+
+		if (Storage->IsRowAssigned(InRowHandle))
 		{
 			const RowHandle UiRowHandle = Storage->AddRow(Storage->FindTable(TableViewerUtils::GetWidgetTableName()));
 
 			const TConstArrayView<TWeakObjectPtr<const UScriptStruct>> ColumnTypes = GetMatchedColumns();
 			if (ColumnTypes.Num() == 1)
 			{
-				Storage->AddColumn<FTypedElementScriptStructTypeInfoColumn>(UiRowHandle, FTypedElementScriptStructTypeInfoColumn{.TypeInfo = *ColumnTypes.begin()});
+				Storage->AddColumn<FTypedElementScriptStructTypeInfoColumn>(UiRowHandle, FTypedElementScriptStructTypeInfoColumn{ .TypeInfo = *ColumnTypes.begin() });
 			}
-			
+
 			if (FTypedElementRowReferenceColumn* RowReference = Storage->GetColumn<FTypedElementRowReferenceColumn>(UiRowHandle))
 			{
 				RowReference->Row = InRowHandle;
 			}
 
-			if(FTypedElementSlateWidgetReferenceColumn* WidgetReferenceColumn = Storage->GetColumn<FTypedElementSlateWidgetReferenceColumn>(UiRowHandle))
+			if (FTypedElementSlateWidgetReferenceColumn* WidgetReferenceColumn = Storage->GetColumn<FTypedElementSlateWidgetReferenceColumn>(UiRowHandle))
 			{
 				WidgetReferenceColumn->WidgetConstructor = CellWidgetConstructor;
 			}
@@ -82,10 +93,10 @@ namespace UE::Editor::DataStorage
 		return RowWidget;
 	}
 
-	SHeaderRow::FColumn::FArguments FTedsTableViewerColumn::ConstructHeaderRowColumn() const
+	SHeaderRow::FColumn::FArguments FTedsTableViewerColumn::ConstructHeaderRowColumn()
 	{
 		const TConstArrayView<TWeakObjectPtr<const UScriptStruct>> ColumnTypes = GetMatchedColumns();
-		
+
 		FString TooltipText = TEXT("Data Storage columns:");
 		for (const TWeakObjectPtr<const UScriptStruct>& ColumnType : ColumnTypes)
 		{
@@ -105,12 +116,12 @@ namespace UE::Editor::DataStorage
 			// TEDS UI TODO: We can't do this from the Widget Constructor because it is a UStruct and does not have access to AsShared(), so we would
 			// be forced to store a raw pointer instead of a weak pointer which is unsafe. Once the widget construction pipleline is improved this can
 			// probably be moved to a better place
-			if(FTypedElementSlateWidgetReferenceColumn* WidgetReferenceColumn = Storage->GetColumn<FTypedElementSlateWidgetReferenceColumn>(UiRowHandle))
+			if (FTypedElementSlateWidgetReferenceColumn* WidgetReferenceColumn = Storage->GetColumn<FTypedElementSlateWidgetReferenceColumn>(UiRowHandle))
 			{
 				WidgetReferenceColumn->WidgetConstructor = HeaderWidgetConstructor;
 			}
-			
-			Widget = StorageUi->ConstructWidget(UiRowHandle, *HeaderWidgetConstructor, 
+
+			Widget = StorageUi->ConstructWidget(UiRowHandle, *HeaderWidgetConstructor,
 				WidgetMetaData);
 		}
 		if (!Widget.IsValid())
@@ -121,35 +132,51 @@ namespace UE::Editor::DataStorage
 			{
 				DefaultHeaderText = FText::FromString(ColumnName.ToString());
 			}
-			
+
 			Widget = SNew(STextBlock)
 				.Text(DefaultHeaderText);
 		}
-		
+
 		SHeaderRow::FColumn::FArguments Column = SHeaderRow::Column(ColumnName)
 			.FillWidth(1)
 			.HeaderComboVisibility(EHeaderComboVisibility::OnHover)
 			.DefaultTooltip(FText::FromString(MoveTemp(TooltipText)))
 			.DefaultLabel(FText::FromName(ColumnName))
+			.SortMode_Raw(this, &FTedsTableViewerColumn::GetSortMode)
+			.SortPriority_Raw(this, &FTedsTableViewerColumn::GetSortPriority)
+			.IsSorting_Raw(this, &FTedsTableViewerColumn::IsSorting)
+			.OnSort_Raw(this, &FTedsTableViewerColumn::OnSortCallback)
 			.HeaderContent()
 			[
-				SNew(SBox)
-					.MinDesiredHeight(20.0f)
-					.VAlign(VAlign_Center)
+				SNew(SHorizontalBox)
+				+SHorizontalBox::Slot()
+					.FillWidth(1.0f)
 					[
-						Widget.ToSharedRef()
+						SNew(SBox)
+							.MinDesiredHeight(20.0f)
+							.VAlign(VAlign_Center)
+							[
+								Widget.ToSharedRef()
+							]
+					]
+				+SHorizontalBox::Slot()
+					.AutoWidth()
+					[
+						SNew(STextBlock)
+							.Text_Raw(this, &FTedsTableViewerColumn::GetSortText)
+							.Visibility_Raw(this, &FTedsTableViewerColumn::GetSortTextVisibility)
 					]
 			];
-		
+
 		if (const FHeaderWidgetSizeColumn* HeaderProperties = Storage->GetColumn<FHeaderWidgetSizeColumn>(UiRowHandle))
 		{
 			float Width = HeaderProperties->Width;
 			switch (HeaderProperties->ColumnSizeMode)
 			{
-				case EColumnSizeMode::Fill: Column.FillWidth(Width); break;
-				case EColumnSizeMode::Fixed: Column.FixedWidth(Width); break;
-				case EColumnSizeMode::Manual: Column.ManualWidth(Width); break;
-				case EColumnSizeMode::FillSized: Column.FillSized(Width); break;
+			case EColumnSizeMode::Fill: Column.FillWidth(Width); break;
+			case EColumnSizeMode::Fixed: Column.FixedWidth(Width); break;
+			case EColumnSizeMode::Manual: Column.ManualWidth(Width); break;
+			case EColumnSizeMode::FillSized: Column.FillSized(Width); break;
 			}
 		}
 		return Column;
@@ -158,16 +185,29 @@ namespace UE::Editor::DataStorage
 	void FTedsTableViewerColumn::Tick()
 	{
 		// Update any rows that could need widget updates
-		if(!RowsToUpdate.IsEmpty())
+		if (!RowsToUpdate.IsEmpty())
 		{
 			UpdateWidgets();
 			RowsToUpdate.Empty();
 		}
 	}
 
-	void FTedsTableViewerColumn::SetIsRowVisibleDelegate(const FIsRowVisible& InIsRowVisibleDelegate)
+	void FTedsTableViewerColumn::SetIsRowVisibleDelegate(FIsRowVisible InIsRowVisibleDelegate)
 	{
-		IsRowVisibleDelegate = InIsRowVisibleDelegate;
+		IsRowVisibleDelegate = MoveTemp(InIsRowVisibleDelegate);
+	}
+
+	void FTedsTableViewerColumn::SetSortDelegates(FSortHandler InSortDelegates)
+	{
+		SortDelegates = MoveTemp(InSortDelegates);
+	}
+
+	void FTedsTableViewerColumn::ClearSort()
+	{
+		SorterShortName = FText::GetEmpty();
+		ColumnSorterIndex = 0;
+		bIsSorted = false;
+		SortPriority = EColumnSortPriority::None;
 	}
 
 	void FTedsTableViewerColumn::RegisterQueries()
@@ -177,9 +217,9 @@ namespace UE::Editor::DataStorage
 		const TConstArrayView<TWeakObjectPtr<const UScriptStruct>> ColumnTypes = GetMatchedColumns();
 		
 		// For each TEDS column this column is matched with, we'll add observers to track addition/removal to update any widgets
-		for(const TWeakObjectPtr<const UScriptStruct>& ColumnType : ColumnTypes)
+		for (const TWeakObjectPtr<const UScriptStruct>& ColumnType : ColumnTypes)
 		{
-			const FName ColumnAddObserverName = *(FString::Printf(TEXT("Column Add Monitor for %s Table Viewer Column Column, %s TEDS Column"), *ColumnName.ToString(), *ColumnType->GetName()));
+			const FName ColumnAddObserverName = *(FString::Printf(TEXT("Column Add Monitor for %s Table Viewer Column, %s TEDS Column"), *ColumnName.ToString(), *ColumnType->GetName()));
 			FObserver AddObserver(FObserver::EEvent::Add, ColumnType.Get());
 			AddObserver.SetExecutionMode(EExecutionMode::GameThread);
 
@@ -190,16 +230,16 @@ namespace UE::Editor::DataStorage
 					ColumnAddObserverName,
 					AddObserver,
 					[this](IQueryContext& Context, RowHandle Row)
-						{
-							RowsToUpdate.Add(TPair<RowHandle, bool>(Row, true));
-						})
+					{
+						RowsToUpdate.Add(TPair<RowHandle, bool>(Row, true));
+					})
 				.Where()
-					.All(ColumnType.Get())
+				.All(ColumnType.Get())
 				.Compile());
-			
+
 			InternalObserverQueries.Add(AddQueryHandle);
 
-			const FName ColumnRemoveObserverName = *(FString::Printf(TEXT("Column Remove Monitor for %s Table Viewer Column Column, %s TEDS Column"), *ColumnName.ToString(), *ColumnType->GetName()));
+			const FName ColumnRemoveObserverName = *(FString::Printf(TEXT("Column Remove Monitor for %s Table Viewer Column, %s TEDS Column"), *ColumnName.ToString(), *ColumnType->GetName()));
 			FObserver RemoveObserver(FObserver::EEvent::Remove, ColumnType.Get());
 			RemoveObserver.SetExecutionMode(EExecutionMode::GameThread);
 
@@ -210,18 +250,18 @@ namespace UE::Editor::DataStorage
 					ColumnRemoveObserverName,
 					RemoveObserver,
 					[this](IQueryContext& Context, RowHandle Row)
-						{
-							RowsToUpdate.Add(TPair<RowHandle, bool>(Row, false));
-						})
+					{
+						RowsToUpdate.Add(TPair<RowHandle, bool>(Row, false));
+					})
 				.Where()
-					.All(ColumnType.Get())
+				.All(ColumnType.Get())
 				.Compile());
-			
+
 			InternalObserverQueries.Add(RemoveQueryHandle);
 		}
 
 		// We are looking for widgets that have a row reference
-		TArray<const UScriptStruct*> SelectionColumns({FTypedElementSlateWidgetReferenceColumn::StaticStruct(), FTypedElementRowReferenceColumn::StaticStruct()});
+		TArray<const UScriptStruct*> SelectionColumns({ FTypedElementSlateWidgetReferenceColumn::StaticStruct(), FTypedElementRowReferenceColumn::StaticStruct() });
 
 		// Query to get all widgets that were created by this column
 		WidgetQuery = Storage->RegisterQuery(
@@ -361,5 +401,58 @@ namespace UE::Editor::DataStorage
 	{
 		return Storage;
 	}
+
+	EColumnSortMode::Type FTedsTableViewerColumn::GetSortMode() const
+	{
+		return bIsSorted 
+			? ((ColumnSorterIndex & 1) == 0 ? EColumnSortMode::Type::Descending : EColumnSortMode::Type::Ascending)
+			: EColumnSortMode::Type::None;
+	}
+
+	EColumnSortPriority::Type FTedsTableViewerColumn::GetSortPriority() const
+	{
+		return SortPriority;
+	}
+
+	bool FTedsTableViewerColumn::IsSorting() const
+	{
+		return (bIsSorted && SortDelegates.IsSortingHandler.IsBound()) ? SortDelegates.IsSortingHandler.Execute() : false;
+	}
+
+	EVisibility FTedsTableViewerColumn::GetSortTextVisibility() const
+	{
+		return (bIsSorted && !SorterShortName.IsEmpty()) ? EVisibility::Visible : EVisibility::Collapsed;
+	}
+
+	FText FTedsTableViewerColumn::GetSortText() const
+	{
+		return SorterShortName;
+	}
+
+	void FTedsTableViewerColumn::OnSortCallback(
+		EColumnSortPriority::Type InSortPriority, const FName& InColumnId, EColumnSortMode::Type InSortDirection)
+	{
+		if (!ColumnSorters.IsEmpty() && InColumnId == ColumnName)
+		{
+			bIsSorted = true;
+
+			SortPriority = (InSortPriority == EColumnSortPriority::Secondary) ? EColumnSortPriority::Secondary : EColumnSortPriority::Primary;
+
+			TSharedPtr<const FColumnSorterInterface> Sorter = ColumnSorters[ColumnSorterIndex >> 1];
+			
+			if (ColumnSorters.Num() > 1)
+			{
+				FText ShortName = Sorter->GetShortName();
+				SorterShortName = ShortName.IsEmpty()
+					? FText::GetEmpty()
+					: FText::Format(NSLOCTEXT("TedsTableViewColumn", "SorterName", " - {0}"), ShortName);
+			}
+			
+			SortDelegates.OnSortHandler.ExecuteIfBound(ColumnName, Sorter, InSortDirection, SortPriority);
+
+			// Each column sorter is used twice, once for ascending and once for descending sorts.
+			ColumnSorterIndex = (ColumnSorterIndex + 1) % (ColumnSorters.Num() * 2);
+		}
+	}
 }
 
diff --git a/Engine/Plugins/Experimental/EditorDataStorageFeatures/Source/TedsTableViewer/Private/TedsTableViewerModel.cpp b/Engine/Plugins/Experimental/EditorDataStorageFeatures/Source/TedsTableViewer/Private/TedsTableViewerModel.cpp
--- a/Engine/Plugins/Experimental/EditorDataStorageFeatures/Source/TedsTableViewer/Private/TedsTableViewerModel.cpp
+++ b/Engine/Plugins/Experimental/EditorDataStorageFeatures/Source/TedsTableViewer/Private/TedsTableViewerModel.cpp
@@ -5,8 +5,9 @@
 #include "Elements/Columns/TypedElementUIColumns.h"
 #include "Elements/Common/EditorDataStorageFeatures.h"
 #include "Elements/Interfaces/TypedElementDataStorageInterface.h"
-#include "Elements/Interfaces/TypedElementDataStorageCompatibilityInterface.h"
 #include "Elements/Interfaces/TypedElementDataStorageUiInterface.h"
+#include "TedsRowOrderInversionNode.h"
+#include "TedsRowSortNode.h"
 #include "TedsTableViewerColumn.h"
 #include "TedsTableViewerUtils.h"
 
@@ -17,16 +18,21 @@ namespace UE::Editor::DataStorage
 	FTedsTableViewerModel::FTedsTableViewerModel(const TSharedPtr<QueryStack::IRowNode>& InRowQueryStack,
 		const TArray<TWeakObjectPtr<const UScriptStruct>>& InRequestedColumns, const IUiProvider::FPurposeID& InCellWidgetPurpose,
 		const IUiProvider::FPurposeID& InHeaderWidgetPurpose, const FIsItemVisible& InIsItemVisibleDelegate)
-		: RowQueryStack(InRowQueryStack)
-		, RequestedTedsColumns(InRequestedColumns)
+		: RequestedTedsColumns(InRequestedColumns)
 		, CellWidgetPurpose(InCellWidgetPurpose)
 		, HeaderWidgetPurpose(InHeaderWidgetPurpose)
 		, IsItemVisible(InIsItemVisibleDelegate)
 	{
 		Storage = GetMutableDataStorageFeature<ICoreProvider>(StorageFeatureName);
 		StorageUi = GetMutableDataStorageFeature<IUiProvider>(UiFeatureName);
-		StorageCompatibility = GetMutableDataStorageFeature<ICompatibilityProvider>(CompatibilityFeatureName);
 		
+		checkf(InRowQueryStack, TEXT("The TEDS Table Viewer requires a valid query stack to provide it rows to process."));
+		RowSecondarySortingNode = MakeShared<QueryStack::FRowSortNode>(*Storage, InRowQueryStack);
+		RowSecondaryInversionNode = MakeShared<QueryStack::FRowOrderInversionNode>(RowSecondarySortingNode, false);
+		RowPrimarySortingNode = MakeShared<QueryStack::FRowSortNode>(*Storage, RowSecondaryInversionNode);
+		RowPrimaryInversionNode = MakeShared<QueryStack::FRowOrderInversionNode>(RowPrimarySortingNode, false);
+		RowQueryStack = RowPrimaryInversionNode;
+
 		TickerHandle = FTSTicker::GetCoreTicker().AddTicker(FTickerDelegate::CreateRaw(this, &FTedsTableViewerModel::Tick), 0);
 
 		GenerateColumns();
@@ -80,6 +86,8 @@ namespace UE::Editor::DataStorage
 
 	bool FTedsTableViewerModel::Tick(float DeltaTime)
 	{
+		RowQueryStack->Update();
+
 		// If the revision ID has changed, refresh to update our rows
 		if(RowQueryStack->GetRevision() != CachedRowQueryStackRevision)
 		{
@@ -159,6 +167,11 @@ namespace UE::Editor::DataStorage
 		// Table Viewer TODO: We should allow users to specify sort order using a TEDS column on the UI row, but for now we put any custom
 		// columns on the front
 		ColumnsView.Insert(InColumn, 0);
+		InColumn->SetSortDelegates(FTedsTableViewerColumn::FSortHandler
+			{
+				.OnSortHandler = FTedsTableViewerColumn::FOnSort::CreateRaw(this, &FTedsTableViewerModel::OnSort),
+				.IsSortingHandler = FTedsTableViewerColumn::FIsSorting::CreateRaw(this, &FTedsTableViewerModel::IsSorting)
+			});
 	}
 
 	ICoreProvider* FTedsTableViewerModel::GetDataStorageInterface() const
@@ -204,6 +217,11 @@ namespace UE::Editor::DataStorage
 				TSharedRef<FTedsTableViewerColumn> Column = MakeShared<FTedsTableViewerColumn>(NameId, CellConstructor, MoveTemp(MatchedColumnsCopy), HeaderConstructor);
 			
 				Column->SetIsRowVisibleDelegate(FTedsTableViewerColumn::FIsRowVisible::CreateRaw(this, &FTedsTableViewerModel::IsRowVisible));
+				Column->SetSortDelegates(FTedsTableViewerColumn::FSortHandler
+					{
+						.OnSortHandler = FTedsTableViewerColumn::FOnSort::CreateRaw(this, &FTedsTableViewerModel::OnSort),
+						.IsSortingHandler = FTedsTableViewerColumn::FIsSorting::CreateRaw(this, &FTedsTableViewerModel::IsSorting)
+					});
 
 				for (const TWeakObjectPtr<const UScriptStruct>& ColumnType : MatchedColumns)
 				{
@@ -251,4 +269,67 @@ namespace UE::Editor::DataStorage
 
 		return true;
 	}
+
+	void FTedsTableViewerModel::OnSort(FName ColumnName, TSharedPtr<const FColumnSorterInterface> ColumnSorter, EColumnSortMode::Type Direction, 
+		EColumnSortPriority::Type Priority)
+	{
+		if (RowPrimarySortingNode && RowSecondarySortingNode && RowPrimaryInversionNode && RowSecondaryInversionNode)
+		{
+			if (Priority == EColumnSortPriority::Secondary)
+			{
+				if (Direction == EColumnSortMode::Type::Ascending)
+				{
+					RowSecondarySortingNode->SetColumnSorter(ColumnSorter);
+					RowSecondaryInversionNode->Enable(false);
+				}
+				else
+				{
+					// Since sorting always happens in ascending first and then descending, use this knowledge to avoid resorting the current
+					// list to get the exact same result and just invert the current list instead.
+					RowSecondaryInversionNode->Enable(true);
+				}
+				SecondarySortColumnName = ColumnName;
+			}
+			else
+			{
+				if (Direction == EColumnSortMode::Type::Ascending)
+				{
+					RowPrimarySortingNode->SetColumnSorter(ColumnSorter);
+					RowPrimaryInversionNode->Enable(false);
+				}
+				else
+				{
+					RowPrimaryInversionNode->Enable(true);
+				}
+				PrimarySortColumnName = ColumnName;
+				
+				// Reset the secondary sorter.
+				SecondarySortColumnName = NAME_None;
+				RowSecondarySortingNode->SetColumnSorter(nullptr);
+				RowSecondaryInversionNode->Enable(false);
+			}
+		}
+		else
+		{
+			// Reset so all columns get cleared.
+			PrimarySortColumnName = NAME_None;
+			SecondarySortColumnName = NAME_None;
+		}
+
+		for (const TSharedRef<FTedsTableViewerColumn>& Column : ColumnsView)
+		{
+			FName LocalColumnName = Column->GetColumnName();
+			if (LocalColumnName != PrimarySortColumnName && LocalColumnName != SecondarySortColumnName)
+			{
+				Column->ClearSort();
+			}
+		}
+	}
+
+	bool FTedsTableViewerModel::IsSorting() const
+	{
+		return (RowPrimarySortingNode && RowSecondarySortingNode) 
+			? (RowPrimarySortingNode->IsSorting() || RowSecondarySortingNode->IsSorting())
+			: false;
+	}
 } // namespace UE::Editor::DataStorage
diff --git a/Engine/Plugins/Experimental/EditorDataStorageFeatures/Source/TedsTableViewer/Public/TedsTableViewerColumn.h b/Engine/Plugins/Experimental/EditorDataStorageFeatures/Source/TedsTableViewer/Public/TedsTableViewerColumn.h
--- a/Engine/Plugins/Experimental/EditorDataStorageFeatures/Source/TedsTableViewer/Public/TedsTableViewerColumn.h
+++ b/Engine/Plugins/Experimental/EditorDataStorageFeatures/Source/TedsTableViewer/Public/TedsTableViewerColumn.h
@@ -4,6 +4,7 @@
 #include "Elements/Common/TypedElementHandles.h"
 #include "Elements/Framework/TypedElementMetaData.h"
 #include "Elements/Framework/TypedElementQueryBuilder.h"
+#include "Elements/Interfaces/TypedElementDataStorageUiInterface.h"
 #include "Templates/SharedPointer.h"
 #include "UObject/NameTypes.h"
 #include "Widgets/Views/SHeaderRow.h"
@@ -23,8 +24,10 @@ namespace UE::Editor::DataStorage
 	 */
 	class FTedsTableViewerColumn
 	{
-	public:
+		// Done so only the view model has access to setting up the sorting callbacks.
+		friend class FTedsTableViewerModel;
 
+	public:
 		// Delegate to check if a row is currently visible in the owning table viewer's UI
 		DECLARE_DELEGATE_RetVal_OneParam(bool, FIsRowVisible, const RowHandle);
 
@@ -40,12 +43,14 @@ namespace UE::Editor::DataStorage
 		TEDSTABLEVIEWER_API	TSharedPtr<SWidget> ConstructRowWidget(RowHandle InRowHandle, 
 			TFunction<void(ICoreProvider&, const RowHandle&)> WidgetRowSetupDelegate = nullptr) const;
 		
-		TEDSTABLEVIEWER_API SHeaderRow::FColumn::FArguments ConstructHeaderRowColumn() const;
-		
+		TEDSTABLEVIEWER_API SHeaderRow::FColumn::FArguments ConstructHeaderRowColumn();
+
 		TEDSTABLEVIEWER_API void Tick();
 		
-		TEDSTABLEVIEWER_API void SetIsRowVisibleDelegate(const FIsRowVisible& InIsRowVisibleDelegate);
+		TEDSTABLEVIEWER_API void SetIsRowVisibleDelegate(FIsRowVisible InIsRowVisibleDelegate);
 		
+		TEDSTABLEVIEWER_API void ClearSort();
+
 		TEDSTABLEVIEWER_API FName GetColumnName() const;
 
 		TEDSTABLEVIEWER_API TConstArrayView<TWeakObjectPtr<const UScriptStruct>> GetMatchedColumns() const;
@@ -59,6 +64,26 @@ namespace UE::Editor::DataStorage
 		void UpdateWidgets();
 		
 	private:
+		// Delegate triggered when a column is sorted.
+		DECLARE_DELEGATE_FourParams(FOnSort, FName ColumnName, TSharedPtr<const FColumnSorterInterface> ColumnSorter,
+			EColumnSortMode::Type Direction, EColumnSortPriority::Type Priority);
+		// Delegate triggered used to determine is sorting is in progress.
+		DECLARE_DELEGATE_RetVal(bool, FIsSorting);
+
+		struct FSortHandler
+		{
+			FOnSort OnSortHandler;
+			FIsSorting IsSortingHandler;
+		};
+
+		void SetSortDelegates(FSortHandler InSortDelegates);
+		EColumnSortMode::Type GetSortMode() const;
+		EColumnSortPriority::Type GetSortPriority() const;
+		bool IsSorting() const;
+		EVisibility GetSortTextVisibility() const;
+		FText GetSortText() const;
+		void OnSortCallback(EColumnSortPriority::Type InSortPriority, const FName& InColumnId, EColumnSortMode::Type InSortDirection);
+
 		// The ID of the column
 		FName ColumnName;
 
@@ -85,6 +110,14 @@ namespace UE::Editor::DataStorage
 		QueryHandle WidgetQuery;
 		TMap<RowHandle, bool> RowsToUpdate;
 
+		// Sorting information
+		TArray<TSharedPtr<const FColumnSorterInterface>> ColumnSorters;
+		FSortHandler SortDelegates;
+		FText SorterShortName;
+		uint32 ColumnSorterIndex = 0;
+		EColumnSortPriority::Type SortPriority = EColumnSortPriority::None;
+		bool bIsSorted = false;
+		
 		// Delegate to check if a row is visible in the owning table viewer
 		FIsRowVisible IsRowVisibleDelegate;
 	};
diff --git a/Engine/Plugins/Experimental/EditorDataStorageFeatures/Source/TedsTableViewer/Public/TedsTableViewerModel.h b/Engine/Plugins/Experimental/EditorDataStorageFeatures/Source/TedsTableViewer/Public/TedsTableViewerModel.h
--- a/Engine/Plugins/Experimental/EditorDataStorageFeatures/Source/TedsTableViewer/Public/TedsTableViewerModel.h
+++ b/Engine/Plugins/Experimental/EditorDataStorageFeatures/Source/TedsTableViewer/Public/TedsTableViewerModel.h
@@ -5,14 +5,20 @@
 #include "Containers/Ticker.h"
 #include "Delegates/Delegate.h"
 #include "Elements/Common/TypedElementHandles.h"
+#include "Elements/Interfaces/TypedElementDataStorageUiInterface.h"
+#include "TedsQueryStackInterfaces.h"
 #include "Templates/SharedPointer.h"
 #include "TypedElementUITypes.h"
-#include "TedsQueryStackInterfaces.h"
-#include "Elements/Interfaces/TypedElementDataStorageUiInterface.h"
 #include "UObject/NameTypes.h"
+#include "Widgets/Views/SHeaderRow.h"
 
 namespace UE::Editor::DataStorage
 {
+	namespace QueryStack
+	{
+		class FRowOrderInversionNode;
+		class FRowSortNode;
+	}
 	class FTedsTableViewerColumn;
 	class ICompatibilityProvider;
 	class IUiProvider;
@@ -77,6 +83,10 @@ namespace UE::Editor::DataStorage
 		// Check if the given row is currently visible in the UI
 		bool IsRowVisible(RowHandle InRowHandle) const;
 
+		void OnSort(FName ColumnName, TSharedPtr<const FColumnSorterInterface> ColumnSorter, EColumnSortMode::Type Direction, 
+			EColumnSortPriority::Type Priority);
+		bool IsSorting() const;
+		
 		bool Tick(float DeltaTime);
 
 		void Refresh();
@@ -88,6 +98,22 @@ namespace UE::Editor::DataStorage
 		// The row query stack used to supply the rows to display
 		TSharedPtr<QueryStack::IRowNode> RowQueryStack;
 
+		// Row node used to sort the rows in the query stack.
+		TSharedPtr<QueryStack::FRowSortNode> RowPrimarySortingNode;
+		// Row node used as a backup to sort the rows in the query stack. Users can hold down shift to setup a secondary sort. As a result
+		// any duplicates from the the primary sorting node will be sorted by the secondary sort.
+		TSharedPtr<QueryStack::FRowSortNode> RowSecondarySortingNode;
+
+		// Inverts the rows from the query stack when needed by the primary sort.
+		TSharedPtr<QueryStack::FRowOrderInversionNode> RowPrimaryInversionNode;
+		// Inverts the rows from the query stack when needed by the primary sort.
+		TSharedPtr<QueryStack::FRowOrderInversionNode> RowSecondaryInversionNode;
+
+		// The name of the active primary sort column or none if not set.
+		FName PrimarySortColumnName;
+		// The name of the active secondary sort column or none if not set.
+		FName SecondarySortColumnName;
+		
 		// The cached list of rows we are currently displaying
 		TArray<TableViewerItemPtr> Items;
 
@@ -115,6 +141,5 @@ namespace UE::Editor::DataStorage
 		// Teds Constructs
 		ICoreProvider* Storage = nullptr;
 		IUiProvider* StorageUi = nullptr;
-		ICompatibilityProvider* StorageCompatibility = nullptr;
 	};
 } // namespace UE::Editor::DataStorage
