Teds Outliner: When finding the parent for a row, only return the row if it is found in the query stack (i.e it matches the input query). This is to fix a bug where valid parent rows that did not match the outliner query would show up when they shouldn't.

diff --git a/Engine/Plugins/Experimental/EditorDataStorageFeatures/Source/TedsOutliner/Private/TedsOutlinerImpl.cpp b/Engine/Plugins/Experimental/EditorDataStorageFeatures/Source/TedsOutliner/Private/TedsOutlinerImpl.cpp
--- a/Engine/Plugins/Experimental/EditorDataStorageFeatures/Source/TedsOutliner/Private/TedsOutlinerImpl.cpp
+++ b/Engine/Plugins/Experimental/EditorDataStorageFeatures/Source/TedsOutliner/Private/TedsOutlinerImpl.cpp
@@ -18,6 +18,8 @@
 #include "TedsOutlinerItem.h"
 #include "TedsQueryMergeNode.h"
 #include "TedsQueryNode.h"
+#include "TedsRowCopyNode.h"
+#include "TedsRowHandleSortNode.h"
 #include "TedsRowMergeNode.h"
 #include "TedsRowMonitorNode.h"
 #include "TedsRowQueryResultsNode.h"
@@ -410,9 +412,11 @@ void FTedsOutlinerImpl::CreateItemsFromQuery(TArray<FSceneOutlinerTreeItemPtr>&
 {
 	using namespace UE::Editor::DataStorage::Queries;
 
-	// We only update the row collector node on demand, since this function is only called on a FullRefresh()
-	// Addition/Removal in between refreshes is handled by the row monitor node
-	RowCollectorNode->Update();
+	// We refresh the row collector node on demand, since this function is only called on a FullRefresh() and the node is setup to not do anything on
+	// update. Addition/Removal in between refreshes is handled by the row monitor node
+	// Technically the RowMonitorNode also handles keeping this list up to date using observers, but the monitor node is optional and it's safer
+	// to do a full clean refresh at this point anyways
+	RowCollectorNode->Refresh();
 	FRowHandleArrayView Rows = RowCollectorNode->GetRows();
 	
 	for (RowHandle Row : Rows)
@@ -549,12 +553,19 @@ DataStorage::RowHandle FTedsOutlinerImpl::GetParentRow(DataStorage::RowHandle In
 	{
 		return DataStorage::InvalidRowHandle;
 	}
-	
+
 	if (!CanDisplayRow(ParentRowHandle))
 	{
 		return DataStorage::InvalidRowHandle;
 	}
 
+	// If the row doesn't exist in our query stack results, it doesn't match the query and should not be displayed
+	DataStorage::FRowHandleArrayView Rows = RowHandleSortNode->GetRows();
+	if (!Rows.Contains(ParentRowHandle))
+	{
+		return DataStorage::InvalidRowHandle;
+	}
+
 	return ParentRowHandle;
 }
 
@@ -584,6 +595,10 @@ void FTedsOutlinerImpl::RecompileQueries()
 	// Node to collect all rows matching the composite query
 	RowCollectorNode = MakeShared<FRowQueryResultsNode>(*Storage, CombinedQueryNode, FRowQueryResultsNode::ESyncFlags::None);
 
+	// We store a copy of the collected rows into an array sorted by row handles for faster lookup
+	TSharedRef<FRowCopyNode> CopyNode = MakeShared<FRowCopyNode>(RowCollectorNode);
+	RowHandleSortNode = MakeShared<FRowHandleSortNode>(CopyNode);
+
 	// Node to monitor row addition/removal
 	if (CreationParams.bUseDefaultObservers)
 	{
@@ -833,6 +848,8 @@ void FTedsOutlinerImpl::Tick()
 	{
 		RowMonitorNode->Update();
 	}
+
+	RowHandleSortNode->Update();
 	
 	// Update the label for any rows that might need it
 	for (DataStorage::RowHandle Row : RowsPendingLabelUpdate)
diff --git a/Engine/Plugins/Experimental/EditorDataStorageFeatures/Source/TedsOutliner/Public/TedsOutlinerImpl.h b/Engine/Plugins/Experimental/EditorDataStorageFeatures/Source/TedsOutliner/Public/TedsOutlinerImpl.h
--- a/Engine/Plugins/Experimental/EditorDataStorageFeatures/Source/TedsOutliner/Public/TedsOutlinerImpl.h
+++ b/Engine/Plugins/Experimental/EditorDataStorageFeatures/Source/TedsOutliner/Public/TedsOutlinerImpl.h
@@ -13,6 +13,8 @@ namespace UE::Editor::DataStorage
 {
 	namespace QueryStack
 	{
+		class FRowHandleSortNode;
+		class FRowQueryResultsNode;
 		class IQueryNode;
 		class FRowMonitorNode;
 		class IRowNode;
@@ -237,11 +239,14 @@ protected:
 	TOptional<FTedsOutlinerHierarchyData> HierarchyData;
 
 	// The query stack node responsible for collecting all rows that match the composite query on FullRefresh()
-	TSharedPtr<DataStorage::QueryStack::IRowNode> RowCollectorNode;
+	TSharedPtr<DataStorage::QueryStack::FRowQueryResultsNode> RowCollectorNode;
 
 	// The query stack node responsible for tracking and row addition/removals that happen in between refreshes
 	TSharedPtr<DataStorage::QueryStack::FRowMonitorNode> RowMonitorNode;
 
+	// A query stack node that keeps a copy of our rows sorted by row handle for faster search/lookup
+	TSharedPtr<DataStorage::QueryStack::FRowHandleSortNode> RowHandleSortNode;
+
 	// Query to get all child rows
 	DataStorage::QueryHandle ChildRowHandleQuery = DataStorage::InvalidQueryHandle;
 
