Scene Outliner:
Add a method to allow Modes to tick before the Outliner tick executes

TEDS Outliner:
Move the query stack updates to Mode->Tick() so they execute before any outliner updates that rely on the data
Modify the RowCollectorQuery to update every frame if there is no observer, if the observer is present it handles updating the row list

Chaos Visual Debugger:
Modify the Outliner query to use FConditions since .Any<> does not behave as expected with the new query merge

diff --git a/Engine/Plugins/ChaosVD/Source/ChaosVD/Private/ChaosVDWorldOutlinerTab.cpp b/Engine/Plugins/ChaosVD/Source/ChaosVD/Private/ChaosVDWorldOutlinerTab.cpp
--- a/Engine/Plugins/ChaosVD/Source/ChaosVD/Private/ChaosVDWorldOutlinerTab.cpp
+++ b/Engine/Plugins/ChaosVD/Source/ChaosVD/Private/ChaosVDWorldOutlinerTab.cpp
@@ -39,9 +39,7 @@ void FChaosVDWorldOutlinerTab::CreateWorldOutlinerWidget()
 	FQueryDescription OutlinerQueryDescription =
 		Select()
 			.ReadOnly<FTypedElementLabelColumn>()
-		.Where()
-			.Any<FChaosVDObjectDataTag, FTypedElementActorTag>()
-			.All<FTypedElementFromCVDWorldTag, FChaosVDActiveObjectTag>()
+		.Where((TColumn<FChaosVDObjectDataTag>() || TColumn<FTypedElementActorTag>()) && TColumn<FTypedElementFromCVDWorldTag>() && TColumn<FChaosVDActiveObjectTag>())
 		.Compile();
 
 	FSceneOutlinerInitializationOptions InitOptions;
diff --git a/Engine/Plugins/Experimental/EditorDataStorageFeatures/Source/TedsOutliner/Private/TedsOutlinerImpl.cpp b/Engine/Plugins/Experimental/EditorDataStorageFeatures/Source/TedsOutliner/Private/TedsOutlinerImpl.cpp
--- a/Engine/Plugins/Experimental/EditorDataStorageFeatures/Source/TedsOutliner/Private/TedsOutlinerImpl.cpp
+++ b/Engine/Plugins/Experimental/EditorDataStorageFeatures/Source/TedsOutliner/Private/TedsOutlinerImpl.cpp
@@ -275,7 +275,7 @@ void FTedsOutlinerImpl::Init()
 	RecompileQueries();
 
 	// Tick post TEDS update to make sure all processors have run and the data is correct
-	Storage->OnUpdateCompleted().AddRaw(this, &FTedsOutlinerImpl::Tick);
+	Storage->OnUpdateCompleted().AddRaw(this, &FTedsOutlinerImpl::PostDataStorageUpdate);
 }
 
 void FTedsOutlinerImpl::OnMonitoredRowsAdded(DataStorage::FRowHandleArrayView InRows)
@@ -560,7 +560,8 @@ DataStorage::RowHandle FTedsOutlinerImpl::GetParentRow(DataStorage::RowHandle In
 	}
 
 	// If the row doesn't exist in our query stack results, it doesn't match the query and should not be displayed
-	DataStorage::FRowHandleArrayView Rows = RowHandleSortNode->GetRows();
+	DataStorage::FRowHandleArrayView Rows(RowHandleSortNode->GetRows());
+
 	if (!Rows.Contains(ParentRowHandle))
 	{
 		return DataStorage::InvalidRowHandle;
@@ -592,20 +593,28 @@ void FTedsOutlinerImpl::RecompileQueries()
 
 	TSharedPtr<IQueryNode> CombinedQueryNode = MakeShared<FQueryMergeNode>(*Storage, QueryNodes);
 
+	// If we have a monitor node, it handles updating the collector node automatically. Otherwise we need it to update on tick
+	FRowQueryResultsNode::ESyncFlags SyncFlags =
+		CreationParams.bUseDefaultObservers ? FRowQueryResultsNode::ESyncFlags::None : FRowQueryResultsNode::ESyncFlags::RefreshOnUpdate;
+	
 	// Node to collect all rows matching the composite query
-	RowCollectorNode = MakeShared<FRowQueryResultsNode>(*Storage, CombinedQueryNode, FRowQueryResultsNode::ESyncFlags::None);
-
-	// We store a copy of the collected rows into an array sorted by row handles for faster lookup
-	TSharedRef<FRowCopyNode> CopyNode = MakeShared<FRowCopyNode>(RowCollectorNode);
-	RowHandleSortNode = MakeShared<FRowHandleSortNode>(CopyNode);
+	RowCollectorNode = MakeShared<FRowQueryResultsNode>(*Storage, CombinedQueryNode, SyncFlags);
 
+	// The "actual" collector node passed into other misc nodes is the monitor node if it exists, the collector node otherwise 
+	TSharedPtr<IRowNode> ActualCollectorNode = RowCollectorNode;
+	
 	// Node to monitor row addition/removal
 	if (CreationParams.bUseDefaultObservers)
 	{
 		RowMonitorNode = MakeShared<FRowMonitorNode>(*Storage, RowCollectorNode, CombinedQueryNode);
 		RowMonitorNode->OnMonitoredRowsAdded().AddRaw(this, &FTedsOutlinerImpl::OnMonitoredRowsAdded);
 		RowMonitorNode->OnMonitoredRowsRemoved().AddRaw(this, &FTedsOutlinerImpl::OnMonitoredRowsRemoved);
+		ActualCollectorNode = RowMonitorNode;
 	}
+	
+	// We store a copy of the collected rows into an array sorted by row handles for faster lookup
+	TSharedRef<FRowCopyNode> CopyNode = MakeShared<FRowCopyNode>(ActualCollectorNode);
+	RowHandleSortNode = MakeShared<FRowHandleSortNode>(CopyNode);
 
 	// Our final query to collect rows to populate the Outliner - currently the same as the initial query the user provided
 	FQueryDescription FinalQueryDescription(InitialQueryDescription);
@@ -642,7 +651,7 @@ void FTedsOutlinerImpl::RecompileQueries()
 		FQueryDescription UpdateParentQueryDescription =
 			Select(
 			TEXT("Update item parent"),
-			FProcessor(EQueryTickPhase::PostPhysics, Storage->GetQueryTickGroupName(EQueryTickGroups::SyncDataStorageToExternal))
+			FProcessor(EQueryTickPhase::PostPhysics, Storage->GetQueryTickGroupName(EQueryTickGroups::SyncExternalToDataStorage))
 				.SetExecutionMode(EExecutionMode::GameThread),
 			[this, ParentColumnType](IQueryContext& Context, const RowHandle* Rows)
 			{
@@ -760,8 +769,7 @@ void FTedsOutlinerImpl::RecompileQueries()
 	
 	if (bUsingQueryConditionsSyntax)
 	{
-		LabelUpdateQueryDescription.Conditions.Emplace(TColumn<FTypedElementSyncBackToWorldTag>());
-		LabelUpdateQueryDescription.Conditions.Emplace(TColumn<FTypedElementSyncFromWorldTag>());
+		LabelUpdateQueryDescription.Conditions.Emplace(TColumn<FTypedElementSyncBackToWorldTag>() || TColumn<FTypedElementSyncFromWorldTag>());
 	}
 	else
 	{
@@ -836,21 +844,8 @@ void FTedsOutlinerImpl::ClearSelection() const
 
 }
 
-void FTedsOutlinerImpl::Tick()
+void FTedsOutlinerImpl::PostDataStorageUpdate()
 {
-	if (bSelectionDirty)
-	{
-		OnTedsOutlinerSelectionChanged.Broadcast();
-		bSelectionDirty = false;
-	}
-
-	if (RowMonitorNode)
-	{
-		RowMonitorNode->Update();
-	}
-
-	RowHandleSortNode->Update();
-	
 	// Update the label for any rows that might need it
 	for (DataStorage::RowHandle Row : RowsPendingLabelUpdate)
 	{
@@ -878,6 +873,18 @@ void FTedsOutlinerImpl::Tick()
 	RowsPendingLabelUpdate.Empty();
 }
 
+void FTedsOutlinerImpl::Tick()
+{
+	if (bSelectionDirty)
+	{
+		OnTedsOutlinerSelectionChanged.Broadcast();
+		bSelectionDirty = false;
+	}
+
+	// The Sort node is the childmost node in all cases, so updating it causes all the nodes in the stack to update
+	RowHandleSortNode->Update();
+}
+
 DataStorage::ICoreProvider* FTedsOutlinerImpl::GetStorage() const
 {
 	return Storage;
diff --git a/Engine/Plugins/Experimental/EditorDataStorageFeatures/Source/TedsOutliner/Private/TedsOutlinerMode.cpp b/Engine/Plugins/Experimental/EditorDataStorageFeatures/Source/TedsOutliner/Private/TedsOutlinerMode.cpp
--- a/Engine/Plugins/Experimental/EditorDataStorageFeatures/Source/TedsOutliner/Private/TedsOutlinerMode.cpp
+++ b/Engine/Plugins/Experimental/EditorDataStorageFeatures/Source/TedsOutliner/Private/TedsOutlinerMode.cpp
@@ -298,6 +298,11 @@ TSharedPtr<SWidget> FTedsOutlinerMode::CreateContextMenu()
 	return UToolMenus::Get()->GenerateWidget(Private::ContextMenuName, MenuContext);
 }
 
+void FTedsOutlinerMode::Tick()
+{
+	TedsOutlinerImpl->Tick();
+}
+
 TUniquePtr<ISceneOutlinerHierarchy> FTedsOutlinerMode::CreateHierarchy()
 {
 	return MakeUnique<FTedsOutlinerHierarchy>(this, TedsOutlinerImpl.ToSharedRef());
diff --git a/Engine/Plugins/Experimental/EditorDataStorageFeatures/Source/TedsOutliner/Public/TedsOutlinerImpl.h b/Engine/Plugins/Experimental/EditorDataStorageFeatures/Source/TedsOutliner/Public/TedsOutlinerImpl.h
--- a/Engine/Plugins/Experimental/EditorDataStorageFeatures/Source/TedsOutliner/Public/TedsOutlinerImpl.h
+++ b/Engine/Plugins/Experimental/EditorDataStorageFeatures/Source/TedsOutliner/Public/TedsOutlinerImpl.h
@@ -204,11 +204,13 @@ public:
 
 	// Unregister all queries used by this table viewer
 	void UnregisterQueries();
+	
+	void Tick();
 
 protected:
 	
 	void ClearSelection() const;
-	void Tick();
+	void PostDataStorageUpdate();
 
 	void CreateFilterQueries();
 
diff --git a/Engine/Plugins/Experimental/EditorDataStorageFeatures/Source/TedsOutliner/Public/TedsOutlinerMode.h b/Engine/Plugins/Experimental/EditorDataStorageFeatures/Source/TedsOutliner/Public/TedsOutlinerMode.h
--- a/Engine/Plugins/Experimental/EditorDataStorageFeatures/Source/TedsOutliner/Public/TedsOutlinerMode.h
+++ b/Engine/Plugins/Experimental/EditorDataStorageFeatures/Source/TedsOutliner/Public/TedsOutlinerMode.h
@@ -59,6 +59,7 @@ public:
 	virtual FSceneOutlinerDragValidationInfo ValidateDrop(const ISceneOutlinerTreeItem& DropTarget, const FSceneOutlinerDragDropPayload& Payload) const override;
 	virtual void OnDrop(ISceneOutlinerTreeItem& DropTarget, const FSceneOutlinerDragDropPayload& Payload, const FSceneOutlinerDragValidationInfo& ValidationInfo) const override;
 	virtual TSharedPtr<SWidget> CreateContextMenu() override;
+	virtual void Tick() override;
 	/* end ISceneOutlinerMode interface */
 
 protected:
diff --git a/Engine/Source/Editor/SceneOutliner/Private/SSceneOutliner.cpp b/Engine/Source/Editor/SceneOutliner/Private/SSceneOutliner.cpp
--- a/Engine/Source/Editor/SceneOutliner/Private/SSceneOutliner.cpp
+++ b/Engine/Source/Editor/SceneOutliner/Private/SSceneOutliner.cpp
@@ -2337,6 +2337,8 @@ void SSceneOutliner::Tick(const FGeometry& AllottedGeometry, const double InCurr
 {	
 	TRACE_CPUPROFILER_EVENT_SCOPE(SSceneOutliner::Tick);
 
+	Mode->Tick();
+
 	for (auto& Pair : Columns)
 	{
 		Pair.Value->Tick(InCurrentTime, InDeltaTime);
diff --git a/Engine/Source/Editor/SceneOutliner/Public/ISceneOutlinerMode.h b/Engine/Source/Editor/SceneOutliner/Public/ISceneOutlinerMode.h
--- a/Engine/Source/Editor/SceneOutliner/Public/ISceneOutlinerMode.h
+++ b/Engine/Source/Editor/SceneOutliner/Public/ISceneOutlinerMode.h
@@ -182,6 +182,9 @@ public:
 
 	/** Function called by the Outliner Filter Bar to compare an item with Type Filters*/
 	virtual bool CompareItemWithClassName(SceneOutliner::FilterBarType InItem, const TSet<FTopLevelAssetPath>&) const { return false; };
+
+	/** Function called at the start of the Outliner's tick to allow the mode to perform operations */
+	virtual void Tick() {}
 public:
 	/* Getters */
 
