CB 2.0 - Fix incorrect usage of query stack in the TEDS-CB prototype and the Content Sources

diff --git a/Engine/Plugins/Experimental/EditorDataStorageFeatures/Source/TedsContentBrowser/Private/TedsContentBrowserModule.cpp b/Engine/Plugins/Experimental/EditorDataStorageFeatures/Source/TedsContentBrowser/Private/TedsContentBrowserModule.cpp
--- a/Engine/Plugins/Experimental/EditorDataStorageFeatures/Source/TedsContentBrowser/Private/TedsContentBrowserModule.cpp
+++ b/Engine/Plugins/Experimental/EditorDataStorageFeatures/Source/TedsContentBrowser/Private/TedsContentBrowserModule.cpp
@@ -16,6 +16,7 @@
 #include "TedsAlertColumns.h"
 #include "TedsAssetDataColumns.h"
 #include "TedsAssetDataWidgetColumns.h"
+#include "TedsRowArrayNode.h"
 #include "TedsRowViewNode.h"
 #include "Elements/Columns/TypedElementLabelColumns.h"
 #include "Elements/Columns/TypedElementTypeInfoColumns.h"
@@ -121,7 +122,12 @@ namespace UE::Editor::ContentBrowser
 			return;
 		}
 
-		Rows.Empty();
+		FRowHandleArray* Rows = RowQueryStack->GetMutableRows();
+		if (ensureMsgf(Rows, TEXT("Expected the row query stack to be mutable since it is used in the table viewer")))
+		{
+			Rows->Empty();
+		}
+		
 		ContentBrowserItemMap.Empty();
 
 		for(TSharedPtr<FAssetViewItem> Item : *InItemsSource)
@@ -134,8 +140,7 @@ namespace UE::Editor::ContentBrowser
 			AddRow(Item);
 		}
 
-		RowQueryStack->ResetView(FRowHandleArrayView(Rows.GetData(), Rows.Num(),
-			FRowHandleArrayView::EFlags::IsSorted | FRowHandleArrayView::EFlags::IsUnique));
+		RowQueryStack->MarkDirty();
 	}
 
 	void FTedsContentBrowserViewExtender::AddRow(const TSharedPtr<FAssetViewItem>& Item)
@@ -145,7 +150,12 @@ namespace UE::Editor::ContentBrowser
 		if(DataStorage->IsRowAssigned(RowHandle))
 		{
 			ContentBrowserItemMap.Emplace(RowHandle, Item);
-			Rows.Add(RowHandle);
+			
+			FRowHandleArray* Rows = RowQueryStack->GetMutableRows();
+			if (ensureMsgf(Rows, TEXT("Expected the row query stack to be mutable since it is used in the table viewer")))
+			{
+				Rows->Add(RowHandle);
+			}
 		}
 	}
 
@@ -432,8 +442,7 @@ namespace UE::Editor::ContentBrowser
 		using namespace UE::Editor::DataStorage;
 		DataStorage = GetMutableDataStorageFeature<ICoreProvider>(StorageFeatureName);
 		
-		RowQueryStack = MakeShared<QueryStack::FRowViewNode>(FRowHandleArrayView(Rows.GetData(), Rows.Num(),
-			FRowHandleArrayView::EFlags::IsSorted | FRowHandleArrayView::EFlags::IsUnique));
+		RowQueryStack = MakeShared<QueryStack::FRowArrayNode>();
 	}
 
 	TSharedRef<SWidget> FTedsContentBrowserViewExtender::CreateView(TArray<TSharedPtr<FAssetViewItem>>* InItemsSource)
diff --git a/Engine/Plugins/Experimental/EditorDataStorageFeatures/Source/TedsContentBrowser/Private/TedsContentBrowserModule.h b/Engine/Plugins/Experimental/EditorDataStorageFeatures/Source/TedsContentBrowser/Private/TedsContentBrowserModule.h
--- a/Engine/Plugins/Experimental/EditorDataStorageFeatures/Source/TedsContentBrowser/Private/TedsContentBrowserModule.h
+++ b/Engine/Plugins/Experimental/EditorDataStorageFeatures/Source/TedsContentBrowser/Private/TedsContentBrowserModule.h
@@ -4,6 +4,7 @@
 
 #include "ContentSources/IContentSource.h"
 #include "Elements/Common/TypedElementHandles.h"
+#include "Elements/Framework/TypedElementRowHandleArray.h"
 #include "Experimental/ContentBrowserViewExtender.h"
 #include "Modules/ModuleInterface.h"
 #include "Templates/SharedPointer.h"
@@ -16,7 +17,7 @@ namespace UE::Editor::DataStorage
 
 	namespace QueryStack
 	{
-		class FRowViewNode;
+		class FRowArrayNode;
 	}
 }
 
@@ -145,10 +146,7 @@ namespace UE::Editor::ContentBrowser
 		TSharedPtr<DataStorage::ITableViewer> TableViewer;
 		
 		// Query stack used by the table viewer
-		TSharedPtr<DataStorage::QueryStack::FRowViewNode> RowQueryStack;
-
-		// The row handles of the items currently in the list
-		TArray<DataStorage::RowHandle> Rows;
+		TSharedPtr<DataStorage::QueryStack::FRowArrayNode> RowQueryStack;
 
 		// A map from row handle -> FAssetView item for lookups
 		TMap<DataStorage::RowHandle, TWeakPtr<FAssetViewItem>> ContentBrowserItemMap;
diff --git a/Engine/Plugins/Experimental/EditorDataStorageFeatures/Source/TedsContentBrowser/Private/Widgets/TedsContentBrowserAssetViewWidget.cpp b/Engine/Plugins/Experimental/EditorDataStorageFeatures/Source/TedsContentBrowser/Private/Widgets/TedsContentBrowserAssetViewWidget.cpp
--- a/Engine/Plugins/Experimental/EditorDataStorageFeatures/Source/TedsContentBrowser/Private/Widgets/TedsContentBrowserAssetViewWidget.cpp
+++ b/Engine/Plugins/Experimental/EditorDataStorageFeatures/Source/TedsContentBrowser/Private/Widgets/TedsContentBrowserAssetViewWidget.cpp
@@ -2,6 +2,8 @@
 
 #include "TedsContentBrowserAssetViewWidget.h"
 
+#include "TedsQueryNode.h"
+#include "TedsRowQueryResultsNode.h"
 #include "TedsRowViewNode.h"
 #include "ContentSources/IContentSource.h"
 #include "ContentSources/Columns/ContentSourcesColumns.h"
@@ -28,13 +30,11 @@ public:
 		using namespace UE::Editor::DataStorage;
 		Storage = Arguments._Storage;
 
-		FQueryDescription QueryDescription = Arguments._InitParams.QueryDescription;
-		
-		QueryHandle = Storage->RegisterQuery(MoveTemp(QueryDescription));
-		
-		RowView = MakeShared<QueryStack::FRowViewNode>(FRowHandleArrayView(Rows.GetData(), Rows.Num(),
-			FRowHandleArrayView::EFlags::IsSorted | FRowHandleArrayView::EFlags::IsUnique));
+		TSharedRef<QueryStack::FQueryNode> QueryNode = MakeShared<QueryStack::FQueryNode>(*Storage, Arguments._InitParams.QueryDescription);
 
+		// For now we just refresh the query on update every frame since this is a prototype, in the future we can use a monitor node if perf is a concern
+		RowView = MakeShared<QueryStack::FRowQueryResultsNode>(*Storage, QueryNode, QueryStack::FRowQueryResultsNode::ESyncFlags::RefreshOnUpdate);
+		
 		TSharedPtr<SWidget> ChildWidget;
 
 		switch (Arguments._InitParams.TableViewMode)
@@ -65,34 +65,6 @@ public:
 		];
 	}
 
-	virtual void Tick(const FGeometry& AllottedGeometry, const double InCurrentTime, const float InDeltaTime) override
-	{
-		// On tick, we run the query again and update the query stack if the rows changed
-		
-		TSet<UE::Editor::DataStorage::RowHandle> NewRows;
-
-		using namespace UE::Editor::DataStorage;
-
-		Storage->RunQuery(QueryHandle,
-			Queries::CreateDirectQueryCallbackBinding([&NewRows]
-				(const UE::Editor::DataStorage::IDirectQueryContext& Context, const RowHandle*)
-				{
-					NewRows.Append(Context.GetRowHandles());
-				}));
-
-		// Check if the two sets are equal, i.e no changes and no need to update the table viewer
-		const bool bSetsEqual = (Rows_Set.Num() == NewRows.Num()) && Rows_Set.Includes(NewRows);
-
-		if (!bSetsEqual)
-		{
-			Rows_Set = NewRows;
-			Rows = Rows_Set.Array();
-			
-			RowView->ResetView(FRowHandleArrayView(Rows.GetData(), Rows.Num(),
-				FRowHandleArrayView::EFlags::IsSorted | FRowHandleArrayView::EFlags::IsUnique));
-		}
-	}
-
 	virtual ~STableViewerWrapper() override
 	{
 		Storage->UnregisterQuery(QueryHandle);
@@ -102,7 +74,7 @@ private:
 	UE::Editor::DataStorage::QueryHandle QueryHandle = UE::Editor::DataStorage::InvalidQueryHandle;
 	TArray<UE::Editor::DataStorage::RowHandle> Rows;
 	TSet<UE::Editor::DataStorage::RowHandle> Rows_Set;
-	TSharedPtr<UE::Editor::DataStorage::QueryStack::FRowViewNode> RowView;
+	TSharedPtr<UE::Editor::DataStorage::QueryStack::FRowQueryResultsNode> RowView;
 	UE::Editor::DataStorage::ICoreProvider* Storage = nullptr;
 	
 };
