[TEDS Debugger Improvments] Queries can now be searched via a custom combo widget

diff --git a/Engine/Plugins/Experimental/EditorDataStorageFeatures/Source/TedsDebugger/Private/STedsDebugger.cpp b/Engine/Plugins/Experimental/EditorDataStorageFeatures/Source/TedsDebugger/Private/STedsDebugger.cpp
--- a/Engine/Plugins/Experimental/EditorDataStorageFeatures/Source/TedsDebugger/Private/STedsDebugger.cpp
+++ b/Engine/Plugins/Experimental/EditorDataStorageFeatures/Source/TedsDebugger/Private/STedsDebugger.cpp
@@ -35,7 +35,7 @@ STedsDebugger::~STedsDebugger()
 
 void STedsDebugger::Construct(const FArguments& InArgs, const TSharedRef<SDockTab>& ConstructUnderMajorTab, const TSharedPtr<SWindow>& ConstructUnderWindow)
 {
-	// Create the tab manager for our sub tabsa
+	// Create the tab manager for our sub tabs
 	TabManager = FGlobalTabmanager::Get()->NewTabManager(ConstructUnderMajorTab);
 	TabManager->SetAllowWindowMenuBar(true);
 
diff --git a/Engine/Plugins/Experimental/EditorDataStorageFeatures/Source/TedsDebugger/Private/Widgets/QueryEditor/TedsConditionSelectionComboWidget.cpp b/Engine/Plugins/Experimental/EditorDataStorageFeatures/Source/TedsDebugger/Private/Widgets/QueryEditor/TedsConditionSelectionComboWidget.cpp
--- a/Engine/Plugins/Experimental/EditorDataStorageFeatures/Source/TedsDebugger/Private/Widgets/QueryEditor/TedsConditionSelectionComboWidget.cpp
+++ b/Engine/Plugins/Experimental/EditorDataStorageFeatures/Source/TedsDebugger/Private/Widgets/QueryEditor/TedsConditionSelectionComboWidget.cpp
@@ -6,8 +6,11 @@
 #include "Algo/FindSortedStringCaseInsensitive.h"
 #include "Elements/Common/TypedElementDataStorageLog.h"
 #include "Filters/SBasicFilterBar.h"
-#include "Widgets/Text/STextBlock.h"
 #include "Widgets/Input/SComboBox.h"
+#include "Widgets/Input/SSearchBox.h"
+#include "Widgets/Text/STextBlock.h"
+
+#define LOCTEXT_NAMESPACE "TedsDebuggerComboWidget"
 
 namespace UE::Editor::DataStorage::Debug::QueryEditor
 {
@@ -22,6 +25,7 @@ namespace UE::Editor::DataStorage::Debug::QueryEditor
 	void SConditionComboWidget::OnConditionCollectionChanged()
 	{		
 		PopulateComboItems();
+		FilterComboItems();
 	}
 
 	SConditionComboWidget::~SConditionComboWidget()
@@ -38,41 +42,63 @@ namespace UE::Editor::DataStorage::Debug::QueryEditor
 		
 		ConditionCollectionChangedHandle = Model->GetModelChangedDelegate().AddRaw(this, &SConditionComboWidget::OnConditionCollectionChanged);
 
-		
-		
-		ComboBox = SNew(SComboBox<TSharedPtr<FComboItem>>)
-			.OptionsSource(&ComboItems)
-			.OnComboBoxOpening_Lambda([this]()
-			{
-				Model->RegenerateColumnsList();
-				PopulateComboItems();
-			})
-			.OnSelectionChanged(this, &SConditionComboWidget::OnSelectionChanged)
-			.OnGenerateWidget(this, &SConditionComboWidget::OnGenerateWidget)
-			.IsEnabled_Lambda([this](){ return !ComboItems.IsEmpty(); })
-			.Content()
+		// This implementation is very similar to the existing SSearchableComboBox implementation, but since that only takes string types, a 
+		// custom one was needed to utilize the custom FComboItem type. 
+		ComboBoxMenuContent =
+			SNew(SBox)
+			.MinDesiredWidth(400.0f)
+			.MaxDesiredHeight(500.0f)
 			[
-				SNew(STextBlock)
-				.Text_Lambda([]() { return FText::FromString("+"); })
+				SNew(SVerticalBox)
+				+ SVerticalBox::Slot()
+				.AutoHeight()
+				.HAlign(EHorizontalAlignment::HAlign_Fill)
+				.Padding(4.0f, 2.0f)
+				[
+					SAssignNew(SearchBox, SSearchBox)
+						.OnTextChanged(this, &SConditionComboWidget::OnSearchBoxTextChanged)
+						.OnTextCommitted(this, &SConditionComboWidget::OnSearchBoxTextCommitted)
+				]
+				+SVerticalBox::Slot()
+				.HAlign(EHorizontalAlignment::HAlign_Fill)
+				[
+					SAssignNew(ComboListView, SListView<TSharedPtr<FComboItem>>)
+						.ListItemsSource(&FilteredComboItems)
+						.OnSelectionChanged(this, &SConditionComboWidget::OnSelectionChanged)
+						.OnGenerateRow(this, &SConditionComboWidget::GenerateMenuItemRow)
+						.OnKeyDownHandler(this, &SConditionComboWidget::OnKeyDownHandler)
+						.SelectionMode(ESelectionMode::Single)
+				]
 			];
 
-		PopulateComboItems();
+		SComboButton::Construct(SComboButton::FArguments()
+			.ComboButtonStyle(&FAppStyle::Get().GetWidgetStyle<FComboBoxStyle>("ComboBox").ComboButtonStyle)
+			.ButtonContent()
+			[
+				SNew(STextBlock)
+					.Text(FText::FromString("+"))
+			]
+			.MenuContent()
+			[
+				ComboBoxMenuContent.ToSharedRef()
+			]
+			.CollapseMenuOnParentFocus(true)
+			.OnComboBoxOpened_Lambda([this]()
+				{
+					Model->RegenerateColumnsList();
+					PopulateComboItems();
+					SearchBox->SetText(FText::GetEmpty());
+				})
+		);
+		SetMenuContentWidgetToFocus(SearchBox);
 
-		ChildSlot
-		[
-		SNew( SVerticalBox )
-		+ SVerticalBox::Slot()
-		.AutoHeight()
-		[
-			ComboBox.ToSharedRef()
-		]
-			
-		];
+		PopulateComboItems();
 	}
 
 	void SConditionComboWidget::PopulateComboItems()
 	{
 		ComboItems.Reset();
+		FilteredComboItems.Reset();
 		Model->GenerateValidOperatorChoices(ConditionType, [this](const FTedsQueryEditorModel&, const QueryEditor::FConditionEntryHandle& Handle)
 		{
 			FComboItem ComboItem;
@@ -89,28 +115,125 @@ namespace UE::Editor::DataStorage::Debug::QueryEditor
 			return AStruct->GetFName().Compare(BStruct->GetFName()) < 0;
 		});
 		
-		ComboBox->RefreshOptions();
+		FilteredComboItems.Append(ComboItems);
+		ComboListView->RequestListRefresh();
 	}
 
-	void SConditionComboWidget::OnSelectionChanged(TSharedPtr<FComboItem> NewSelection, ESelectInfo::Type SelectInfo)
+	void SConditionComboWidget::FilterComboItems()
 	{
-		if (NewSelection)
+		// Filter the current combo items by the search text, since combo items are already sorted on population, these will not need to be sorted again
+		FilteredComboItems.Reset();
+
+		if (SearchText.IsEmpty())
 		{
-			EErrorCode ErrorCode = Model->SetOperatorType(NewSelection->Handle, ConditionType);
-			if (ErrorCode != EErrorCode::Success)
+			FilteredComboItems.Append(ComboItems);
+		}
+		else
+		{
+			TArray<FString> SearchTokens;
+			SearchText.ToString().ParseIntoArrayWS(SearchTokens);
+
+			for (TSharedPtr<FComboItem> ComboItem : ComboItems)
 			{
-				UE_LOG(LogEditorDataStorage, Error, TEXT("Could not set model condition: [%d]"), static_cast<int>(ErrorCode));
+				if (const UScriptStruct* ColumnType = Model->GetColumnScriptStruct(ComboItem->Handle))
+				{
+					const FString ColumnName = ColumnType->GetName();
+
+					bool bAllTokensMatch = true;
+					for (const FString& SearchToken : SearchTokens)
+					{
+						if (ColumnName.IsEmpty() || !ColumnName.Contains(SearchToken, ESearchCase::Type::IgnoreCase))
+						{
+							bAllTokensMatch = false;
+							break;
+						}
+					}
+
+					if (bAllTokensMatch)
+					{
+						FilteredComboItems.Add(ComboItem);
+					}
+				}
 			}
-			ComboBox->ClearSelection();
+		}
+
+		ComboListView->RequestListRefresh();
+	}
+
+	// Handle when the user presses enter to select an item that was navigated to via down/up keys
+	FReply SConditionComboWidget::OnKeyDownHandler(const FGeometry& MyGeometry, const FKeyEvent& InKeyEvent)
+	{
+		if (InKeyEvent.GetKey() == EKeys::Enter && FilteredComboItems.Num() > 0)
+		{
+			if (ComboListView->GetNumItemsSelected() > 0)
+			{
+
+				SelectItem(ComboListView->GetSelectedItems()[0]);
+			}
+			else
+			{
+				SelectItem(FilteredComboItems[0]);
+			}
+			
+			return FReply::Handled();
+		}
+
+		return FReply::Unhandled();
+	}
+
+	// Handle when the user clicks an item directly to select it
+	void SConditionComboWidget::OnSelectionChanged(TSharedPtr<FComboItem> NewSelection, ESelectInfo::Type SelectInfo)
+	{
+		if (NewSelection && SelectInfo == ESelectInfo::OnMouseClick)
+		{
+			SelectItem(NewSelection);
 		}
 	}
 
 	TSharedRef<SWidget> SConditionComboWidget::OnGenerateWidget(TSharedPtr<FComboItem> Item)
 	{
 		const UScriptStruct* ColumnType = Model->GetColumnScriptStruct(Item->Handle);
-		const FText Label = ColumnType ? FText::FromString(ColumnType->GetName()) : FText::FromString(TEXT("Null Column Type"));
+		const FText ColumnName = ColumnType ? FText::FromString(ColumnType->GetName()) : LOCTEXT("UnknownColumnName", "Unknown Column Type");
 
 		return SNew(STextBlock)
-			.Text(Label);
+			.Text(ColumnName);
 	}
+
+	TSharedRef<ITableRow> SConditionComboWidget::GenerateMenuItemRow(TSharedPtr<FComboItem> InItem, const TSharedRef<STableViewBase>& OwnerTable)
+	{
+		return SNew(SComboRow<TSharedPtr<FComboItem>>, OwnerTable)
+			.Style(&FAppStyle::Get().GetWidgetStyle<FTableRowStyle>("ComboBox.Row"))
+			.Padding(FMargin(8.0f, 3.f))
+			[
+				OnGenerateWidget(InItem)
+			];
+	}
+
+	void SConditionComboWidget::OnSearchBoxTextChanged(const FText& InSearchText)
+	{
+		SearchText = InSearchText;
+		FilterComboItems();
+	}
+
+	// Handles when the user presses enter while focused on the search box, it will select the first item in the list
+	void SConditionComboWidget::OnSearchBoxTextCommitted(const FText& InSearchText, ETextCommit::Type InCommitType)
+	{
+		if ((InCommitType == ETextCommit::Type::OnEnter) && FilteredComboItems.Num() > 0)
+		{
+			SelectItem(FilteredComboItems[0]);
+		}
+	}
+
+	void SConditionComboWidget::SelectItem(TSharedPtr<FComboItem> InItem)
+	{
+		EErrorCode ErrorCode = Model->SetOperatorType(InItem->Handle, ConditionType);
+		if (ErrorCode != EErrorCode::Success)
+		{
+			UE_LOG(LogEditorDataStorage, Error, TEXT("Could not set model condition: [%d]"), static_cast<int>(ErrorCode));
+		}
+		ComboListView->ClearSelection();
+	}
+
 }
+
+#undef LOCTEXT_NAMESPACE
\ No newline at end of file
diff --git a/Engine/Plugins/Experimental/EditorDataStorageFeatures/Source/TedsDebugger/Private/Widgets/QueryEditor/TedsConditionSelectionComboWidget.h b/Engine/Plugins/Experimental/EditorDataStorageFeatures/Source/TedsDebugger/Private/Widgets/QueryEditor/TedsConditionSelectionComboWidget.h
--- a/Engine/Plugins/Experimental/EditorDataStorageFeatures/Source/TedsDebugger/Private/Widgets/QueryEditor/TedsConditionSelectionComboWidget.h
+++ b/Engine/Plugins/Experimental/EditorDataStorageFeatures/Source/TedsDebugger/Private/Widgets/QueryEditor/TedsConditionSelectionComboWidget.h
@@ -3,7 +3,8 @@
 #include "Filters/GenericFilter.h"
 #include "Filters/SFilterBar.h"
 #include "QueryEditor/TedsQueryEditorModel.h"
-#include "Widgets/SCompoundWidget.h"
+#include "Widgets/Input/SComboButton.h"
+#include "Widgets/Views/SListView.h"
 
 template<typename>
 class SComboBox;
@@ -13,7 +14,7 @@ namespace UE::Editor::DataStorage::Debug::QueryEditor
 	class FTedsQueryEditorModel;
 	struct FConditionEntry;
 
-	class SConditionComboWidget : public SCompoundWidget
+	class SConditionComboWidget : public SComboButton
 	{
 	public:
 		SLATE_BEGIN_ARGS( SConditionComboWidget ){}
@@ -22,21 +23,40 @@ namespace UE::Editor::DataStorage::Debug::QueryEditor
 		~SConditionComboWidget() override;
 		void OnConditionCollectionChanged();
 		void Construct(const FArguments& InArgs, FTedsQueryEditorModel& InModel, QueryEditor::EOperatorType InConditionType);
-	private:
 
+	private:
 		struct FComboItem;
+
 		void PopulateComboItems();
+		void FilterComboItems();
+
 		void OnSelectionChanged(TSharedPtr<FComboItem> NewSelection, ESelectInfo::Type SelectInfo);
+		FReply OnKeyDownHandler(const FGeometry& MyGeometry, const FKeyEvent& InKeyEvent);
 		TSharedRef<SWidget> OnGenerateWidget(TSharedPtr<FComboItem> Item);
+		TSharedRef<ITableRow> GenerateMenuItemRow(TSharedPtr<FComboItem> InItem, const TSharedRef<STableViewBase>& OwnerTable);
+
+		void OnSearchBoxTextChanged(const FText& InSearchText);
+		void OnSearchBoxTextCommitted(const FText& InSearchText, ETextCommit::Type InCommitType);
+
+		void SelectItem(TSharedPtr<FComboItem> InItem);
 
 		FTedsQueryEditorModel* Model = nullptr;
 		QueryEditor::EOperatorType ConditionType = QueryEditor::EOperatorType::Invalid;
 
 		FDelegateHandle ConditionCollectionChangedHandle;
-		bool bConditionCollectionDirty = true;
+
+		/** Updated whenever search text is changed */
+		FText SearchText;
 
 		TArray<TSharedPtr<FComboItem>> ComboItems;
-		TSharedPtr<SComboBox<TSharedPtr<FComboItem>>> ComboBox;
+		TArray<TSharedPtr<FComboItem>> FilteredComboItems;
+
+		/** Widget that gets passed as the combo box content containing the list and search bar */
+		TSharedPtr<SWidget> ComboBoxMenuContent;
+
+		/** The ListView that pops up with the available columns and tags */
+		TSharedPtr<SListView<TSharedPtr<FComboItem>>> ComboListView;
+		TSharedPtr<SSearchBox> SearchBox;
 
 		TSharedPtr<SFilterBar<TSharedPtr<FComboItem>>> FilterThing;
 	};
