Teds Outliner:
- Remove the query binder and instead generate the columns inside FTedsOutliner Impl
- Register all Teds Outliners as rows in Teds mapped using the outliner identifier for lookup

diff --git a/Engine/Plugins/ChaosVD/Source/ChaosVD/Private/ChaosVDWorldOutlinerTab.cpp b/Engine/Plugins/ChaosVD/Source/ChaosVD/Private/ChaosVDWorldOutlinerTab.cpp
--- a/Engine/Plugins/ChaosVD/Source/ChaosVD/Private/ChaosVDWorldOutlinerTab.cpp
+++ b/Engine/Plugins/ChaosVD/Source/ChaosVD/Private/ChaosVDWorldOutlinerTab.cpp
@@ -55,6 +55,11 @@ void FChaosVDWorldOutlinerTab::CreateWorldOutlinerWidget()
 	Params.bShowRowHandleColumn = false;
 	Params.bUseDefaultObservers = false; // CVD uses custom observers to time slice addition and removal currently
 
+	Params.ColumnQueryDescription =
+		Select()
+			.ReadOnly<FVisibleInEditorColumn>()
+		.Compile();
+
 	const FTedsOutlinerHierarchyData::FGetParentRowHandle RowHandleGetter = FTedsOutlinerHierarchyData::FGetParentRowHandle::CreateLambda([](const void* InColumnData)
 	{
 		if(const FTableRowParentColumn* ParentColumn = static_cast<const FTableRowParentColumn *>(InColumnData))
@@ -108,18 +113,6 @@ void FChaosVDWorldOutlinerTab::CreateWorldOutlinerWidget()
 	using namespace UE::Editor::DataStorage;
 	ICoreProvider* Storage = GetMutableDataStorageFeature<ICoreProvider>(StorageFeatureName);
 
-	static QueryHandle ColumnQuery = Storage->RegisterQuery(
-	Select()
-		.ReadOnly<FVisibleInEditorColumn>()
-	.Compile());
-
-	FSceneOutlinerTedsQueryBinder::GetInstance().RegisterTreeItemIDDealiaser(TedsOutliner, FTreeItemIDDealiaser::CreateLambda([](RowHandle RowHandle)
-	{
-		return FSceneOutlinerTreeItemID(RowHandle);
-	}));
-
-	FSceneOutlinerTedsQueryBinder::GetInstance().AssignQuery(ColumnQuery, TedsOutliner, Params.CellWidgetPurpose, Params.LabelWidgetPurpose);
-	
 	SceneOutlinerWidget = TedsOutliner;
 
 	// TODO: There is an issue where the actors created with CVD's world, do not appear until the outliner does a new query (by searching something for example)
diff --git a/Engine/Plugins/Experimental/Dataflow/Source/DataflowEditor/Private/Dataflow/DataflowOutlinerView.cpp b/Engine/Plugins/Experimental/Dataflow/Source/DataflowEditor/Private/Dataflow/DataflowOutlinerView.cpp
--- a/Engine/Plugins/Experimental/Dataflow/Source/DataflowEditor/Private/Dataflow/DataflowOutlinerView.cpp
+++ b/Engine/Plugins/Experimental/Dataflow/Source/DataflowEditor/Private/Dataflow/DataflowOutlinerView.cpp
@@ -48,6 +48,7 @@ TSharedPtr<ISceneOutliner> FDataflowOutlinerView::CreateWidget()
 	Params.QueryDescription = RowQueryDescription;
 	Params.bUseDefaultTedsFilters = false;
 	Params.bShowRowHandleColumn = false;
+	Params.ColumnQueryDescription = ColumnQueryDescription;
 
 	// Add outliner filter queries
 	UE::Editor::DataStorage::FQueryDescription SimulationFilterQuery =
@@ -81,17 +82,9 @@ TSharedPtr<ISceneOutliner> FDataflowOutlinerView::CreateWidget()
 		return new FDataflowOutlinerMode(Params, ConstructionScene, SimulationScene);
 	});
 	
-	UE::Editor::DataStorage::QueryHandle InitialColumnQuery = GetMutableDataStorageFeature<ICoreProvider>(
-		UE::Editor::DataStorage::StorageFeatureName)->RegisterQuery(MoveTemp(ColumnQueryDescription));
 
 	InitOptions.ColumnMap.Add(FSceneOutlinerBuiltInColumnTypes::Label(), FSceneOutlinerColumnInfo(ESceneOutlinerColumnVisibility::Visible, 10));
 	OutlinerWidget = SNew(SSceneOutliner, InitOptions);
-	
-	FSceneOutlinerTedsQueryBinder::GetInstance().AssignQuery(InitialColumnQuery, OutlinerWidget, Params.CellWidgetPurpose, Params.HeaderWidgetPurpose);
-	FSceneOutlinerTedsQueryBinder::GetInstance().RegisterTreeItemIDDealiaser(OutlinerWidget, FTreeItemIDDealiaser::CreateLambda([](RowHandle RowHandle)
-	{
-		return FSceneOutlinerTreeItemID(RowHandle);
-	}));
 
 	return OutlinerWidget;
 }
diff --git a/Engine/Plugins/Experimental/EditorDataStorageFeatures/Source/TedsOutliner/Private/Compatibility/SceneOutlinerTedsBridge.cpp b/Engine/Plugins/Experimental/EditorDataStorageFeatures/Source/TedsOutliner/Private/Compatibility/SceneOutlinerTedsBridge.cpp
--- a/Engine/Plugins/Experimental/EditorDataStorageFeatures/Source/TedsOutliner/Private/Compatibility/SceneOutlinerTedsBridge.cpp
+++ b/Engine/Plugins/Experimental/EditorDataStorageFeatures/Source/TedsOutliner/Private/Compatibility/SceneOutlinerTedsBridge.cpp
@@ -2,344 +2,36 @@
 
 #include "Compatibility/SceneOutlinerTedsBridge.h"
 
-#include "ActorTreeItem.h"
 #include "Columns/TedsOutlinerColumns.h"
-#include "Elements/Columns/TypedElementLabelColumns.h"
 #include "Elements/Columns/TypedElementMiscColumns.h"
 #include "Elements/Columns/TypedElementPackageColumns.h"
-#include "Elements/Columns/TypedElementTransformColumns.h"
 #include "Elements/Columns/TypedElementTypeInfoColumns.h"
 #include "Elements/Columns/TypedElementVisibilityColumns.h"
-#include "Elements/Common/EditorDataStorageFeatures.h"
-#include "Elements/Framework/TypedElementQueryBuilder.h"
 #include "Elements/Interfaces/TypedElementDataStorageInterface.h"
-#include "Elements/Interfaces/TypedElementDataStorageCompatibilityInterface.h"
 #include "Elements/Interfaces/TypedElementDataStorageUiInterface.h"
-#include "ILevelEditor.h"
-#include "ISceneOutliner.h"
-#include "ISceneOutlinerColumn.h"
-#include "ISceneOutlinerTreeItem.h"
-#include "LevelEditor.h"
-#include "SSceneOutliner.h"
-#include "SceneOutlinerFwd.h"
-#include "SceneOutlinerModule.h"
 #include "SceneOutlinerPublicTypes.h"
-#include "TedsTableViewerColumn.h"
-#include "TedsTableViewerUtils.h"
-#include "Elements/Interfaces/Capabilities/TypedElementUiTextCapability.h"
-#include "TedsOutlinerItem.h"
+#include "TedsOutlinerHelpers.h"
 #include "Modules/ModuleManager.h"
 #include "Widgets/Docking/SDockTab.h"
 
 #define LOCTEXT_NAMESPACE "SceneOutlinerTedsBridge"
 
-FAutoConsoleCommand BindColumnsToSceneOutlinerConsoleCommand(
-	TEXT("TEDS.UI.BindColumnsToSceneOutliner"),
-	TEXT("Bind one or more columns to the most recently used Scene Outliner. Several prebuild configurations are offered as well.")
-	TEXT("An example input to show a label column is 'TEDS.UI.BindColumnsToSceneOutliner /Script/TypedElementFramework.TypedElementLabelColumn'."),
-	FConsoleCommandWithArgsDelegate::CreateLambda([](const TArray<FString>& Args)
-		{
-			using namespace UE::Editor::DataStorage::Queries;
-
-			if (ICoreProvider* DataStorage = GetMutableDataStorageFeature<ICoreProvider>(StorageFeatureName))
-			{
-				static QueryHandle Queries[] =
-				{
-					DataStorage->RegisterQuery(Select().ReadWrite<FTypedElementLabelColumn>().Compile()),
-					DataStorage->RegisterQuery(Select().ReadOnly<FTypedElementLocalTransformColumn>().Compile()),
-					DataStorage->RegisterQuery(Select().ReadOnly<FTypedElementPackagePathColumn>().Compile()),
-					DataStorage->RegisterQuery(
-						Select()
-							.ReadWrite<FTypedElementLabelColumn>()
-							.ReadOnly<FTypedElementLocalTransformColumn>()
-						.Compile()),
-					DataStorage->RegisterQuery(
-						Select()
-							.ReadOnly<FTypedElementLabelColumn>()
-							.ReadOnly<FTypedElementLabelHashColumn>()
-						.Compile())
-				};
-
-				FSceneOutlinerTedsQueryBinder& Binder = FSceneOutlinerTedsQueryBinder::GetInstance();
-				const TWeakPtr<ILevelEditor> LevelEditor = FModuleManager::GetModuleChecked<FLevelEditorModule>(TEXT("LevelEditor")).GetLevelEditorInstance();
-				const TSharedPtr<ISceneOutliner> SceneOutliner = LevelEditor.IsValid() ? LevelEditor.Pin()->GetMostRecentlyUsedSceneOutliner() : nullptr;
-				if (SceneOutliner.IsValid())
-				{
-					if (!Args.IsEmpty())
-					{
-						if (Args[0].IsNumeric())
-						{
-							int32 QueryIndex = FCString::Atoi(*Args[0]);
-							if (QueryIndex < sizeof(Queries) / sizeof(QueryHandle))
-							{
-								Binder.AssignQuery(Queries[QueryIndex], SceneOutliner);
-								return;
-							}
-						}
-						else
-						{
-							uint32 AdditionCount = 0;
-							Select Query;
-							for (const FString& Arg : Args)
-							{
-								FTopLevelAssetPath Path;
-								// TrySetPath has an ensure that checks if the path starts with an '/' and otherwise throws
-								// an assert.
-								if (!Arg.IsEmpty() && Arg[0] == '/' && Path.TrySetPath(Arg))
-								{
-									const UScriptStruct* ColumnType = TypeOptional(Path);
-									if (ColumnType && ColumnType->IsChildOf(FColumn::StaticStruct()))
-									{
-										Query.ReadOnly(ColumnType);
-										++AdditionCount;
-									}
-								}
-							}
-							if (AdditionCount > 0)
-							{
-								static QueryHandle CustomQuery = InvalidQueryHandle;
-								if (CustomQuery != InvalidQueryHandle)
-								{
-									DataStorage->UnregisterQuery(CustomQuery);
-								}
-								CustomQuery = DataStorage->RegisterQuery(Query.Compile());
-								Binder.AssignQuery(CustomQuery, SceneOutliner);
-								return;
-							}
-						}
-					}
-					Binder.AssignQuery(InvalidQueryHandle, SceneOutliner);
-				}
-			}
-		}));
+//
+// USceneOutlinerTedsBridgeFactory
+// 
 
-class FSceneOutlinerTedsBridge
+void USceneOutlinerTedsBridgeFactory::PreRegister(UE::Editor::DataStorage::ICoreProvider& DataStorage)
 {
-public:
-	~FSceneOutlinerTedsBridge();
-
-	void Initialize(
-		UE::Editor::DataStorage::ICoreProvider& InStorage,
-		UE::Editor::DataStorage::IUiProvider& InStorageUi,
-		UE::Editor::DataStorage::ICompatibilityProvider& InStorageCompatibility,
-		const TSharedPtr<ISceneOutliner>& InOutliner);
-	
-	void AssignQuery(UE::Editor::DataStorage::QueryHandle Query,
-		const UE::Editor::DataStorage::IUiProvider::FPurposeID& CellWidgetPurpose, const UE::Editor::DataStorage::IUiProvider::FPurposeID& HeaderWidgetPurpose);
-	
-	void RegisterDealiaser(const FTreeItemIDDealiaser& InDealiaser);
-	FTreeItemIDDealiaser GetDealiaser();
-	
-private:
-	void ClearColumns(ISceneOutliner& InOutliner);
-
-	TArray<FName> AddedColumns;
-	TWeakPtr<ISceneOutliner> Outliner;
-	UE::Editor::DataStorage::ICoreProvider* Storage{ nullptr };
-	UE::Editor::DataStorage::IUiProvider* StorageUi{ nullptr };
-	UE::Editor::DataStorage::ICompatibilityProvider* StorageCompatibility{ nullptr };
-	FTreeItemIDDealiaser Dealiaser;
-	UE::Editor::DataStorage::IUiProvider::FPurposeID CellWidgetPurpose;
-};
+	// Map the type column from the TEDS to the default Outliner type column, so we can show type info for objects not in TEDS
+	TEDSToOutlinerDefaultColumnMapping.Add(FTypedElementClassTypeInfoColumn::StaticStruct(), FSceneOutlinerBuiltInColumnTypes::ActorInfo());
+	TEDSToOutlinerDefaultColumnMapping.Add(FVisibleInEditorColumn::StaticStruct(), FSceneOutlinerBuiltInColumnTypes::Gutter());
+	TEDSToOutlinerDefaultColumnMapping.Add(FNameColumn::StaticStruct(), FSceneOutlinerBuiltInColumnTypes::IDName());
+}
 
-class FOutlinerColumn : public ISceneOutlinerColumn
+void USceneOutlinerTedsBridgeFactory::RegisterTables(UE::Editor::DataStorage::ICoreProvider& DataStorage)
 {
-public:
-	FOutlinerColumn(
-		UE::Editor::DataStorage::QueryHandle InQuery,
-		UE::Editor::DataStorage::ICoreProvider& InStorage, 
-		UE::Editor::DataStorage::IUiProvider& InStorageUi, 
-		UE::Editor::DataStorage::ICompatibilityProvider& InStorageCompatibility,
-		FName InNameId, 
-		TArray<TWeakObjectPtr<const UScriptStruct>> InColumnTypes,
-		TSharedPtr<FTypedElementWidgetConstructor> InHeaderWidgetConstructor,
-		TSharedPtr<FTypedElementWidgetConstructor> InCellWidgetConstructor,
-		FName InFallbackColumnName,
-		TWeakPtr<ISceneOutliner> InOwningOutliner,
-		const FTreeItemIDDealiaser& InDealiaser)
-		: Storage(InStorage)
-		, StorageUi(InStorageUi)
-		, StorageCompatibility(InStorageCompatibility)
-		, QueryHandle(InQuery)
-		, NameId(InNameId)
-		, OwningOutliner(InOwningOutliner)
-		, Dealiaser(InDealiaser)
-	{
-		MetaData.AddOrSetMutableData(TEXT("Name"), NameId.ToString());
-
-		using namespace UE::Editor::DataStorage;
-
-		TableViewerColumnImpl = MakeUnique<FTedsTableViewerColumn>(NameId, InCellWidgetConstructor, InColumnTypes,
-			InHeaderWidgetConstructor, FComboMetaDataView(FGenericMetaDataView(MetaData)).Next(FQueryMetaDataView(Storage.GetQueryDescription(QueryHandle))));
-
-		TableViewerColumnImpl->SetIsRowVisibleDelegate(
-		FTedsTableViewerColumn::FIsRowVisible::CreateRaw(this, &FOutlinerColumn::IsRowVisible)
-		);
-		
-		// Try to find a fallback column from the regular item, for handling cases like folders which are not in TEDS but want to use TEDS columns
-		FSceneOutlinerModule& SceneOutlinerModule = FModuleManager::LoadModuleChecked<FSceneOutlinerModule>("SceneOutliner");
-		FallbackColumn = SceneOutlinerModule.FactoryColumn(InFallbackColumnName, *OwningOutliner.Pin());
-	};
-	
-	FName GetColumnID() override
-	{
-		return NameId;
-	}
-
-	virtual void Tick(double InCurrentTime, float InDeltaTime) override
-	{
-		TableViewerColumnImpl->Tick();
-		if (FallbackColumn)
-		{
-			FallbackColumn->Tick(InCurrentTime, InDeltaTime);
-		}
-	}
-
-	bool IsRowVisible(const UE::Editor::DataStorage::RowHandle InRowHandle) const
-	{
-		TSharedPtr<ISceneOutliner> OutlinerPinned = OwningOutliner.Pin();
-
-		if (!OutlinerPinned)
-		{
-			return false;
-		}
-		
-		// Try to grab the TEDS Outliner item from the row handle
-		FSceneOutlinerTreeItemPtr Item = OutlinerPinned->GetTreeItem(InRowHandle);
-
-		// If it doesn't exist, this could be a legacy item that uses something other than the row id as the ID, so check if we have a dealiaser
-		if (!Item)
-		{
-			if (Dealiaser.IsBound())
-			{
-				Item = OutlinerPinned->GetTreeItem(Dealiaser.Execute(InRowHandle));
-			}
-		}
-
-		if (!Item)
-		{
-			return false;
-		}
-
-		// Check if the item is visible in the tree
-		return OutlinerPinned->GetTree().IsItemVisible(Item);
-	}
-
-	virtual SHeaderRow::FColumn::FArguments ConstructHeaderRowColumn() override
-	{
-		return TableViewerColumnImpl->ConstructHeaderRowColumn();
-	}
-
-	// TODO: Sorting is currently handled through the fallback column if it exists because we have no way to sort columns through TEDS
-	virtual void SortItems(TArray<FSceneOutlinerTreeItemPtr>& RootItems, const EColumnSortMode::Type SortMode) const override
-	{
-		if (FallbackColumn)
-		{
-			FallbackColumn->SortItems(RootItems, SortMode);
-		}
-	}
-
-	virtual bool SupportsSorting() const override
-	{
-		return FallbackColumn ? FallbackColumn->SupportsSorting() : false;
-	}
-
-	void SetHighlightText(SWidget& Widget)
-	{
-		TSharedPtr<ISceneOutliner> OutlinerPinned = OwningOutliner.Pin();
-
-		if (!OutlinerPinned)
-		{
-			return;
-		}
-
-		if (TSharedPtr<ITypedElementUiTextCapability> TextCapability = Widget.GetMetaData<ITypedElementUiTextCapability>())
-		{
-			TextCapability->SetHighlightText(OutlinerPinned->GetFilterHighlightText());
-		}
-	
-		if (FChildren* ChildWidgets = Widget.GetChildren())
-		{
-			ChildWidgets->ForEachWidget([this](SWidget& ChildWidget)
-				{
-					SetHighlightText(ChildWidget);
-				});
-		}
-	}
-	
-	const TSharedRef<SWidget> ConstructRowWidget(FSceneOutlinerTreeItemRef TreeItem, const STableRow<FSceneOutlinerTreeItemPtr>& Row) override
-	{
-		using namespace UE::Editor::DataStorage;
-		using namespace UE::Editor::Outliner;
-		
-		RowHandle TargetRowHandle = InvalidRowHandle;
-
-		TSharedPtr<SWidget> RowWidget;
-
-		if (const FTedsOutlinerTreeItem* TedsItem = TreeItem->CastTo<FTedsOutlinerTreeItem>())
-		{
-			TargetRowHandle = TedsItem->GetRowHandle();
-			
-		}
-		else if (const FActorTreeItem* ActorItem = TreeItem->CastTo<FActorTreeItem>())
-		{
-			if (AActor* Actor = ActorItem->Actor.Get())
-			{
-				TargetRowHandle = StorageCompatibility.FindRowWithCompatibleObject(Actor);
-			}
-		}
-		else if (FallbackColumn)
-		{
-			RowWidget = FallbackColumn->ConstructRowWidget(TreeItem, Row);
-		}
-
-		if(Storage.IsRowAssigned(TargetRowHandle))
-		{
-			RowWidget = TableViewerColumnImpl->ConstructRowWidget(TargetRowHandle,
-				[&](UE::Editor::DataStorage::ICoreProvider& DataStorage, const RowHandle& WidgetRow)
-				{
-					DataStorage.AddColumn(WidgetRow, FTedsOutlinerColumn{ .Outliner = OwningOutliner });
-				});
-		}
-
-		if (RowWidget)
-		{
-			SetHighlightText(*RowWidget);
-			return RowWidget.ToSharedRef();
-		}
-
-		return SNullWidget::NullWidget;
-	}
-
-	virtual void PopulateSearchStrings( const ISceneOutlinerTreeItem& Item, TArray< FString >& OutSearchStrings ) const override
-	{
-		// TODO: We don't currently have a way to convert TEDS widgets into searchable strings, but we can rely on the fallback column if it exists
-		if (FallbackColumn)
-		{
-			FallbackColumn->PopulateSearchStrings(Item, OutSearchStrings);
-		}
-	}
-
-	// The table viewer implementation that we internally use to create our widgets
-	TUniquePtr<UE::Editor::DataStorage::FTedsTableViewerColumn> TableViewerColumnImpl;
-	
-	UE::Editor::DataStorage::ICoreProvider& Storage;
-	UE::Editor::DataStorage::IUiProvider& StorageUi;
-	UE::Editor::DataStorage::ICompatibilityProvider& StorageCompatibility;
-	UE::Editor::DataStorage::QueryHandle QueryHandle;
-	UE::Editor::DataStorage::FMetaData MetaData;
-	FName NameId;
-	TSharedPtr<ISceneOutlinerColumn> FallbackColumn;
-	TWeakPtr<ISceneOutliner> OwningOutliner;
-	FTreeItemIDDealiaser Dealiaser;
-};
-
-
-
-
-//
-// USceneOutlinerTedsBridgeFactory
-// 
+	DataStorage.RegisterTable<FTedsOutlinerColumn, FTedsOutlinerColumnQueryColumn>(UE::Editor::Outliner::Helpers::GetTedsOutlinerTableName());
+}
 
 void USceneOutlinerTedsBridgeFactory::RegisterWidgetPurposes(UE::Editor::DataStorage::IUiProvider& DataStorageUi) const
 {
@@ -369,36 +61,7 @@ void USceneOutlinerTedsBridgeFactory::RegisterWidgetPurposes(UE::Editor::DataSto
 			GeneralRowLabelPurposeID));
 }
 
-
-
-//
-// FSceneOutlinerTedsQueryBinder
-// 
-
-const FName FSceneOutlinerTedsQueryBinder::CellWidgetTableName(TEXT("Editor_SceneOutlinerCellWidgetTable"));
-const FName FSceneOutlinerTedsQueryBinder::HeaderWidgetPurpose(TEXT("SceneOutliner.Header"));
-const FName FSceneOutlinerTedsQueryBinder::CellWidgetPurpose(TEXT("SceneOutliner.Cell"));
-const FName FSceneOutlinerTedsQueryBinder::ItemLabelCellWidgetPurpose(TEXT("SceneOutliner.RowLabel"));
-
-FSceneOutlinerTedsQueryBinder::FSceneOutlinerTedsQueryBinder()
-{
-	using namespace UE::Editor::DataStorage;
-	Storage = GetMutableDataStorageFeature<UE::Editor::DataStorage::ICoreProvider>(StorageFeatureName);
-	StorageUi = GetMutableDataStorageFeature<UE::Editor::DataStorage::IUiProvider>(UiFeatureName);
-	StorageCompatibility = GetMutableDataStorageFeature<UE::Editor::DataStorage::ICompatibilityProvider>(CompatibilityFeatureName);
-
-	SetupDefaultColumnMapping();
-}
-
-void FSceneOutlinerTedsQueryBinder::SetupDefaultColumnMapping()
-{
-	// Map the type column from the TEDS to the default Outliner type column, so we can show type info for objects not in TEDS
-	TEDSToOutlinerDefaultColumnMapping.Add(FTypedElementClassTypeInfoColumn::StaticStruct(), FSceneOutlinerBuiltInColumnTypes::ActorInfo());
-	TEDSToOutlinerDefaultColumnMapping.Add(FVisibleInEditorColumn::StaticStruct(), FSceneOutlinerBuiltInColumnTypes::Gutter());
-	TEDSToOutlinerDefaultColumnMapping.Add(FNameColumn::StaticStruct(), FSceneOutlinerBuiltInColumnTypes::IDName());
-}
-
-FName FSceneOutlinerTedsQueryBinder::FindOutlinerColumnFromTEDSColumns(TConstArrayView<TWeakObjectPtr<const UScriptStruct>> TEDSColumns) const
+FName USceneOutlinerTedsBridgeFactory::FindOutlinerColumnFromTedsColumns(TConstArrayView<TWeakObjectPtr<const UScriptStruct>> TEDSColumns) const
 {
 	// Currently, the algorithm naively looks through the mapping and returns the first match
 	for(const TWeakObjectPtr<const UScriptStruct>& Column : TEDSColumns)
@@ -408,217 +71,7 @@ FName FSceneOutlinerTedsQueryBinder::FindOutlinerColumnFromTEDSColumns(TConstArr
 			return *FoundDefaultColumn;
 		}
 	}
-
-	return FName();
-}
-
-void FSceneOutlinerTedsQueryBinder::RefreshLevelEditorOutliners() const
-{
-	FLevelEditorModule& LevelEditorModule = FModuleManager::GetModuleChecked<FLevelEditorModule>("LevelEditor");
-	TSharedPtr<FTabManager> LevelEditorTabManager = LevelEditorModule.GetLevelEditorTabManager();
-
-	static const FName TabIDS[] = { LevelEditorTabIds::LevelEditorSceneOutliner, LevelEditorTabIds::LevelEditorSceneOutliner2, LevelEditorTabIds::LevelEditorSceneOutliner3, LevelEditorTabIds::LevelEditorSceneOutliner4 };
-
-	for (const FName& TabID : TabIDS)
-	{
-		if (LevelEditorTabManager->FindExistingLiveTab(TabID).IsValid())
-		{
-			LevelEditorTabManager->TryInvokeTab(TabID)->RequestCloseTab();
-			LevelEditorTabManager->TryInvokeTab(TabID);
-		}
-	}
-}
-
-FSceneOutlinerTedsQueryBinder& FSceneOutlinerTedsQueryBinder::GetInstance()
-{
-	static FSceneOutlinerTedsQueryBinder Binder;
-	return Binder;
-}
-
-TSharedPtr<FSceneOutlinerTedsBridge>* FSceneOutlinerTedsQueryBinder::FindOrAddQueryMapping(const TSharedPtr<ISceneOutliner>& Outliner)
-{
-	TSharedPtr<FSceneOutlinerTedsBridge>* QueryMapping = SceneOutliners.Find(Outliner);
-	if (QueryMapping == nullptr)
-	{
-		QueryMapping = &SceneOutliners.Add(Outliner, MakeShared<FSceneOutlinerTedsBridge>());
-		(*QueryMapping)->Initialize(*Storage, *StorageUi, *StorageCompatibility, Outliner);
-	}
-
-	return QueryMapping;
-}
-
-TSharedPtr<FSceneOutlinerTedsBridge>* FSceneOutlinerTedsQueryBinder::FindQueryMapping(const TSharedPtr<ISceneOutliner>& Outliner)
-{
-	return SceneOutliners.Find(Outliner);
-}
-
-
-void FSceneOutlinerTedsQueryBinder::AssignQuery(UE::Editor::DataStorage::QueryHandle Query, const TSharedPtr<ISceneOutliner>& Outliner,
-	const UE::Editor::DataStorage::IUiProvider::FPurposeID& InWidgetPurpose, const UE::Editor::DataStorage::IUiProvider::FPurposeID& InHeaderPurpose)
-{
-	CleanupStaleOutliners();
-
-	TSharedPtr<FSceneOutlinerTedsBridge>* QueryMapping = FindOrAddQueryMapping(Outliner);
-	(*QueryMapping)->AssignQuery(Query, InWidgetPurpose, InHeaderPurpose);
-}
-
-void FSceneOutlinerTedsQueryBinder::AssignQuery(UE::Editor::DataStorage::QueryHandle Query, const TSharedPtr<ISceneOutliner>& Outliner)
-{
-	UE::Editor::DataStorage::IUiProvider::FPurposeID CellPurpose =
-		UE::Editor::DataStorage::IUiProvider::FPurposeInfo("SceneOutliner", "Cell", NAME_None).GeneratePurposeID();
-	
-	UE::Editor::DataStorage::IUiProvider::FPurposeID HeaderPurpose =
-		UE::Editor::DataStorage::IUiProvider::FPurposeInfo("SceneOutliner", "Header", NAME_None).GeneratePurposeID();
-
-	AssignQuery(Query, Outliner, CellPurpose, HeaderPurpose);
-
-}
-
-void FSceneOutlinerTedsQueryBinder::RegisterTreeItemIDDealiaser(const TSharedPtr<ISceneOutliner>& Outliner, const FTreeItemIDDealiaser& InDealiaser)
-{
-	TSharedPtr<FSceneOutlinerTedsBridge>* QueryMapping = FindOrAddQueryMapping(Outliner);
-	(*QueryMapping)->RegisterDealiaser(InDealiaser);
-}
-
-FTreeItemIDDealiaser FSceneOutlinerTedsQueryBinder::GetTreeItemIDDealiaser(const TSharedPtr<ISceneOutliner>& Widget)
-{
-	TSharedPtr<FSceneOutlinerTedsBridge>* QueryMapping = FindQueryMapping(Widget);
-
-	if (QueryMapping)
-	{
-		return (*QueryMapping)->GetDealiaser();
-	}
-
-	return FTreeItemIDDealiaser();
-}
-
-void FSceneOutlinerTedsQueryBinder::CleanupStaleOutliners()
-{
-	for (TMap<TWeakPtr<ISceneOutliner>, TSharedPtr<FSceneOutlinerTedsBridge>>::TIterator It(SceneOutliners); It; ++It)
-	{
-		// Remove any query mappings where the target Outliner doesn't exist anymore
-		if (!It.Key().IsValid())
-		{
-			It.RemoveCurrent();
-		}
-	}
-}
-
-//
-// FSceneOutlinerTedsBridge
-//
-
-FSceneOutlinerTedsBridge::~FSceneOutlinerTedsBridge()
-{
-	TSharedPtr<ISceneOutliner> OutlinerPinned = Outliner.Pin();
-	if (OutlinerPinned)
-	{
-		ClearColumns(*OutlinerPinned);
-	}
-}
-
-void FSceneOutlinerTedsBridge::Initialize(
-	UE::Editor::DataStorage::ICoreProvider& InStorage,
-	UE::Editor::DataStorage::IUiProvider& InStorageUi,
-	UE::Editor::DataStorage::ICompatibilityProvider& InStorageCompatibility,
-	const TSharedPtr<ISceneOutliner>& InOutliner)
-{
-	Storage = &InStorage;
-	StorageUi = &InStorageUi;
-	StorageCompatibility = &InStorageCompatibility;
-	Outliner = InOutliner;
-}
-
-void FSceneOutlinerTedsBridge::RegisterDealiaser(const FTreeItemIDDealiaser& InDealiaser)
-{
-	Dealiaser = InDealiaser;
-}
-
-FTreeItemIDDealiaser FSceneOutlinerTedsBridge::GetDealiaser()
-{
-	return Dealiaser;
-}
-
-
-void FSceneOutlinerTedsBridge::AssignQuery(UE::Editor::DataStorage::QueryHandle Query,
-	const UE::Editor::DataStorage::IUiProvider::FPurposeID& InCellWidgetPurposes, const UE::Editor::DataStorage::IUiProvider::FPurposeID& HeaderWidgetPurpose)
-{
-	using MatchApproach = UE::Editor::DataStorage::IUiProvider::EMatchApproach;
-	constexpr uint8 DefaultPriorityIndex = 100;
-	FSceneOutlinerTedsQueryBinder& Binder = FSceneOutlinerTedsQueryBinder::GetInstance();
-	CellWidgetPurpose = InCellWidgetPurposes;
-
-	if (TSharedPtr<ISceneOutliner> OutlinerPinned = Outliner.Pin())
-	{
-		const UE::Editor::DataStorage::FQueryDescription& Description = Storage->GetQueryDescription(Query);
-		UE::Editor::DataStorage::FQueryMetaDataView MetaDataView(Description);
-		UE::Editor::DataStorage::RowHandle CellPurposeRow = StorageUi->FindPurpose(CellWidgetPurpose);
-		UE::Editor::DataStorage::RowHandle HeaderPurposeRow = StorageUi->FindPurpose(HeaderWidgetPurpose);
-		
-		ClearColumns(*OutlinerPinned);
-
-		if (Description.Action == UE::Editor::DataStorage::FQueryDescription::EActionType::Select)
-		{
-			int32 SelectionCount = Description.SelectionTypes.Num();
-			AddedColumns.Reset(SelectionCount);
-
-			TArray<TWeakObjectPtr<const UScriptStruct>> ColumnTypes = UE::Editor::DataStorage::TableViewerUtils::CreateVerifiedColumnTypeArray(Description.SelectionTypes);
-
-			int32 IndexOffset = 0;
-			auto ColumnConstructor = [this, Query, MetaDataView, &IndexOffset, &OutlinerPinned, Binder, HeaderPurposeRow](
-				TUniquePtr<FTypedElementWidgetConstructor> Constructor, TConstArrayView<TWeakObjectPtr<const UScriptStruct>> ColumnTypes)
-				{
-					TSharedPtr<FTypedElementWidgetConstructor> CellConstructor(Constructor.Release());
-
-					/* If we have a fallback column for this query, remove it, take over it's priority and 
-					 * replace it with the TEDS column. But also allow the TEDS-Outliner column to fallback to it for
-					 * data not in TEDS yet.
-				 	 */
-					FName FallbackColumn = Binder.FindOutlinerColumnFromTEDSColumns(ColumnTypes);
-					const FSceneOutlinerColumnInfo* FallbackColumnInfo = OutlinerPinned->GetSharedData().ColumnMap.Find(FallbackColumn);
-					uint8 ColumnPriority = FallbackColumnInfo ? FallbackColumnInfo->PriorityIndex : static_cast<uint8>(FMath::Clamp(DefaultPriorityIndex + static_cast<uint8>(FMath::Clamp(IndexOffset, 0, 255)), 0, 255));
-
-					OutlinerPinned->RemoveColumn(FallbackColumn);
-
-					FName NameId = UE::Editor::DataStorage::TableViewerUtils::FindLongestMatchingName(ColumnTypes, IndexOffset);
-					FText DisplayName = CellConstructor.Get()->CreateWidgetDisplayNameText(Storage);
-					AddedColumns.Add(NameId);
-					OutlinerPinned->AddColumn(NameId,
-						FSceneOutlinerColumnInfo(
-							ESceneOutlinerColumnVisibility::Visible, 
-							ColumnPriority,
-							FCreateSceneOutlinerColumn::CreateLambda(
-								[this, Query, MetaDataView, NameId, &ColumnTypes, CellConstructor, &OutlinerPinned, FallbackColumn, HeaderPurposeRow](ISceneOutliner&)
-								{
-									TSharedPtr<FTypedElementWidgetConstructor> HeaderConstructor = 
-										UE::Editor::DataStorage::TableViewerUtils::CreateHeaderWidgetConstructor(*StorageUi, MetaDataView, ColumnTypes, HeaderPurposeRow);
-									return MakeShared<FOutlinerColumn>(
-										Query, *Storage, *StorageUi, *StorageCompatibility, NameId,
-										TArray<TWeakObjectPtr<const UScriptStruct>>(ColumnTypes.GetData(), ColumnTypes.Num()), 
-										MoveTemp(HeaderConstructor), CellConstructor, FallbackColumn, OutlinerPinned, Dealiaser);
-
-								}),
-							true,
-							TOptional<float>(),
-							DisplayName
-						)
-					);
-					++IndexOffset;
-					return true;
-				};
-
-			StorageUi->CreateWidgetConstructors(CellPurposeRow, MatchApproach::LongestMatch, ColumnTypes, 
-				MetaDataView, ColumnConstructor);
-		}
-	}
-}
-
-void FSceneOutlinerTedsBridge::ClearColumns(ISceneOutliner& InOutliner)
-{
-	for (FName ColumnName : AddedColumns)
-	{
-		InOutliner.RemoveColumn(ColumnName);
-	}
+	return NAME_None;
 }
 
 #undef LOCTEXT_NAMESPACE
diff --git a/Engine/Plugins/Experimental/EditorDataStorageFeatures/Source/TedsOutliner/Private/TedsOutlinerColumn.cpp b/Engine/Plugins/Experimental/EditorDataStorageFeatures/Source/TedsOutliner/Private/TedsOutlinerColumn.cpp
new file mode 100644
--- /dev/null
+++ b/Engine/Plugins/Experimental/EditorDataStorageFeatures/Source/TedsOutliner/Private/TedsOutlinerColumn.cpp
@@ -0,0 +1,173 @@
+// Copyright Epic Games, Inc. All Rights Reserved.
+
+#include "TedsOutlinerColumn.h"
+
+#include "ActorTreeItem.h"
+#include "Elements/Interfaces/Capabilities/TypedElementUiTextCapability.h"
+#include "Modules/ModuleManager.h"
+#include "SceneOutlinerModule.h"
+#include "TedsOutlinerHelpers.h"
+#include "TedsOutlinerItem.h"
+
+namespace UE::Editor::Outliner
+{
+	FTedsOutlinerUiColumn::FTedsOutlinerUiColumn(
+		DataStorage::FQueryDescription& InQuery,
+		DataStorage::ICoreProvider& InStorage, 
+		DataStorage::IUiProvider& InStorageUi, 
+		DataStorage::ICompatibilityProvider& InStorageCompatibility,
+		FName InNameId, 
+		TArray<TWeakObjectPtr<const UScriptStruct>> InColumnTypes,
+		TSharedPtr<FTypedElementWidgetConstructor> InHeaderWidgetConstructor,
+		TSharedPtr<FTypedElementWidgetConstructor> InCellWidgetConstructor,
+		FName InFallbackColumnName,
+		TWeakPtr<ISceneOutliner> InOwningOutliner,
+		const FTreeItemIDDealiaser& InDealiaser)
+		: Storage(InStorage)
+		, StorageUi(InStorageUi)
+		, StorageCompatibility(InStorageCompatibility)
+		, QueryDescription(InQuery)
+		, NameId(InNameId)
+		, OwningOutliner(InOwningOutliner)
+		, Dealiaser(InDealiaser)
+	{
+		MetaData.AddOrSetMutableData(TEXT("Name"), NameId.ToString());
+
+		using namespace UE::Editor::DataStorage;
+
+		TableViewerColumnImpl = MakeUnique<FTedsTableViewerColumn>(NameId, InCellWidgetConstructor, InColumnTypes,
+			InHeaderWidgetConstructor, FComboMetaDataView(FGenericMetaDataView(MetaData)).Next(FQueryMetaDataView(QueryDescription)));
+
+		TableViewerColumnImpl->SetIsRowVisibleDelegate(
+			FTedsTableViewerColumn::FIsRowVisible::CreateRaw(this, &FTedsOutlinerUiColumn::IsRowVisible)
+		);
+		
+		// Try to find a fallback column from the regular item, for handling cases like folders which are not in TEDS but want to use TEDS columns
+		FSceneOutlinerModule& SceneOutlinerModule = FModuleManager::LoadModuleChecked<FSceneOutlinerModule>("SceneOutliner");
+		FallbackColumn = SceneOutlinerModule.FactoryColumn(InFallbackColumnName, *OwningOutliner.Pin());
+	};
+		
+	FName FTedsOutlinerUiColumn::GetColumnID()
+	{
+		return NameId;
+	}
+
+	 void FTedsOutlinerUiColumn::Tick(double InCurrentTime, float InDeltaTime)
+	{
+		TableViewerColumnImpl->Tick();
+		if (FallbackColumn)
+		{
+			FallbackColumn->Tick(InCurrentTime, InDeltaTime);
+		}
+	}
+
+	bool FTedsOutlinerUiColumn::IsRowVisible(const UE::Editor::DataStorage::RowHandle InRowHandle) const
+	{
+		TSharedPtr<ISceneOutliner> OutlinerPinned = OwningOutliner.Pin();
+
+		if (!OutlinerPinned)
+		{
+			return false;
+		}
+		
+		// Try to grab the TEDS Outliner item from the row handle
+		FSceneOutlinerTreeItemPtr Item = Helpers::GetTreeItemFromRowHandle(&Storage, OutlinerPinned.ToSharedRef(), InRowHandle);
+
+		// Check if the item is visible in the tree
+		return OutlinerPinned->GetTree().IsItemVisible(Item);
+	}
+
+	SHeaderRow::FColumn::FArguments FTedsOutlinerUiColumn::ConstructHeaderRowColumn()
+	{
+		return TableViewerColumnImpl->ConstructHeaderRowColumn();
+	}
+
+	// TODO: Sorting is currently handled through the fallback column if it exists because we have no way to sort columns through TEDS
+	void FTedsOutlinerUiColumn::SortItems(TArray<FSceneOutlinerTreeItemPtr>& RootItems, const EColumnSortMode::Type SortMode) const
+	{
+		if (FallbackColumn)
+		{
+			FallbackColumn->SortItems(RootItems, SortMode);
+		}
+	}
+
+	bool FTedsOutlinerUiColumn::SupportsSorting() const 
+	{
+		return FallbackColumn ? FallbackColumn->SupportsSorting() : false;
+	}
+
+	void FTedsOutlinerUiColumn::SetHighlightText(SWidget& Widget)
+	{
+		TSharedPtr<ISceneOutliner> OutlinerPinned = OwningOutliner.Pin();
+
+		if (!OutlinerPinned)
+		{
+			return;
+		}
+
+		if (TSharedPtr<ITypedElementUiTextCapability> TextCapability = Widget.GetMetaData<ITypedElementUiTextCapability>())
+		{
+			TextCapability->SetHighlightText(OutlinerPinned->GetFilterHighlightText());
+		}
+	
+		if (FChildren* ChildWidgets = Widget.GetChildren())
+		{
+			ChildWidgets->ForEachWidget([this](SWidget& ChildWidget)
+				{
+					SetHighlightText(ChildWidget);
+				});
+		}
+	}
+	
+	const TSharedRef<SWidget> FTedsOutlinerUiColumn::ConstructRowWidget(FSceneOutlinerTreeItemRef TreeItem, const STableRow<FSceneOutlinerTreeItemPtr>& Row)
+	{
+		using namespace UE::Editor::DataStorage;
+		
+		RowHandle TargetRowHandle = InvalidRowHandle;
+
+		TSharedPtr<SWidget> RowWidget;
+
+		if (const FTedsOutlinerTreeItem* TedsItem = TreeItem->CastTo<FTedsOutlinerTreeItem>())
+		{
+			TargetRowHandle = TedsItem->GetRowHandle();
+			
+		}
+		else if (const FActorTreeItem* ActorItem = TreeItem->CastTo<FActorTreeItem>())
+		{
+			if (AActor* Actor = ActorItem->Actor.Get())
+			{
+				TargetRowHandle = StorageCompatibility.FindRowWithCompatibleObject(Actor);
+			}
+		}
+		else if (FallbackColumn)
+		{
+			RowWidget = FallbackColumn->ConstructRowWidget(TreeItem, Row);
+		}
+
+		if(Storage.IsRowAssigned(TargetRowHandle))
+		{
+			RowWidget = TableViewerColumnImpl->ConstructRowWidget(TargetRowHandle,
+				[&](ICoreProvider& DataStorage, const RowHandle& WidgetRow)
+				{
+					DataStorage.AddColumn(WidgetRow, FTedsOutlinerColumn{ .Outliner = OwningOutliner });
+				});
+		}
+
+		if (RowWidget)
+		{
+			SetHighlightText(*RowWidget);
+			return RowWidget.ToSharedRef();
+		}
+
+		return SNullWidget::NullWidget;
+	}
+
+	void FTedsOutlinerUiColumn::PopulateSearchStrings( const ISceneOutlinerTreeItem& Item, TArray< FString >& OutSearchStrings ) const
+	{
+		// TODO: We don't currently have a way to convert TEDS widgets into searchable strings, but we can rely on the fallback column if it exists
+		if (FallbackColumn)
+		{
+			FallbackColumn->PopulateSearchStrings(Item, OutSearchStrings);
+		}
+	}
+}
diff --git a/Engine/Plugins/Experimental/EditorDataStorageFeatures/Source/TedsOutliner/Private/TedsOutlinerColumn.h b/Engine/Plugins/Experimental/EditorDataStorageFeatures/Source/TedsOutliner/Private/TedsOutlinerColumn.h
new file mode 100644
--- /dev/null
+++ b/Engine/Plugins/Experimental/EditorDataStorageFeatures/Source/TedsOutliner/Private/TedsOutlinerColumn.h
@@ -0,0 +1,58 @@
+// Copyright Epic Games, Inc. All Rights Reserved.
+
+#pragma once
+
+#include "ISceneOutlinerColumn.h"
+#include "TedsTableViewerColumn.h"
+#include "Columns/TedsOutlinerColumns.h"
+#include "Elements/Common/TypedElementHandles.h"
+#include "Elements/Interfaces/TypedElementDataStorageUiInterface.h"
+
+namespace UE::Editor::Outliner
+{
+	class FTedsOutlinerUiColumn : public ISceneOutlinerColumn
+	{
+	public:
+		
+		FTedsOutlinerUiColumn(
+			DataStorage::FQueryDescription& InQuery,
+			DataStorage::ICoreProvider& InStorage, 
+			DataStorage::IUiProvider& InStorageUi, 
+			DataStorage::ICompatibilityProvider& InStorageCompatibility,
+			FName InNameId, 
+			TArray<TWeakObjectPtr<const UScriptStruct>> InColumnTypes,
+			TSharedPtr<FTypedElementWidgetConstructor> InHeaderWidgetConstructor,
+			TSharedPtr<FTypedElementWidgetConstructor> InCellWidgetConstructor,
+			FName InFallbackColumnName,
+			TWeakPtr<ISceneOutliner> InOwningOutliner,
+			const FTreeItemIDDealiaser& InDealiaser);
+		
+		FName GetColumnID() override;
+
+		virtual void Tick(double InCurrentTime, float InDeltaTime) override;
+		virtual SHeaderRow::FColumn::FArguments ConstructHeaderRowColumn() override;
+		virtual void SortItems(TArray<FSceneOutlinerTreeItemPtr>& RootItems, const EColumnSortMode::Type SortMode) const override;
+		virtual bool SupportsSorting() const override;
+		const TSharedRef<SWidget> ConstructRowWidget(FSceneOutlinerTreeItemRef TreeItem, const STableRow<FSceneOutlinerTreeItemPtr>& Row) override;
+		virtual void PopulateSearchStrings( const ISceneOutlinerTreeItem& Item, TArray< FString >& OutSearchStrings ) const override;
+		
+		void SetHighlightText(SWidget& Widget);
+		bool IsRowVisible(const DataStorage::RowHandle InRowHandle) const;
+
+	protected:
+
+		// The table viewer implementation that we internally use to create our widgets
+		TUniquePtr<DataStorage::FTedsTableViewerColumn> TableViewerColumnImpl;
+		
+		DataStorage::ICoreProvider& Storage;
+		DataStorage::IUiProvider& StorageUi;
+		DataStorage::ICompatibilityProvider& StorageCompatibility;
+		DataStorage::FQueryDescription QueryDescription;
+		DataStorage::FMetaData MetaData;
+		FName NameId;
+		TSharedPtr<ISceneOutlinerColumn> FallbackColumn;
+		TWeakPtr<ISceneOutliner> OwningOutliner;
+		FTreeItemIDDealiaser Dealiaser;
+	};
+
+}
diff --git a/Engine/Plugins/Experimental/EditorDataStorageFeatures/Source/TedsOutliner/Private/TedsOutlinerHelpers.cpp b/Engine/Plugins/Experimental/EditorDataStorageFeatures/Source/TedsOutliner/Private/TedsOutlinerHelpers.cpp
new file mode 100644
--- /dev/null
+++ b/Engine/Plugins/Experimental/EditorDataStorageFeatures/Source/TedsOutliner/Private/TedsOutlinerHelpers.cpp
@@ -0,0 +1,86 @@
+// Copyright Epic Games, Inc. All Rights Reserved.
+
+#include "TedsOutlinerHelpers.h"
+
+#include "Columns/TedsOutlinerColumns.h"
+#include "Compatibility/SceneOutlinerTedsBridge.h"
+#include "Elements/Interfaces/TypedElementDataStorageInterface.h"
+#include "Modules/ModuleManager.h"
+#include "ISceneOutliner.h"
+#include "LevelEditor.h"
+#include "Widgets/Docking/SDockTab.h"
+
+namespace UE::Editor::Outliner::Helpers
+{
+	FSceneOutlinerTreeItemPtr GetTreeItemFromRowHandle(const DataStorage::ICoreProvider* Storage, TSharedRef<ISceneOutliner> InOutliner, DataStorage::RowHandle InRowHandle)
+	{
+		// Check if the item is indexed by the row handle for the trivial case
+		FSceneOutlinerTreeItemPtr Item = InOutliner->GetTreeItem(InRowHandle);
+
+		if (Item)
+		{
+			return Item;
+		}
+
+		// Otherwise look up any dealiasers that might be bound to this outliner instance
+		DataStorage::RowHandle OutlinerRow = Storage->LookupMappedRow(MappingDomain, DataStorage::FMapKey(InOutliner->GetOutlinerIdentifier()));
+
+		if (Storage->IsRowAvailable(OutlinerRow))
+		{
+			if (const FTedsOutlinerDealiaserColumn* DealiaserColumn = Storage->GetColumn<FTedsOutlinerDealiaserColumn>(OutlinerRow))
+			{
+				if (DealiaserColumn->Dealiaser.IsBound())
+				{
+					return InOutliner->GetTreeItem(DealiaserColumn->Dealiaser.Execute(InRowHandle));
+				}
+			}
+		}
+		return nullptr;
+		
+	}
+	
+	bool RegisterOutlinerDealiaser(DataStorage::ICoreProvider* Storage, TSharedRef<ISceneOutliner> InOutliner, const FTreeItemIDDealiaser& InDealiaser)
+	{
+		DataStorage::RowHandle OutlinerRow = Storage->LookupMappedRow(MappingDomain, DataStorage::FMapKey(InOutliner->GetOutlinerIdentifier()));
+
+		if (Storage->IsRowAvailable(OutlinerRow))
+		{
+			Storage->AddColumn(OutlinerRow, FTedsOutlinerDealiaserColumn{.Dealiaser = InDealiaser});
+			return true;
+		}
+
+		return false;
+	}
+	
+	FName GetTedsOutlinerTableName()
+	{
+		return TEXT("Editor_TedsOutlinerTable");
+	}
+	
+	void RefreshLevelEditorOutliners()
+	{
+		FLevelEditorModule& LevelEditorModule = FModuleManager::GetModuleChecked<FLevelEditorModule>("LevelEditor");
+		TSharedPtr<FTabManager> LevelEditorTabManager = LevelEditorModule.GetLevelEditorTabManager();
+
+		static const FName TabIDS[] = { LevelEditorTabIds::LevelEditorSceneOutliner, LevelEditorTabIds::LevelEditorSceneOutliner2, LevelEditorTabIds::LevelEditorSceneOutliner3, LevelEditorTabIds::LevelEditorSceneOutliner4 };
+
+		for (const FName& TabID : TabIDS)
+		{
+			if (LevelEditorTabManager->FindExistingLiveTab(TabID).IsValid())
+			{
+				LevelEditorTabManager->TryInvokeTab(TabID)->RequestCloseTab();
+				LevelEditorTabManager->TryInvokeTab(TabID);
+			}
+		}
+	}
+	
+	FName FindOutlinerColumnFromTedsColumns(const DataStorage::ICoreProvider* Storage, TConstArrayView<TWeakObjectPtr<const UScriptStruct>> TEDSColumns)
+	{
+		if(const USceneOutlinerTedsBridgeFactory* Factory = Storage->FindFactory<USceneOutlinerTedsBridgeFactory>())
+		{
+			return Factory->FindOutlinerColumnFromTedsColumns(TEDSColumns);
+		}
+
+		return NAME_None;
+	}
+}
diff --git a/Engine/Plugins/Experimental/EditorDataStorageFeatures/Source/TedsOutliner/Private/TedsOutlinerImpl.cpp b/Engine/Plugins/Experimental/EditorDataStorageFeatures/Source/TedsOutliner/Private/TedsOutlinerImpl.cpp
--- a/Engine/Plugins/Experimental/EditorDataStorageFeatures/Source/TedsOutliner/Private/TedsOutlinerImpl.cpp
+++ b/Engine/Plugins/Experimental/EditorDataStorageFeatures/Source/TedsOutliner/Private/TedsOutlinerImpl.cpp
@@ -3,6 +3,7 @@
 #include "TedsOutlinerImpl.h"
 
 #include "SSceneOutliner.h"
+#include "TedsOutlinerColumn.h"
 #include "TedsTableViewerUtils.h"
 #include "Elements/Columns/TypedElementLabelColumns.h"
 #include "Elements/Columns/TypedElementMiscColumns.h"
@@ -15,6 +16,7 @@
 #include "Columns/SlateDelegateColumns.h"
 #include "Compatibility/SceneOutlinerRowHandleColumn.h"
 #include "TedsOutlinerFilter.h"
+#include "TedsOutlinerHelpers.h"
 #include "TedsOutlinerItem.h"
 #include "TedsQueryMergeNode.h"
 #include "TedsQueryNode.h"
@@ -24,6 +26,7 @@
 #include "TedsRowMonitorNode.h"
 #include "TedsRowQueryResultsNode.h"
 #include "Columns/TedsOutlinerColumns.h"
+#include "Compatibility/SceneOutlinerTedsBridge.h"
 #include "Elements/Columns/TypedElementCompatibilityColumns.h"
 #include "Elements/Columns/TypedElementOverrideColumns.h"
 #include "Elements/Columns/TypedElementSelectionColumns.h"
@@ -192,6 +195,7 @@ FTedsOutlinerParams::FTedsOutlinerParams(SSceneOutliner* InSceneOutliner)
 FTedsOutlinerImpl::FTedsOutlinerImpl(const FTedsOutlinerParams& InParams, ISceneOutlinerMode* InMode)
 	: CreationParams(InParams)
 	, CellWidgetPurpose(InParams.CellWidgetPurpose)
+	, HeaderWidgetPurpose(InParams.HeaderWidgetPurpose)
 	, LabelWidgetPurpose(InParams.LabelWidgetPurpose)
 	, InitialQueryDescription(InParams.QueryDescription)
 	, HierarchyData(InParams.HierarchyData)
@@ -271,11 +275,19 @@ void FTedsOutlinerImpl::CreateFilterQueries()
 
 void FTedsOutlinerImpl::Init()
 {
+	RegisterTedsOutliner();
 	CreateFilterQueries();
 	RecompileQueries();
 
 	// Tick post TEDS update to make sure all processors have run and the data is correct
 	Storage->OnUpdateCompleted().AddRaw(this, &FTedsOutlinerImpl::PostDataStorageUpdate);
+
+	// FTedsOutlinerImpl::Init is called before the Outliner UI has been fully init, so we generate the columns in the next tick
+	FTSTicker::GetCoreTicker().AddTicker(FTickerDelegate::CreateLambda([this](float)
+	{
+		RegenerateColumns();
+		return false;
+	}));
 }
 
 void FTedsOutlinerImpl::OnMonitoredRowsAdded(DataStorage::FRowHandleArrayView InRows)
@@ -304,6 +316,7 @@ FTedsOutlinerImpl::~FTedsOutlinerImpl()
 {
 	if (Storage)
 	{
+		Storage->RemoveRow(OutlinerRowHandle);
 		Storage->OnUpdateCompleted().RemoveAll(this);
 	}
 	UnregisterQueries();
@@ -408,6 +421,128 @@ bool FTedsOutlinerImpl::CanDisplayRow(DataStorage::RowHandle ItemRowHandle) cons
 	return true;
 }
 
+void FTedsOutlinerImpl::RegisterTedsOutliner()
+{
+	DataStorage::TableHandle OutlinerTable = Storage->FindTable(UE::Editor::Outliner::Helpers::GetTedsOutlinerTableName());
+
+	if (OutlinerTable != DataStorage::InvalidTableHandle)
+	{
+		OutlinerRowHandle = Storage->AddRow(OutlinerTable);
+
+		if (FTedsOutlinerColumn* TedsOutlinerColumn = Storage->GetColumn<FTedsOutlinerColumn>(OutlinerRowHandle))
+		{
+			TedsOutlinerColumn->Outliner = StaticCastSharedRef<ISceneOutliner>(SceneOutliner->AsShared());
+		}
+		
+		if (FTedsOutlinerColumnQueryColumn* ColumnQueryColumn = Storage->GetColumn<FTedsOutlinerColumnQueryColumn>(OutlinerRowHandle))
+		{
+			ColumnQueryColumn->ColumnQueryDescription = CreationParams.ColumnQueryDescription;
+
+			ColumnQueryColumn->OnRefreshColumns.BindSP(this, &FTedsOutlinerImpl::RegenerateColumns);
+		}
+
+		FName Identifier = SceneOutliner->GetOutlinerIdentifier();
+
+		if (!Identifier.IsNone())
+		{
+			Storage->MapRow(MappingDomain, DataStorage::FMapKey(SceneOutliner->GetOutlinerIdentifier()), OutlinerRowHandle);
+		}
+	}
+}
+
+void FTedsOutlinerImpl::ClearColumns() const
+{
+	for (FName ColumnName : AddedColumns)
+	{
+		SceneOutliner->RemoveColumn(ColumnName);
+	}
+}
+
+void FTedsOutlinerImpl::RegenerateColumns()
+{
+	ClearColumns();
+	
+	FTedsOutlinerColumnQueryColumn* ColumnQueryColumn = Storage->GetColumn<FTedsOutlinerColumnQueryColumn>(OutlinerRowHandle);
+	FTedsOutlinerColumn* TedsOutlinerColumn = Storage->GetColumn<FTedsOutlinerColumn>(OutlinerRowHandle);
+
+	if  (!ColumnQueryColumn || !TedsOutlinerColumn)
+	{
+		return;
+	}
+
+	FTreeItemIDDealiaser Dealiaser;
+
+	if (FTedsOutlinerDealiaserColumn* DealiaserColumn = Storage->GetColumn<FTedsOutlinerDealiaserColumn>(OutlinerRowHandle))
+	{
+		Dealiaser = DealiaserColumn->Dealiaser;
+	}
+
+	using MatchApproach = DataStorage::IUiProvider::EMatchApproach;
+	constexpr uint8 DefaultPriorityIndex = 100;
+
+	DataStorage::FQueryMetaDataView MetaDataView(ColumnQueryColumn->ColumnQueryDescription);
+	
+	int32 SelectionCount = ColumnQueryColumn->ColumnQueryDescription.SelectionTypes.Num();
+	AddedColumns.Reset(SelectionCount);
+
+	TArray<TWeakObjectPtr<const UScriptStruct>> ColumnTypes =
+		DataStorage::TableViewerUtils::CreateVerifiedColumnTypeArray(ColumnQueryColumn->ColumnQueryDescription.SelectionTypes);
+	
+	DataStorage::RowHandle CellPurposeRow = StorageUi->FindPurpose(CellWidgetPurpose);
+	DataStorage::RowHandle HeaderPurposeRow = StorageUi->FindPurpose(HeaderWidgetPurpose);
+
+	int32 IndexOffset = 0;
+	auto ColumnConstructor = [this, ColumnQueryColumn, TedsOutlinerColumn, MetaDataView, &IndexOffset, HeaderPurposeRow, Dealiaser](
+		TUniquePtr<FTypedElementWidgetConstructor> Constructor, TConstArrayView<TWeakObjectPtr<const UScriptStruct>> ColumnTypes)
+		{
+			TSharedPtr<FTypedElementWidgetConstructor> CellConstructor(Constructor.Release());
+
+			/* If we have a fallback column for this query, remove it, take over it's priority and 
+			 * replace it with the TEDS column. But also allow the TEDS-Outliner column to fallback to it for
+			 * data not in TEDS yet.
+			 */
+			FName FallbackColumn = Helpers::FindOutlinerColumnFromTedsColumns(Storage, ColumnTypes);
+			const FSceneOutlinerColumnInfo* FallbackColumnInfo = SceneOutliner->GetSharedData().ColumnMap.Find(FallbackColumn);
+			uint8 ColumnPriority = FallbackColumnInfo
+				? FallbackColumnInfo->PriorityIndex
+				: static_cast<uint8>(
+					FMath::Clamp(DefaultPriorityIndex + static_cast<uint8>(
+						FMath::Clamp(IndexOffset, 0, 255)), 0, 255));
+
+			SceneOutliner->RemoveColumn(FallbackColumn);
+
+			FName NameId = UE::Editor::DataStorage::TableViewerUtils::FindLongestMatchingName(ColumnTypes, IndexOffset);
+			FText DisplayName = CellConstructor.Get()->CreateWidgetDisplayNameText(Storage);
+			AddedColumns.Add(NameId);
+			SceneOutliner->AddColumn(NameId,
+				FSceneOutlinerColumnInfo(
+					ESceneOutlinerColumnVisibility::Visible, 
+					ColumnPriority,
+					FCreateSceneOutlinerColumn::CreateLambda(
+						[this, ColumnQueryColumn, TedsOutlinerColumn, MetaDataView, NameId, &ColumnTypes, CellConstructor, FallbackColumn, HeaderPurposeRow, Dealiaser](ISceneOutliner&)
+						{
+							TSharedPtr<FTypedElementWidgetConstructor> HeaderConstructor = 
+								UE::Editor::DataStorage::TableViewerUtils::CreateHeaderWidgetConstructor(*StorageUi, MetaDataView, ColumnTypes, HeaderPurposeRow);
+							
+							return MakeShared<FTedsOutlinerUiColumn>(
+								ColumnQueryColumn->ColumnQueryDescription, *Storage, *StorageUi, *StorageCompatibility, NameId,
+								TArray<TWeakObjectPtr<const UScriptStruct>>(ColumnTypes.GetData(), ColumnTypes.Num()), 
+								MoveTemp(HeaderConstructor), CellConstructor, FallbackColumn, TedsOutlinerColumn->Outliner, Dealiaser);
+
+						}),
+					true,
+					TOptional<float>(),
+					DisplayName
+				)
+			);
+			++IndexOffset;
+			return true;
+		};
+
+	StorageUi->CreateWidgetConstructors(CellPurposeRow, MatchApproach::LongestMatch, ColumnTypes, 
+		MetaDataView, ColumnConstructor);
+}
+
 void FTedsOutlinerImpl::CreateItemsFromQuery(TArray<FSceneOutlinerTreeItemPtr>& OutItems, ISceneOutlinerMode* InMode) const
 {
 	using namespace UE::Editor::DataStorage::Queries;
diff --git a/Engine/Plugins/Experimental/EditorDataStorageFeatures/Source/TedsOutliner/Private/TedsOutlinerModule.cpp b/Engine/Plugins/Experimental/EditorDataStorageFeatures/Source/TedsOutliner/Private/TedsOutlinerModule.cpp
--- a/Engine/Plugins/Experimental/EditorDataStorageFeatures/Source/TedsOutliner/Private/TedsOutlinerModule.cpp
+++ b/Engine/Plugins/Experimental/EditorDataStorageFeatures/Source/TedsOutliner/Private/TedsOutlinerModule.cpp
@@ -65,7 +65,7 @@ FTedsOutlinerModule::FTedsOutlinerModule()
 {
 }
 
-TSharedRef<ISceneOutliner> FTedsOutlinerModule::CreateTedsOutliner(const FSceneOutlinerInitializationOptions& InInitOptions, const FTedsOutlinerParams& InInitTedsOptions, DataStorage::QueryHandle ColumnQuery) const
+TSharedRef<ISceneOutliner> FTedsOutlinerModule::CreateTedsOutliner(const FSceneOutlinerInitializationOptions& InInitOptions, const FTedsOutlinerParams& InInitTedsOptions) const
 {
 	using namespace UE::Editor::DataStorage;
 	ensureMsgf(AreEditorDataStorageFeaturesEnabled(), TEXT("Unable to initialize the Teds-Outliner before TEDS itself is initialized."));
@@ -95,8 +95,6 @@ TSharedRef<ISceneOutliner> FTedsOutlinerModule::CreateTedsOutliner(const FSceneO
 	
 	TSharedRef<ISceneOutliner> TedsOutlinerShared = SNew(SSceneOutliner, InitOptions);
 	
-	FSceneOutlinerTedsQueryBinder::GetInstance().AssignQuery(ColumnQuery, TedsOutlinerShared, InInitTedsOptions.CellWidgetPurpose, InInitTedsOptions.HeaderWidgetPurpose);
-	
 	return TedsOutlinerShared;
 }
 
@@ -114,7 +112,7 @@ void FTedsOutlinerModule::ShutdownModule()
 	IModuleInterface::ShutdownModule();
 }
 
-DataStorage::QueryHandle FTedsOutlinerModule::GetLevelEditorTedsOutlinerColumnQuery()
+UE::Editor::DataStorage::FQueryDescription FTedsOutlinerModule::GetLevelEditorTedsOutlinerColumnQueryDescription()
 {
 	using namespace UE::Editor::DataStorage;
 	using namespace UE::Editor::DataStorage::Columns;
@@ -122,27 +120,26 @@ DataStorage::QueryHandle FTedsOutlinerModule::GetLevelEditorTedsOutlinerColumnQu
 		
 	using namespace DataStorage::Queries;
 
-	static DataStorage::QueryHandle ColumnQuery = Storage->RegisterQuery(
+	static FQueryDescription ColumnQuery = 
 		Select()
 			.ReadOnly<FTypedElementClassTypeInfoColumn,
 					  FAlertColumn, 
 					  FChildAlertColumn,
 					  FNameColumn,
 					  FVisibleInEditorColumn>()
-		.Compile());
+		.Compile();
 
 	// Query to also include revision control info
-	static DataStorage::QueryHandle RevisionControlQuery = Storage->RegisterQuery(
+	static FQueryDescription RevisionControlQuery = 
 		Select()
 			.ReadOnly<FTypedElementClassTypeInfoColumn,  
 					  FTypedElementPackageReference, 
 					  FAlertColumn, 
 					  FNameColumn,
 					  FVisibleInEditorColumn>()
-		.Compile());
+		.Compile();
 
 	return Private::bUseNewRevisionControlWidgets ? RevisionControlQuery : ColumnQuery;
-	
 }
 
 TSharedRef<SWidget> FTedsOutlinerModule::CreateLevelEditorTedsOutliner()
@@ -170,6 +167,7 @@ TSharedRef<SWidget> FTedsOutlinerModule::CreateLevelEditorTedsOutliner()
 
 	FTedsOutlinerParams Params(nullptr);
 	Params.QueryDescription = OutlinerQueryDescription;
+	Params.ColumnQueryDescription = GetLevelEditorTedsOutlinerColumnQueryDescription();
 	Params.bUseDefaultTedsFilters = true;
 
 	// Example Query to filter for actors
@@ -183,7 +181,7 @@ TSharedRef<SWidget> FTedsOutlinerModule::CreateLevelEditorTedsOutliner()
 	// Empty selection set name is currently the level editor
 	Params.SelectionSetOverride = FName();
 	
-	TSharedRef<ISceneOutliner> TEDSOutlinerShared = CreateTedsOutliner(InitOptions, Params, GetLevelEditorTedsOutlinerColumnQuery());
+	TSharedRef<ISceneOutliner> TEDSOutlinerShared = CreateTedsOutliner(InitOptions, Params);
 	
 	return TEDSOutlinerShared;
 }
diff --git a/Engine/Plugins/Experimental/EditorDataStorageFeatures/Source/TedsOutliner/Private/Widgets/OutlinerFolderLabelWidget.cpp b/Engine/Plugins/Experimental/EditorDataStorageFeatures/Source/TedsOutliner/Private/Widgets/OutlinerFolderLabelWidget.cpp
--- a/Engine/Plugins/Experimental/EditorDataStorageFeatures/Source/TedsOutliner/Private/Widgets/OutlinerFolderLabelWidget.cpp
+++ b/Engine/Plugins/Experimental/EditorDataStorageFeatures/Source/TedsOutliner/Private/Widgets/OutlinerFolderLabelWidget.cpp
@@ -14,6 +14,7 @@
 #include "Elements/Framework/TypedElementAttributeBinding.h"
 #include "ISceneOutliner.h"
 #include "SceneOutlinerHelpers.h"
+#include "TedsOutlinerHelpers.h"
 #include "Widgets/Images/SImage.h"
 #include "Widgets/Layout/SSpacer.h"
 #include "Widgets/Layout/SBox.h"
@@ -43,12 +44,9 @@ TSharedPtr<SWidget> FOutlinerFolderLabelWidgetConstructor::CreateWidget(UE::Edit
 	const UE::Editor::DataStorage::FMetaDataView& Arguments)
 {
 	using namespace UE::Editor::DataStorage;
+	
 	if (DataStorage->IsRowAvailable(TargetRow))
 	{
-		FAttributeBinder Binder(TargetRow, DataStorage);
-
-		
-
 		return SNew(SHorizontalBox)
 			+SHorizontalBox::Slot()
 			.AutoWidth()
@@ -265,13 +263,7 @@ TWeakPtr<ISceneOutlinerTreeItem> FOutlinerFolderLabelWidgetConstructor::GetTreeI
 	{
 		if (TSharedPtr<ISceneOutliner> Outliner = TedsOutlinerColumn->Outliner.Pin())
 		{
-			FSceneOutlinerTedsQueryBinder& Binder = FSceneOutlinerTedsQueryBinder::GetInstance();
-          
-			if (FTreeItemIDDealiaser Dealiaser = Binder.GetTreeItemIDDealiaser(Outliner);
-			   Dealiaser.IsBound())
-			{
-				return Outliner->GetTreeItem(Dealiaser.Execute(TargetRow));
-			}
+			return UE::Editor::Outliner::Helpers::GetTreeItemFromRowHandle(DataStorage, Outliner.ToSharedRef(), TargetRow);
 		}
 	}
 
diff --git a/Engine/Plugins/Experimental/EditorDataStorageFeatures/Source/TedsOutliner/Private/Widgets/OutlinerFolderLabelWidget.h b/Engine/Plugins/Experimental/EditorDataStorageFeatures/Source/TedsOutliner/Private/Widgets/OutlinerFolderLabelWidget.h
--- a/Engine/Plugins/Experimental/EditorDataStorageFeatures/Source/TedsOutliner/Private/Widgets/OutlinerFolderLabelWidget.h
+++ b/Engine/Plugins/Experimental/EditorDataStorageFeatures/Source/TedsOutliner/Private/Widgets/OutlinerFolderLabelWidget.h
@@ -2,6 +2,7 @@
 
 #pragma once
 
+#include "Columns/TedsOutlinerColumns.h"
 #include "Elements/Interfaces/TypedElementDataStorageFactory.h"
 #include "Elements/Interfaces/TypedElementDataStorageUiInterface.h"
 #include "Internationalization/Text.h"
@@ -51,5 +52,6 @@ protected:
 	static FText GetTooltipText(UE::Editor::DataStorage::ICoreProvider* DataStorage, UE::Editor::DataStorage::RowHandle TargetRow);
 	static FSlateColor GetForegroundColor(UE::Editor::DataStorage::ICoreProvider* DataStorage, UE::Editor::DataStorage::RowHandle TargetRow, UE::Editor::DataStorage::RowHandle WidgetRow);
 	
-	static TWeakPtr<ISceneOutlinerTreeItem> GetTreeItemForRow(UE::Editor::DataStorage::ICoreProvider* DataStorage, UE::Editor::DataStorage::RowHandle TargetRow, UE::Editor::DataStorage::RowHandle WidgetRow);
+	static TWeakPtr<ISceneOutlinerTreeItem> GetTreeItemForRow(UE::Editor::DataStorage::ICoreProvider* DataStorage,
+		UE::Editor::DataStorage::RowHandle TargetRow, UE::Editor::DataStorage::RowHandle WidgetRow);
 };
\ No newline at end of file
diff --git a/Engine/Plugins/Experimental/EditorDataStorageFeatures/Source/TedsOutliner/Private/Widgets/OutlinerVisibilityWidget.cpp b/Engine/Plugins/Experimental/EditorDataStorageFeatures/Source/TedsOutliner/Private/Widgets/OutlinerVisibilityWidget.cpp
--- a/Engine/Plugins/Experimental/EditorDataStorageFeatures/Source/TedsOutliner/Private/Widgets/OutlinerVisibilityWidget.cpp
+++ b/Engine/Plugins/Experimental/EditorDataStorageFeatures/Source/TedsOutliner/Private/Widgets/OutlinerVisibilityWidget.cpp
@@ -17,6 +17,7 @@
 #include "Widgets/SBoxPanel.h"
 #include "ISceneOutliner.h"
 #include "ActorTreeItem.h"
+#include "TedsOutlinerHelpers.h"
 #include "TedsOutlinerItem.h"
 
 #define LOCTEXT_NAMESPACE "OutlinerVisibilityWidget"
@@ -361,14 +362,7 @@ FSceneOutlinerTreeItemPtr STedsVisibilityWidget::GetTreeItem(RowHandle InRow) co
 		return nullptr;
 	}
 
-	FSceneOutlinerTedsQueryBinder& Binder = FSceneOutlinerTedsQueryBinder::GetInstance();
-	if (FTreeItemIDDealiaser Dealiaser = Binder.GetTreeItemIDDealiaser(Outliner); Dealiaser.IsBound())
-	{
-		FSceneOutlinerTreeItemID ID = Dealiaser.Execute(InRow);
-		return Outliner->GetTreeItem(ID);
-	}
-
-	return nullptr;
+	return UE::Editor::Outliner::Helpers::GetTreeItemFromRowHandle(DataStorage, Outliner.ToSharedRef(), InRow);
 }
 
 UE::Editor::DataStorage::ICoreProvider* STedsVisibilityWidget::GetDataStorage()
diff --git a/Engine/Plugins/Experimental/EditorDataStorageFeatures/Source/TedsOutliner/Public/Columns/TedsOutlinerColumns.h b/Engine/Plugins/Experimental/EditorDataStorageFeatures/Source/TedsOutliner/Public/Columns/TedsOutlinerColumns.h
--- a/Engine/Plugins/Experimental/EditorDataStorageFeatures/Source/TedsOutliner/Public/Columns/TedsOutlinerColumns.h
+++ b/Engine/Plugins/Experimental/EditorDataStorageFeatures/Source/TedsOutliner/Public/Columns/TedsOutlinerColumns.h
@@ -3,6 +3,7 @@
 #pragma once
 
 #include "Elements/Interfaces/TypedElementDataStorageInterface.h"
+#include "SceneOutlinerPublicTypes.h"
 #include "Templates/SharedPointer.h"
 
 #include "TedsOutlinerColumns.generated.h"
@@ -11,6 +12,18 @@ class ISceneOutliner;
 class UToolMenu;
 class SSceneOutliner;
 
+namespace UE::Editor::Outliner
+{
+	// Delegate that refreshes columns on the Teds Outliner it is bound to when executed
+	DECLARE_DELEGATE(FOnRefreshColumns)
+
+	// Delegate to get the TreeItemID for a row
+	DECLARE_DELEGATE_RetVal_OneParam(FSceneOutlinerTreeItemID, FTreeItemIDDealiaser, DataStorage::RowHandle);
+
+	// Mapping domain name for Outliners
+	inline static const FName MappingDomain = "TedsOutliner";
+}
+
 // Column used to store a reference to the Teds Outliner owning a specific row
 // Currently only added to widget rows in the Teds Outliner
 USTRUCT(meta = (DisplayName = "Owning Table Viewer"))
@@ -30,4 +43,26 @@ struct FTedsOutlinerContextMenuColumn final : public FEditorDataStorageColumn
 	GENERATED_BODY()
 
 	FTedsOutlinerContextMenuDelegate OnCreateContextMenu;
-};
\ No newline at end of file
+};
+
+/** Column used to store information about the columns being viewed in the Teds Outliner */
+USTRUCT()
+struct FTedsOutlinerColumnQueryColumn final : public FEditorDataStorageColumn
+{
+	GENERATED_BODY()
+
+	// The SELECT() part of the query description specifies the columns to be displayed
+	UE::Editor::DataStorage::FQueryDescription ColumnQueryDescription;
+
+	// Call this delegate to refresh the columns in the Teds Outliner with the values specified in ColumnQueryDescription
+	UE::Editor::Outliner::FOnRefreshColumns OnRefreshColumns;
+};
+
+/** Column used to store a dealiaser from outliner item -> row handle */
+USTRUCT()
+struct FTedsOutlinerDealiaserColumn final : public FEditorDataStorageColumn
+{
+	GENERATED_BODY()
+
+	UE::Editor::Outliner::FTreeItemIDDealiaser Dealiaser;
+};
diff --git a/Engine/Plugins/Experimental/EditorDataStorageFeatures/Source/TedsOutliner/Public/Compatibility/SceneOutlinerTedsBridge.h b/Engine/Plugins/Experimental/EditorDataStorageFeatures/Source/TedsOutliner/Public/Compatibility/SceneOutlinerTedsBridge.h
--- a/Engine/Plugins/Experimental/EditorDataStorageFeatures/Source/TedsOutliner/Public/Compatibility/SceneOutlinerTedsBridge.h
+++ b/Engine/Plugins/Experimental/EditorDataStorageFeatures/Source/TedsOutliner/Public/Compatibility/SceneOutlinerTedsBridge.h
@@ -21,61 +21,6 @@ namespace UE::Editor::DataStorage
 	class ICompatibilityProvider;
 } // namespace UE::Editor::DataStorage
 
-DECLARE_DELEGATE_RetVal_OneParam(FSceneOutlinerTreeItemID, FTreeItemIDDealiaser, UE::Editor::DataStorage::RowHandle);
-
-class FSceneOutlinerTedsBridge;
-
-/**
- * Utility class to bind Typed Elements Data Storage queries to a Scene Outliner. The provided query is expected to be a select query
- * and will be used to populate the Scene Outliner in addition to already existing data.
- */
-class TEDSOUTLINER_API FSceneOutlinerTedsQueryBinder
-{
-public:
-	static const FName CellWidgetTableName;
-	static const FName HeaderWidgetPurpose;
-	static const FName CellWidgetPurpose;
-	static const FName ItemLabelCellWidgetPurpose;
-
-
-	static FSceneOutlinerTedsQueryBinder& GetInstance();
-
-	void AssignQuery(UE::Editor::DataStorage::QueryHandle Query, const TSharedPtr<ISceneOutliner>& Widget,
-		const UE::Editor::DataStorage::IUiProvider::FPurposeID& InWidgetPurpose, const UE::Editor::DataStorage::IUiProvider::FPurposeID& InHeaderPurpose);
-
-	// Assign query and use default widget purposes for the scene outliner
-	void AssignQuery(UE::Editor::DataStorage::QueryHandle Query, const TSharedPtr<ISceneOutliner>& Outliner);
-
-	// Register a dealiaser for a specific TEDS-Outliner to convert a row handle to an FSceneOutlinerTreeItemID
-	void RegisterTreeItemIDDealiaser(const TSharedPtr<ISceneOutliner>& Widget, const FTreeItemIDDealiaser& InDealiaser);
-
-	// Get the dealiaser for a specific outliner instance
-	FTreeItemIDDealiaser GetTreeItemIDDealiaser(const TSharedPtr<ISceneOutliner>& Widget);
-
-	// Get the name of the Outliner column corresponding to the given TEDS column (if any)
-	FName FindOutlinerColumnFromTEDSColumns(TConstArrayView<TWeakObjectPtr<const UScriptStruct>> TEDSColumns) const;
-
-	// Refresh all the Outliners currently open in the level editor
-	void RefreshLevelEditorOutliners() const;
-
-private:
-	FSceneOutlinerTedsQueryBinder();
-	void SetupDefaultColumnMapping();
-	void CleanupStaleOutliners();
-	
-	TSharedPtr<FSceneOutlinerTedsBridge>* FindOrAddQueryMapping(const TSharedPtr<ISceneOutliner>& Widget);
-	TSharedPtr<FSceneOutlinerTedsBridge>* FindQueryMapping(const TSharedPtr<ISceneOutliner>& Widget);
-
-	
-	TMap<TWeakPtr<ISceneOutliner>, TSharedPtr<FSceneOutlinerTedsBridge>> SceneOutliners;
-
-	UE::Editor::DataStorage::ICoreProvider* Storage{ nullptr };
-	UE::Editor::DataStorage::IUiProvider* StorageUi{ nullptr };
-	UE::Editor::DataStorage::ICompatibilityProvider* StorageCompatibility{ nullptr };
-
-	TMap<TWeakObjectPtr<const UScriptStruct>, FName> TEDSToOutlinerDefaultColumnMapping;
-};
-
 UCLASS()
 class USceneOutlinerTedsBridgeFactory : public UEditorDataStorageFactory
 {
@@ -84,5 +29,12 @@ class USceneOutlinerTedsBridgeFactory : public UEditorDataStorageFactory
 public:
 	~USceneOutlinerTedsBridgeFactory() override = default;
 
-	void RegisterWidgetPurposes(UE::Editor::DataStorage::IUiProvider& DataStorageUi) const override;
+	virtual void PreRegister(UE::Editor::DataStorage::ICoreProvider& DataStorage) override;
+	virtual void RegisterTables(UE::Editor::DataStorage::ICoreProvider& DataStorage) override;
+	virtual void RegisterWidgetPurposes(UE::Editor::DataStorage::IUiProvider& DataStorageUi) const override;
+
+	FName FindOutlinerColumnFromTedsColumns(TConstArrayView<TWeakObjectPtr<const UScriptStruct>> TEDSColumns) const;
+protected:
+
+	TMap<TWeakObjectPtr<const UScriptStruct>, FName> TEDSToOutlinerDefaultColumnMapping;
 };
diff --git a/Engine/Plugins/Experimental/EditorDataStorageFeatures/Source/TedsOutliner/Public/TedsOutlinerHelpers.h b/Engine/Plugins/Experimental/EditorDataStorageFeatures/Source/TedsOutliner/Public/TedsOutlinerHelpers.h
new file mode 100644
--- /dev/null
+++ b/Engine/Plugins/Experimental/EditorDataStorageFeatures/Source/TedsOutliner/Public/TedsOutlinerHelpers.h
@@ -0,0 +1,28 @@
+// Copyright Epic Games, Inc. All Rights Reserved.
+
+#pragma once
+
+#include "Elements/Common/TypedElementHandles.h"
+#include "Elements/Interfaces/TypedElementDataStorageInterface.h"
+#include "SceneOutlinerFwd.h"
+#include "Columns/TedsOutlinerColumns.h"
+
+// Helper functions for Teds Outliner related functionality
+namespace UE::Editor::Outliner::Helpers
+{
+	// Get the internal Outliner tree item from the row handle from the specific Outliner
+	TEDSOUTLINER_API FSceneOutlinerTreeItemPtr GetTreeItemFromRowHandle(const DataStorage::ICoreProvider* Storage, TSharedRef<ISceneOutliner> InOutliner, DataStorage::RowHandle InRowHandle);
+
+	// Register a dealiaser for the given outliner, overriding any previously registered ones.
+	// Returns false if the registration failed, mostly commonly because the Outliner doesn't have a unique OutlinerIdentifier specificed
+	TEDSOUTLINER_API bool RegisterOutlinerDealiaser(DataStorage::ICoreProvider* Storage, TSharedRef<ISceneOutliner> InOutliner, const FTreeItemIDDealiaser& InDealiaser);
+
+	// Get the name of the Teds table that all Outliners are stored in
+	TEDSOUTLINER_API FName GetTedsOutlinerTableName();
+	
+	// Refresh all the Outliners currently open in the level editor
+	TEDSOUTLINER_API void RefreshLevelEditorOutliners();
+	
+	// Get the name of the Outliner column corresponding to the given TEDS column (if any)
+	FName FindOutlinerColumnFromTedsColumns(const DataStorage::ICoreProvider* Storage, TConstArrayView<TWeakObjectPtr<const UScriptStruct>> TEDSColumns);
+}
diff --git a/Engine/Plugins/Experimental/EditorDataStorageFeatures/Source/TedsOutliner/Public/TedsOutlinerImpl.h b/Engine/Plugins/Experimental/EditorDataStorageFeatures/Source/TedsOutliner/Public/TedsOutlinerImpl.h
--- a/Engine/Plugins/Experimental/EditorDataStorageFeatures/Source/TedsOutliner/Public/TedsOutlinerImpl.h
+++ b/Engine/Plugins/Experimental/EditorDataStorageFeatures/Source/TedsOutliner/Public/TedsOutlinerImpl.h
@@ -96,6 +96,9 @@ struct FTedsOutlinerParams
 
 	// The query description that will be used to populate rows in the TEDS-Outliner
 	DataStorage::FQueryDescription QueryDescription;
+
+	// The query description that will be used to populate the initial columns in the Teds Outliner
+	DataStorage::FQueryDescription ColumnQueryDescription;
 	
 	// TEDS queries that will be used to create filters in this Outliner
 	// TEDS-Outliner TODO: Can we consolidate this with the SceneOutliner API to create filters? Currently has to be separate because FTEDSOutlinerFilter
@@ -217,6 +220,12 @@ protected:
 	// Check if this row can be displayed in this table viewer
 	bool CanDisplayRow(DataStorage::RowHandle ItemRowHandle) const;
 
+	// Register this Teds Outliner with TEDS
+	void RegisterTedsOutliner();
+
+	void ClearColumns() const;
+	void RegenerateColumns();
+
 protected:
 	// TEDS Storage Constructs
 	DataStorage::ICoreProvider* Storage{ nullptr };
@@ -227,6 +236,9 @@ protected:
 
 	// The purpose to use when generating widgets for row/column pairs through TEDS UI
 	DataStorage::IUiProvider::FPurposeID CellWidgetPurpose;
+	
+	// The purpose to use when generating widgets for headers pairs through TEDS UI
+	DataStorage::IUiProvider::FPurposeID HeaderWidgetPurpose;
 
 	// The purpose to use when generating widgets for the "Item Label" column through TEDS UI
 	DataStorage::IUiProvider::FPurposeID LabelWidgetPurpose;
@@ -290,5 +302,11 @@ protected:
 	// Addition and Label Updates are deferred because they access data storage implicitly
 	// Deletion is currently not deferred to work nicely with object lifecycles in some cases - but can be once everything goes through the query stack
 	TSet<DataStorage::RowHandle> RowsPendingLabelUpdate;
+
+	// The Row Handle associated with this outliner
+	DataStorage::RowHandle OutlinerRowHandle = DataStorage::InvalidRowHandle;
+
+	// The columns currently being displayed by the Teds Outliner
+	TArray<FName> AddedColumns;
 };
 } // namsepace UE::Editor::Outliner
diff --git a/Engine/Plugins/Experimental/EditorDataStorageFeatures/Source/TedsOutliner/Public/TedsOutlinerModule.h b/Engine/Plugins/Experimental/EditorDataStorageFeatures/Source/TedsOutliner/Public/TedsOutlinerModule.h
--- a/Engine/Plugins/Experimental/EditorDataStorageFeatures/Source/TedsOutliner/Public/TedsOutlinerModule.h
+++ b/Engine/Plugins/Experimental/EditorDataStorageFeatures/Source/TedsOutliner/Public/TedsOutlinerModule.h
@@ -6,6 +6,7 @@
 #include "Modules/ModuleInterface.h"
 #include "Templates/SharedPointer.h"
 #include "Delegates/IDelegateInstance.h"
+#include "Elements/Common/TypedElementQueryDescription.h"
 
 class ISceneOutliner;
 struct FSceneOutlinerInitializationOptions;
@@ -37,14 +38,14 @@ public:
 	 * @return	New scene outliner widget
 	 */
 	virtual TSharedRef<ISceneOutliner> CreateTedsOutliner(
-		const FSceneOutlinerInitializationOptions& InInitOptions, const FTedsOutlinerParams& InInitTedsOptions, UE::Editor::DataStorage::QueryHandle ColumnQuery) const;
+		const FSceneOutlinerInitializationOptions& InInitOptions, const FTedsOutlinerParams& InInitTedsOptions) const;
 
 	// IModuleInterface interface
 	virtual void StartupModule() override;
 	virtual void ShutdownModule() override;
 
 	// Get the column query the default table viewer uses
-	virtual UE::Editor::DataStorage::QueryHandle GetLevelEditorTedsOutlinerColumnQuery();
+	virtual UE::Editor::DataStorage::FQueryDescription GetLevelEditorTedsOutlinerColumnQueryDescription();
 
 	// The name of the tab the default table viewer is opened in
 	FName GetTedsOutlinerTabName();
diff --git a/Engine/Source/Editor/SceneOutliner/Public/ISceneOutliner.h b/Engine/Source/Editor/SceneOutliner/Public/ISceneOutliner.h
--- a/Engine/Source/Editor/SceneOutliner/Public/ISceneOutliner.h
+++ b/Engine/Source/Editor/SceneOutliner/Public/ISceneOutliner.h
@@ -192,6 +192,9 @@ public:
 	/** Set a delay (in seconds) before the the next requested UI refresh executes. Note: This only applies until the next refresh that executes
 	 * and any subsequent ones will be immediate. */
 	virtual void SetNextUIRefreshDelay(float InDelay) = 0;
+
+	/** Get the unqiue identifier of this Outliner (if specified, NAME_None otherwise) */
+	virtual FName GetOutlinerIdentifier() const = 0;
 protected:
 	ISceneOutlinerMode* Mode;
 };
diff --git a/Engine/Source/Editor/SceneOutliner/Public/SSceneOutliner.h b/Engine/Source/Editor/SceneOutliner/Public/SSceneOutliner.h
--- a/Engine/Source/Editor/SceneOutliner/Public/SSceneOutliner.h
+++ b/Engine/Source/Editor/SceneOutliner/Public/SSceneOutliner.h
@@ -554,7 +554,7 @@ public:
 	bool GetIsReentrant() const { return bIsReentrant; }
 
 	/** Get the unique identifier associated with this outliner */
-	FName GetOutlinerIdentifier() const { return OutlinerIdentifier; }
+	virtual FName GetOutlinerIdentifier() const override { return OutlinerIdentifier; }
 
 	/** Toggle if SceneOutliner should show Transient objects */
 	void SetShowTransient(bool bInShowTransient) { SharedData.Get()->bShowTransient = bInShowTransient; }
