[TEDS] Made GetMutableRows required for all row nodes in the Query Stack.

The concept of parent nodes that allowed their children to be modified already existed through the`GetMutableRows` function but this was optional. This resulted in some nodes requiring it and other didn't. This eventually led to unexpected asserts. The update changes the access to `GetMutableRows` to be mandatory and all places where multiple cases were supported are now updated to only return their own row list or the one inherited from their parent. As a result the array view node has been updated to be a reference view into a row handle array instead of having a const view in between.

diff --git a/Engine/Plugins/Experimental/EditorDataStorageFeatures/Source/TedsContentBrowser/Private/TedsContentBrowserModule.cpp b/Engine/Plugins/Experimental/EditorDataStorageFeatures/Source/TedsContentBrowser/Private/TedsContentBrowserModule.cpp
--- a/Engine/Plugins/Experimental/EditorDataStorageFeatures/Source/TedsContentBrowser/Private/TedsContentBrowserModule.cpp
+++ b/Engine/Plugins/Experimental/EditorDataStorageFeatures/Source/TedsContentBrowser/Private/TedsContentBrowserModule.cpp
@@ -122,12 +122,7 @@ namespace UE::Editor::ContentBrowser
 			return;
 		}
 
-		FRowHandleArray* Rows = RowQueryStack->GetMutableRows();
-		if (ensureMsgf(Rows, TEXT("Expected the row query stack to be mutable since it is used in the table viewer")))
-		{
-			Rows->Empty();
-		}
-		
+		RowQueryStack->GetMutableRows().Empty();
 		ContentBrowserItemMap.Empty();
 
 		for(TSharedPtr<FAssetViewItem> Item : *InItemsSource)
@@ -151,11 +146,8 @@ namespace UE::Editor::ContentBrowser
 		{
 			ContentBrowserItemMap.Emplace(RowHandle, Item);
 			
-			FRowHandleArray* Rows = RowQueryStack->GetMutableRows();
-			if (ensureMsgf(Rows, TEXT("Expected the row query stack to be mutable since it is used in the table viewer")))
-			{
-				Rows->Add(RowHandle);
-			}
+			FRowHandleArray& Rows = RowQueryStack->GetMutableRows();
+			Rows.Add(RowHandle);
 		}
 	}
 
diff --git a/Engine/Plugins/Experimental/EditorDataStorageFeatures/Source/TedsQueryStack/Private/TedsRowArrayNode.cpp b/Engine/Plugins/Experimental/EditorDataStorageFeatures/Source/TedsQueryStack/Private/TedsRowArrayNode.cpp
--- a/Engine/Plugins/Experimental/EditorDataStorageFeatures/Source/TedsQueryStack/Private/TedsRowArrayNode.cpp
+++ b/Engine/Plugins/Experimental/EditorDataStorageFeatures/Source/TedsQueryStack/Private/TedsRowArrayNode.cpp
@@ -28,8 +28,8 @@ namespace UE::Editor::DataStorage::QueryStack
 		return Rows.GetRows();
 	}
 
-	FRowHandleArray* FRowArrayNode::GetMutableRows()
+	FRowHandleArray& FRowArrayNode::GetMutableRows()
 	{
-		return &Rows;
+		return Rows;
 	}
 } // namespace UE::Editor::DataStorage::QueryStack
diff --git a/Engine/Plugins/Experimental/EditorDataStorageFeatures/Source/TedsQueryStack/Private/TedsRowCopyNode.cpp b/Engine/Plugins/Experimental/EditorDataStorageFeatures/Source/TedsQueryStack/Private/TedsRowCopyNode.cpp
--- a/Engine/Plugins/Experimental/EditorDataStorageFeatures/Source/TedsQueryStack/Private/TedsRowCopyNode.cpp
+++ b/Engine/Plugins/Experimental/EditorDataStorageFeatures/Source/TedsQueryStack/Private/TedsRowCopyNode.cpp
@@ -38,8 +38,8 @@ namespace UE::Editor::DataStorage::QueryStack
 		return Rows.GetRows();
 	}
 
-	FRowHandleArray* FRowCopyNode::GetMutableRows()
+	FRowHandleArray& FRowCopyNode::GetMutableRows()
 	{
-		return &Rows;
+		return Rows;
 	}
 } // namespace UE::Editor::DataStorage::QueryStack
diff --git a/Engine/Plugins/Experimental/EditorDataStorageFeatures/Source/TedsQueryStack/Private/TedsRowHandleSortNode.cpp b/Engine/Plugins/Experimental/EditorDataStorageFeatures/Source/TedsQueryStack/Private/TedsRowHandleSortNode.cpp
--- a/Engine/Plugins/Experimental/EditorDataStorageFeatures/Source/TedsQueryStack/Private/TedsRowHandleSortNode.cpp
+++ b/Engine/Plugins/Experimental/EditorDataStorageFeatures/Source/TedsQueryStack/Private/TedsRowHandleSortNode.cpp
@@ -7,11 +7,8 @@ namespace UE::Editor::DataStorage::QueryStack
 	FRowHandleSortNode::FRowHandleSortNode(TSharedPtr<IRowNode> InParent)
 		: Parent(MoveTemp(InParent))
 	{
-		FRowHandleArray* ParentRows = Parent->GetMutableRows();
-		checkf(ParentRows != nullptr, 
-			TEXT("FRowHandleSortNode requires a parent node that provides write access to its row handle array."));
-
-		ParentRows->Sort();
+		FRowHandleArray& ParentRows = Parent->GetMutableRows();
+		ParentRows.Sort();
 		ParentRevision = Parent->GetRevision();
 	}
 
@@ -26,7 +23,7 @@ namespace UE::Editor::DataStorage::QueryStack
 
 		if (Parent->GetRevision() != ParentRevision)
 		{
-			Parent->GetMutableRows()->Sort();
+			Parent->GetMutableRows().Sort();
 			ParentRevision = Parent->GetRevision();
 		}
 	}
@@ -36,7 +33,7 @@ namespace UE::Editor::DataStorage::QueryStack
 		return Parent->GetRows();
 	}
 
-	FRowHandleArray* FRowHandleSortNode::GetMutableRows()
+	FRowHandleArray& FRowHandleSortNode::GetMutableRows()
 	{
 		return Parent->GetMutableRows();
 	}
diff --git a/Engine/Plugins/Experimental/EditorDataStorageFeatures/Source/TedsQueryStack/Private/TedsRowMergeNode.cpp b/Engine/Plugins/Experimental/EditorDataStorageFeatures/Source/TedsQueryStack/Private/TedsRowMergeNode.cpp
--- a/Engine/Plugins/Experimental/EditorDataStorageFeatures/Source/TedsQueryStack/Private/TedsRowMergeNode.cpp
+++ b/Engine/Plugins/Experimental/EditorDataStorageFeatures/Source/TedsQueryStack/Private/TedsRowMergeNode.cpp
@@ -47,9 +47,9 @@ namespace UE::Editor::DataStorage::QueryStack
 		return Rows.GetRows();
 	}
 
-	FRowHandleArray* FRowMergeNode::GetMutableRows()
+	FRowHandleArray& FRowMergeNode::GetMutableRows()
 	{
-		return &Rows;
+		return Rows;
 	}
 
 	void FRowMergeNode::Merge()
diff --git a/Engine/Plugins/Experimental/EditorDataStorageFeatures/Source/TedsQueryStack/Private/TedsRowMonitorNode.cpp b/Engine/Plugins/Experimental/EditorDataStorageFeatures/Source/TedsQueryStack/Private/TedsRowMonitorNode.cpp
--- a/Engine/Plugins/Experimental/EditorDataStorageFeatures/Source/TedsQueryStack/Private/TedsRowMonitorNode.cpp
+++ b/Engine/Plugins/Experimental/EditorDataStorageFeatures/Source/TedsQueryStack/Private/TedsRowMonitorNode.cpp
@@ -8,36 +8,33 @@
 namespace UE::Editor::DataStorage::QueryStack
 {
 	FRowMonitorNode::FRowMonitorNode(
-		ICoreProvider& Storage, TSharedPtr<IRowNode> ParentRow, TSharedPtr<IQueryNode> QueryNode)
+		ICoreProvider& Storage, TSharedPtr<IRowNode> ParentNode, TSharedPtr<IQueryNode> QueryNode)
 		: QueryNode(MoveTemp(QueryNode))
-		, ParentRow(MoveTemp(ParentRow))
+		, ParentNode(MoveTemp(ParentNode))
 		, Storage(Storage)
 	{
-		UpdateRows();
 		UpdateColumnsFromQuery();
 		UpdateMonitoredColumns();
 	}
 
 	FRowMonitorNode::FRowMonitorNode(
-		ICoreProvider& Storage, TSharedPtr<IRowNode> ParentRow, TArray<TObjectPtr<const UScriptStruct>> Columns)
+		ICoreProvider& Storage, TSharedPtr<IRowNode> ParentNode, TArray<TObjectPtr<const UScriptStruct>> Columns)
 		: MonitoredColumns(MoveTemp(Columns))
-		, ParentRow(MoveTemp(ParentRow))
+		, ParentNode(MoveTemp(ParentNode))
 		, Storage(Storage)
 	{
-		UpdateRows();
 		UpdateMonitoredColumns();
 	}
 
 	FRowMonitorNode::FRowMonitorNode(
-		ICoreProvider& Storage, TSharedPtr<IQueryNode> QueryNode, TSharedPtr<IRowNode> ParentRow, 
+		ICoreProvider& Storage, TSharedPtr<IQueryNode> QueryNode, TSharedPtr<IRowNode> ParentNode,
 		TArray<TObjectPtr<const UScriptStruct>> MonitoredColumns)
 		: MonitoredColumns(MoveTemp(MonitoredColumns))
 		, QueryNode(MoveTemp(QueryNode))
-		, ParentRow(MoveTemp(ParentRow))
+		, ParentNode(MoveTemp(ParentNode))
 		, Storage(Storage)
 		, bFixedColumns(true)
 	{
-		UpdateRows();
 		UpdateColumnsFromQuery();
 		UpdateMonitoredColumns();
 	}
@@ -95,12 +92,11 @@ namespace UE::Editor::DataStorage::QueryStack
 		}
 		
 		bool MergeChanges = true;
-		if (ParentRow)
+		if (ParentNode)
 		{
-			ParentRow->Update();
-			if (ParentRow->GetRevision() != ParentRevision)
+			ParentNode->Update();
+			if (ParentNode->GetRevision() != ParentRevision)
 			{
-				UpdateRows();
 				Revision++;
 				MergeChanges = false;
 			}
@@ -114,7 +110,7 @@ namespace UE::Editor::DataStorage::QueryStack
 			RemovedRows.Sort();
 			RemovedRows.MakeUnique();
 
-			FRowHandleArray& TargetRows = ResolveRows();
+			FRowHandleArray& TargetRows = ParentNode->GetMutableRows();
 			TargetRows.SortedMerge(AddedRows);
 			TargetRows.Remove(RemovedRows.GetRows());
 
@@ -136,12 +132,12 @@ namespace UE::Editor::DataStorage::QueryStack
 
 	FRowHandleArrayView FRowMonitorNode::GetRows() const
 	{
-		return ResolveRows().GetRows();
+		return ParentNode->GetRows();
 	}
 
-	FRowHandleArray* FRowMonitorNode::GetMutableRows()
+	FRowHandleArray& FRowMonitorNode::GetMutableRows()
 	{
-		return &ResolveRows();
+		return ParentNode->GetMutableRows();
 	}
 
 	FRowMonitorNode::FOnRowsChanged& FRowMonitorNode::OnMonitoredRowsAdded()
@@ -154,16 +150,6 @@ namespace UE::Editor::DataStorage::QueryStack
 		return RowsRemovedEvent;
 	}
 
-	const FRowHandleArray& FRowMonitorNode::ResolveRows() const
-	{
-		return const_cast<FRowMonitorNode*>(this)->ResolveRows();
-	}
-	FRowHandleArray& FRowMonitorNode::ResolveRows()
-	{
-		FRowHandleArray* ParentRowArray = ParentRow ? ParentRow->GetMutableRows() : nullptr;
-		return ParentRowArray == nullptr ? Rows : *ParentRowArray;
-	}
-
 	void FRowMonitorNode::UpdateColumnsFromQuery()
 	{
 		using namespace UE::Editor::DataStorage::Queries;
@@ -235,16 +221,6 @@ namespace UE::Editor::DataStorage::QueryStack
 		}
 	}
 
-	void FRowMonitorNode::UpdateRows()
-	{
-		if (FRowHandleArray* ParentRowArray = ParentRow->GetMutableRows(); ParentRowArray == nullptr)
-		{
-			Rows.Empty();
-			Rows.Append(ParentRow->GetRows());
-		}
-		ParentRevision = ParentRow->GetRevision();
-	}
-
 	void FRowMonitorNode::UpdateMonitoredColumns()
 	{
 		using namespace UE::Editor::DataStorage::Queries;
diff --git a/Engine/Plugins/Experimental/EditorDataStorageFeatures/Source/TedsQueryStack/Private/TedsRowOrderInversionNode.cpp b/Engine/Plugins/Experimental/EditorDataStorageFeatures/Source/TedsQueryStack/Private/TedsRowOrderInversionNode.cpp
--- a/Engine/Plugins/Experimental/EditorDataStorageFeatures/Source/TedsQueryStack/Private/TedsRowOrderInversionNode.cpp
+++ b/Engine/Plugins/Experimental/EditorDataStorageFeatures/Source/TedsQueryStack/Private/TedsRowOrderInversionNode.cpp
@@ -8,9 +8,6 @@ namespace UE::Editor::DataStorage::QueryStack
 		: Parent(MoveTemp(InParent))
 		, bIsEnabled(bInIsEnabled)
 	{
-		FRowHandleArray* Rows = Parent->GetMutableRows();
-		checkf(Rows != nullptr, TEXT("FRowOrderInversionNode requires a parent that can provide mutable rows."));
-
 		if (bInIsEnabled)
 		{
 			ApplyOrder();
@@ -62,18 +59,14 @@ namespace UE::Editor::DataStorage::QueryStack
 		return Parent->GetRows();
 	}
 
-	FRowHandleArray* FRowOrderInversionNode::GetMutableRows()
+	FRowHandleArray& FRowOrderInversionNode::GetMutableRows()
 	{
 		return Parent->GetMutableRows();
 	}
 
 	void FRowOrderInversionNode::ApplyOrder()
 	{
-		if (FRowHandleArray* Rows = Parent->GetMutableRows())
-		{
-			Rows->InvertOrder();
-		}
-
+		Parent->GetMutableRows().InvertOrder();
 		Revision++;
 	}
 } // namespace UE::Editor::DataStorage::QueryStack
diff --git a/Engine/Plugins/Experimental/EditorDataStorageFeatures/Source/TedsQueryStack/Private/TedsRowQueryResultsNode.cpp b/Engine/Plugins/Experimental/EditorDataStorageFeatures/Source/TedsQueryStack/Private/TedsRowQueryResultsNode.cpp
--- a/Engine/Plugins/Experimental/EditorDataStorageFeatures/Source/TedsQueryStack/Private/TedsRowQueryResultsNode.cpp
+++ b/Engine/Plugins/Experimental/EditorDataStorageFeatures/Source/TedsQueryStack/Private/TedsRowQueryResultsNode.cpp
@@ -69,9 +69,9 @@ namespace UE::Editor::DataStorage::QueryStack
 		return Rows.GetRows();
 	}
 
-	FRowHandleArray* FRowQueryResultsNode::GetMutableRows()
+	FRowHandleArray& FRowQueryResultsNode::GetMutableRows()
 	{
-		return &Rows;
+		return Rows;
 	}
 
 	void FRowQueryResultsNode::RefreshInternal(FRowHandleArray& TargetRows)
diff --git a/Engine/Plugins/Experimental/EditorDataStorageFeatures/Source/TedsQueryStack/Private/TedsRowSortNode.cpp b/Engine/Plugins/Experimental/EditorDataStorageFeatures/Source/TedsQueryStack/Private/TedsRowSortNode.cpp
--- a/Engine/Plugins/Experimental/EditorDataStorageFeatures/Source/TedsQueryStack/Private/TedsRowSortNode.cpp
+++ b/Engine/Plugins/Experimental/EditorDataStorageFeatures/Source/TedsQueryStack/Private/TedsRowSortNode.cpp
@@ -25,27 +25,21 @@ namespace UE::Editor::DataStorage::QueryStack
 		};
 	} // namespace Private
 
-	FRowSortNode::FRowSortNode(ICoreProvider& InStorage, TSharedPtr<IRowNode> InParent, 
+	FRowSortNode::FRowSortNode(ICoreProvider& InStorage, TSharedPtr<IRowNode> InParent,
 		FTimespan InMaxFrameDuration)
 		: Parent(MoveTemp(InParent))
 		, MaxFrameDuration(InMaxFrameDuration)
 		, Storage(InStorage)
 	{
-		FRowHandleArray* ParentRows = Parent->GetMutableRows();
-		checkf(ParentRows != nullptr,
-			TEXT("FRowHandleSortNode requires a parent node that provides write access to its row handle array."));
 	}
 
-	FRowSortNode::FRowSortNode(ICoreProvider& InStorage, TSharedPtr<IRowNode> InParent, 
+	FRowSortNode::FRowSortNode(ICoreProvider& InStorage, TSharedPtr<IRowNode> InParent,
 		TSharedPtr<const FColumnSorterInterface> InColumnSorter, FTimespan InMaxFrameDuration)
 		: ColumnSorter(MoveTemp(InColumnSorter))
 		, Parent(MoveTemp(InParent))
 		, MaxFrameDuration(InMaxFrameDuration)
 		, Storage(InStorage)
 	{
-		FRowHandleArray* ParentRows = Parent->GetMutableRows();
-		checkf(ParentRows != nullptr, 
-			TEXT("FRowHandleSortNode requires a parent node that provides write access to its row handle array."));
 	}
 
 	void FRowSortNode::SetColumnSorter(TSharedPtr<const FColumnSorterInterface> InColumnSorter)
@@ -128,97 +122,73 @@ namespace UE::Editor::DataStorage::QueryStack
 
 	void FRowSortNode::SetupFixedSize64()
 	{
-		if (FRowHandleArray* MutableRows = Parent->GetMutableRows())
-		{
-			FRowHandleArrayView Rows = Parent->GetRows();
-			TArrayView<RowHandle> RowView = MutableRows->GetMutableRows(MutableRows->IsUnique()
-				? FRowHandleArray::EFlags::IsUnique
-				: FRowHandleArray::EFlags::None);
-			SortContext->PrefixArray.AddUninitialized(Rows.Num());
-			SortContext->Tasks.CollectPrefixes(SortContext->PrefixArray, Rows, Storage, ColumnSorter, Private::FSortContext::PrefixBatchSize);
-			SortContext->Tasks.AddBarrier();
-
-			SortContext->Tasks.RadixSortPrefixes(SortContext->PrefixArray);
-			SortContext->Tasks.AddBarrier();
-
-			SortContext->Tasks.CopyPrefixesToRowHandleArray(RowView, SortContext->PrefixArray);
-		}
-		else
-		{
-			checkf(false, TEXT("A 64-bit fixed size sort was attempted on a parent node that can't provide mutable rows."));
-		}
+		FRowHandleArray& MutableRows = Parent->GetMutableRows();
+		FRowHandleArrayView Rows = Parent->GetRows();
+		TArrayView<RowHandle> RowView = MutableRows.GetMutableRows(MutableRows.IsUnique()
+			? FRowHandleArray::EFlags::IsUnique
+			: FRowHandleArray::EFlags::None);
+		SortContext->PrefixArray.AddUninitialized(Rows.Num());
+		SortContext->Tasks.CollectPrefixes(SortContext->PrefixArray, Rows, Storage, ColumnSorter, Private::FSortContext::PrefixBatchSize);
+		SortContext->Tasks.AddBarrier();
+
+		SortContext->Tasks.RadixSortPrefixes(SortContext->PrefixArray);
+		SortContext->Tasks.AddBarrier();
+
+		SortContext->Tasks.CopyPrefixesToRowHandleArray(RowView, SortContext->PrefixArray);
 	}
 
 	void FRowSortNode::SetupFixedSizeOnly()
 	{
-		if (FRowHandleArray* MutableRows = Parent->GetMutableRows())
-		{
-			FRowHandleArrayView Rows = Parent->GetRows();
-			SortContext->PrefixArray.AddUninitialized(Rows.Num());
-			SortContext->Tasks.CollectPrefixes(
-				SortContext->PrefixArray, Rows, Storage, ColumnSorter, Private::FSortContext::PrefixBatchSize);
-			SortContext->Tasks.AddBarrier();
-
-			SortContext->PrefixShadowArray.AddUninitialized(Rows.Num());
-			FRowHandleArray::EFlags Flags = MutableRows->IsUnique() ? FRowHandleArray::EFlags::IsUnique : FRowHandleArray::EFlags::None;
-			SortContext->Tasks.BucketizeRange(MutableRows->GetMutableRows(Flags), SortContext->PrefixArray,
-				SortContext->PrefixShadowArray, ColumnSorter, Storage, Private::FSortContext::PrefixBatchSize,
-				0 /* Never fall down to a comparative sort. The assumption is that this is called when all prefixes are the same limited size, though this is not mandatory. */);
-		}
-		else
-		{
-			checkf(false, TEXT("A fixed size sort was attempted on a parent node that can't provide mutable rows."));
-		}
+		FRowHandleArray& MutableRows = Parent->GetMutableRows();
+		FRowHandleArrayView Rows = Parent->GetRows();
+		SortContext->PrefixArray.AddUninitialized(Rows.Num());
+		SortContext->Tasks.CollectPrefixes(
+			SortContext->PrefixArray, Rows, Storage, ColumnSorter, Private::FSortContext::PrefixBatchSize);
+		SortContext->Tasks.AddBarrier();
+
+		SortContext->PrefixShadowArray.AddUninitialized(Rows.Num());
+		FRowHandleArray::EFlags Flags = MutableRows.IsUnique() ? FRowHandleArray::EFlags::IsUnique : FRowHandleArray::EFlags::None;
+		SortContext->Tasks.BucketizeRange(MutableRows.GetMutableRows(Flags), SortContext->PrefixArray,
+			SortContext->PrefixShadowArray, ColumnSorter, Storage, Private::FSortContext::PrefixBatchSize,
+			0 /* Never fall down to a comparative sort. The assumption is that this is called when all prefixes are the same limited size, though this is not mandatory. */);
 	}
 
 	void FRowSortNode::SetupComparativeSort()
 	{
-		if (FRowHandleArray* MutableRows = Parent->GetMutableRows())
-		{
-			FRowHandleArray::EFlags Flags = MutableRows->IsUnique() ? FRowHandleArray::EFlags::IsUnique : FRowHandleArray::EFlags::None;
-			SortContext->Tasks.ComparativeSortRowHandleRanges(MutableRows->GetMutableRows(Flags), Storage, ColumnSorter, Private::FSortContext::SortBatchSize);
-			SortContext->Tasks.AddBarrier();
-
-			SortContext->RowHandleArray.AddUninitialized(MutableRows->Num() / 2); // Used only for the right portion of the merge for temporary storage.
-			SortContext->Tasks.MergeSortedAdjacentRanges(
-				MutableRows->GetMutableRows(Flags), SortContext->RowHandleArray, Storage, ColumnSorter, Private::FSortContext::SortBatchSize);
-		}
-		else
-		{
-			checkf(false, TEXT("A 64-bit fixed size sort was attempted on a parent node that can't provide mutable rows."));
-		}
+		FRowHandleArray& MutableRows = Parent->GetMutableRows();
+		FRowHandleArray::EFlags Flags = MutableRows.IsUnique() ? FRowHandleArray::EFlags::IsUnique : FRowHandleArray::EFlags::None;
+		SortContext->Tasks.ComparativeSortRowHandleRanges(MutableRows.GetMutableRows(Flags), Storage, ColumnSorter, Private::FSortContext::SortBatchSize);
+		SortContext->Tasks.AddBarrier();
+
+		SortContext->RowHandleArray.AddUninitialized(MutableRows.Num() / 2); // Used only for the right portion of the merge for temporary storage.
+		SortContext->Tasks.MergeSortedAdjacentRanges(
+			MutableRows.GetMutableRows(Flags), SortContext->RowHandleArray, Storage, ColumnSorter, Private::FSortContext::SortBatchSize);
 	}
 
 	void FRowSortNode::SetupHybridSort()
 	{
-		if (FRowHandleArray* MutableRows = Parent->GetMutableRows())
-		{
-			FRowHandleArrayView Rows = Parent->GetRows();
+		FRowHandleArray& MutableRows = Parent->GetMutableRows();
+		FRowHandleArrayView Rows = Parent->GetRows();
 			
-			if (Rows.Num() >= Private::FSortContext::SortBatchSize)
-			{
-				SortContext->PrefixArray.AddUninitialized(Rows.Num());
-				SortContext->Tasks.CollectPrefixes(
-					SortContext->PrefixArray, Rows, Storage, ColumnSorter, Private::FSortContext::PrefixBatchSize);
-				SortContext->Tasks.AddBarrier();
-
-				SortContext->PrefixShadowArray.AddUninitialized(Rows.Num());
-				FRowHandleArray::EFlags Flags = MutableRows->IsUnique() ? FRowHandleArray::EFlags::IsUnique : FRowHandleArray::EFlags::None;
-				SortContext->Tasks.BucketizeRange(MutableRows->GetMutableRows(Flags), SortContext->PrefixArray,
-					SortContext->PrefixShadowArray, ColumnSorter, Storage, Private::FSortContext::PrefixBatchSize,
-					Private::FSortContext::SortBatchSize);
-			}
-			else
-			{
-				// If there are less than the sort batch size of entries then there's no need for a prefix pass as it'll add extra work to sort since each
-				// bucket is going to be smaller than the batch size and instantly go to comparative sort.
-				FRowHandleArray::EFlags Flags = MutableRows->IsUnique() ? FRowHandleArray::EFlags::IsUnique : FRowHandleArray::EFlags::None;
-				SortContext->Tasks.ComparativeSortRowHandleRanges(MutableRows->GetMutableRows(Flags), Storage, ColumnSorter, Private::FSortContext::SortBatchSize);
-			}
+		if (Rows.Num() >= Private::FSortContext::SortBatchSize)
+		{
+			SortContext->PrefixArray.AddUninitialized(Rows.Num());
+			SortContext->Tasks.CollectPrefixes(
+				SortContext->PrefixArray, Rows, Storage, ColumnSorter, Private::FSortContext::PrefixBatchSize);
+			SortContext->Tasks.AddBarrier();
+
+			SortContext->PrefixShadowArray.AddUninitialized(Rows.Num());
+			FRowHandleArray::EFlags Flags = MutableRows.IsUnique() ? FRowHandleArray::EFlags::IsUnique : FRowHandleArray::EFlags::None;
+			SortContext->Tasks.BucketizeRange(MutableRows.GetMutableRows(Flags), SortContext->PrefixArray,
+				SortContext->PrefixShadowArray, ColumnSorter, Storage, Private::FSortContext::PrefixBatchSize,
+				Private::FSortContext::SortBatchSize);
 		}
 		else
 		{
-			checkf(false, TEXT("A hybrid sort was attempted on a parent node that can't provide mutable rows."));
+			// If there are less than the sort batch size of entries then there's no need for a prefix pass as it'll add extra work to sort since each
+			// bucket is going to be smaller than the batch size and instantly go to comparative sort.
+			FRowHandleArray::EFlags Flags = MutableRows.IsUnique() ? FRowHandleArray::EFlags::IsUnique : FRowHandleArray::EFlags::None;
+			SortContext->Tasks.ComparativeSortRowHandleRanges(MutableRows.GetMutableRows(Flags), Storage, ColumnSorter, Private::FSortContext::SortBatchSize);
 		}
 	}
 
@@ -227,7 +197,7 @@ namespace UE::Editor::DataStorage::QueryStack
 		return Parent->GetRows();
 	}
 
-	FRowHandleArray* FRowSortNode::GetMutableRows()
+	FRowHandleArray& FRowSortNode::GetMutableRows()
 	{
 		return Parent->GetMutableRows();
 	}
diff --git a/Engine/Plugins/Experimental/EditorDataStorageFeatures/Source/TedsQueryStack/Private/TedsRowViewNode.cpp b/Engine/Plugins/Experimental/EditorDataStorageFeatures/Source/TedsQueryStack/Private/TedsRowViewNode.cpp
--- a/Engine/Plugins/Experimental/EditorDataStorageFeatures/Source/TedsQueryStack/Private/TedsRowViewNode.cpp
+++ b/Engine/Plugins/Experimental/EditorDataStorageFeatures/Source/TedsQueryStack/Private/TedsRowViewNode.cpp
@@ -4,8 +4,8 @@
 
 namespace UE::Editor::DataStorage::QueryStack
 {
-	FRowViewNode::FRowViewNode(FRowHandleArrayView Rows)
-		: Rows(Rows)
+	FRowViewNode::FRowViewNode(FRowHandleArray& Rows)
+		: Rows(&Rows)
 	{}
 
 	void FRowViewNode::MarkDirty()
@@ -13,9 +13,9 @@ namespace UE::Editor::DataStorage::QueryStack
 		Revision++;
 	}
 
-	void FRowViewNode::ResetView(FRowHandleArrayView InRows)
+	void FRowViewNode::ResetView(FRowHandleArray& InRows)
 	{
-		Rows = InRows;
+		Rows = &InRows;
 		MarkDirty();
 	}
 
@@ -30,11 +30,13 @@ namespace UE::Editor::DataStorage::QueryStack
 
 	FRowHandleArrayView FRowViewNode::GetRows() const
 	{
-		return Rows;
+		return Rows->GetRows();
 	}
 
-	FRowHandleArray* FRowViewNode::GetMutableRows()
+	FRowHandleArray& FRowViewNode::GetMutableRows()
 	{
-		return nullptr;
+		// No need to check for validity as there's no function that allows a null pointer. The reason it's a pointer
+		// is so it can be reassigned.
+		return *Rows;
 	}
 } // namespace UE::Editor::DataStorage::QueryStack
diff --git a/Engine/Plugins/Experimental/EditorDataStorageFeatures/Source/TedsQueryStack/Public/TedsQueryStackInterfaces.h b/Engine/Plugins/Experimental/EditorDataStorageFeatures/Source/TedsQueryStack/Public/TedsQueryStackInterfaces.h
--- a/Engine/Plugins/Experimental/EditorDataStorageFeatures/Source/TedsQueryStack/Public/TedsQueryStackInterfaces.h
+++ b/Engine/Plugins/Experimental/EditorDataStorageFeatures/Source/TedsQueryStack/Public/TedsQueryStackInterfaces.h
@@ -38,6 +38,6 @@ namespace UE::Editor::DataStorage::QueryStack
 		/** Retrieve access to the rows used by this node. */
 		virtual FRowHandleArrayView GetRows() const = 0;
 		/** Retrieve write access to the rows used by this node. Is allowed to return null if write access can't be granted. */
-		virtual FRowHandleArray* GetMutableRows() = 0;
+		virtual FRowHandleArray& GetMutableRows() = 0;
 	};
 } // UE::Editor::DataStorage::QueryStack
diff --git a/Engine/Plugins/Experimental/EditorDataStorageFeatures/Source/TedsQueryStack/Public/TedsRowArrayNode.h b/Engine/Plugins/Experimental/EditorDataStorageFeatures/Source/TedsQueryStack/Public/TedsRowArrayNode.h
--- a/Engine/Plugins/Experimental/EditorDataStorageFeatures/Source/TedsQueryStack/Public/TedsRowArrayNode.h
+++ b/Engine/Plugins/Experimental/EditorDataStorageFeatures/Source/TedsQueryStack/Public/TedsRowArrayNode.h
@@ -25,7 +25,7 @@ namespace UE::Editor::DataStorage::QueryStack
 		TEDSQUERYSTACK_API virtual RevisionId GetRevision() const override;
 		TEDSQUERYSTACK_API virtual void Update() override;
 		TEDSQUERYSTACK_API virtual FRowHandleArrayView GetRows() const override;
-		TEDSQUERYSTACK_API virtual FRowHandleArray* GetMutableRows() override;
+		TEDSQUERYSTACK_API virtual FRowHandleArray& GetMutableRows() override;
 
 	private:
 		FRowHandleArray Rows;
diff --git a/Engine/Plugins/Experimental/EditorDataStorageFeatures/Source/TedsQueryStack/Public/TedsRowCopyNode.h b/Engine/Plugins/Experimental/EditorDataStorageFeatures/Source/TedsQueryStack/Public/TedsRowCopyNode.h
--- a/Engine/Plugins/Experimental/EditorDataStorageFeatures/Source/TedsQueryStack/Public/TedsRowCopyNode.h
+++ b/Engine/Plugins/Experimental/EditorDataStorageFeatures/Source/TedsQueryStack/Public/TedsRowCopyNode.h
@@ -26,7 +26,7 @@ namespace UE::Editor::DataStorage::QueryStack
 		TEDSQUERYSTACK_API virtual RevisionId GetRevision() const override;
 		TEDSQUERYSTACK_API virtual void Update() override;
 		TEDSQUERYSTACK_API virtual FRowHandleArrayView GetRows() const override;
-		TEDSQUERYSTACK_API virtual FRowHandleArray* GetMutableRows() override;
+		TEDSQUERYSTACK_API virtual FRowHandleArray& GetMutableRows() override;
 
 	private:
 		FRowHandleArray Rows;
diff --git a/Engine/Plugins/Experimental/EditorDataStorageFeatures/Source/TedsQueryStack/Public/TedsRowFilterNode.h b/Engine/Plugins/Experimental/EditorDataStorageFeatures/Source/TedsQueryStack/Public/TedsRowFilterNode.h
--- a/Engine/Plugins/Experimental/EditorDataStorageFeatures/Source/TedsQueryStack/Public/TedsRowFilterNode.h
+++ b/Engine/Plugins/Experimental/EditorDataStorageFeatures/Source/TedsQueryStack/Public/TedsRowFilterNode.h
@@ -45,9 +45,9 @@ namespace UE::Editor::DataStorage::QueryStack
 			return Rows.GetRows();
 		}
 		
-		virtual FRowHandleArray* GetMutableRows() override
+		virtual FRowHandleArray& GetMutableRows() override
 		{
-			return nullptr;
+			return Rows;
 		}
 		
 		virtual RevisionId GetRevision() const override
diff --git a/Engine/Plugins/Experimental/EditorDataStorageFeatures/Source/TedsQueryStack/Public/TedsRowHandleSortNode.h b/Engine/Plugins/Experimental/EditorDataStorageFeatures/Source/TedsQueryStack/Public/TedsRowHandleSortNode.h
--- a/Engine/Plugins/Experimental/EditorDataStorageFeatures/Source/TedsQueryStack/Public/TedsRowHandleSortNode.h
+++ b/Engine/Plugins/Experimental/EditorDataStorageFeatures/Source/TedsQueryStack/Public/TedsRowHandleSortNode.h
@@ -8,7 +8,7 @@ namespace UE::Editor::DataStorage::QueryStack
 {
 	/**
 	 * Sorts the rows in the parent by row handle. Prefer this over using the FRowSortNode when sorting
-	 * by row handle as this is somewhat faster and allows the row handle array to get its sorted status.
+	 * by row handle as this is somewhat faster and allows the row handle array to keep its sorted status.
 	 * If the row handle array is already sorted this node doesn't do any work, making it even faster than
 	 * using FRowSortNode.
 	 */
@@ -21,7 +21,7 @@ namespace UE::Editor::DataStorage::QueryStack
 		TEDSQUERYSTACK_API virtual RevisionId GetRevision() const override;
 		TEDSQUERYSTACK_API virtual void Update() override;
 		TEDSQUERYSTACK_API virtual FRowHandleArrayView GetRows() const override;
-		TEDSQUERYSTACK_API virtual FRowHandleArray* GetMutableRows() override;
+		TEDSQUERYSTACK_API virtual FRowHandleArray& GetMutableRows() override;
 
 	private:
 		TSharedPtr<IRowNode> Parent;
diff --git a/Engine/Plugins/Experimental/EditorDataStorageFeatures/Source/TedsQueryStack/Public/TedsRowMergeNode.h b/Engine/Plugins/Experimental/EditorDataStorageFeatures/Source/TedsQueryStack/Public/TedsRowMergeNode.h
--- a/Engine/Plugins/Experimental/EditorDataStorageFeatures/Source/TedsQueryStack/Public/TedsRowMergeNode.h
+++ b/Engine/Plugins/Experimental/EditorDataStorageFeatures/Source/TedsQueryStack/Public/TedsRowMergeNode.h
@@ -26,7 +26,7 @@ namespace UE::Editor::DataStorage::QueryStack
 		TEDSQUERYSTACK_API virtual RevisionId GetRevision() const override;
 		TEDSQUERYSTACK_API virtual void Update() override;
 		TEDSQUERYSTACK_API virtual FRowHandleArrayView GetRows() const override;
-		TEDSQUERYSTACK_API virtual FRowHandleArray* GetMutableRows() override;
+		TEDSQUERYSTACK_API virtual FRowHandleArray& GetMutableRows() override;
 
 	private:
 		void Merge();
diff --git a/Engine/Plugins/Experimental/EditorDataStorageFeatures/Source/TedsQueryStack/Public/TedsRowMonitorNode.h b/Engine/Plugins/Experimental/EditorDataStorageFeatures/Source/TedsQueryStack/Public/TedsRowMonitorNode.h
--- a/Engine/Plugins/Experimental/EditorDataStorageFeatures/Source/TedsQueryStack/Public/TedsRowMonitorNode.h
+++ b/Engine/Plugins/Experimental/EditorDataStorageFeatures/Source/TedsQueryStack/Public/TedsRowMonitorNode.h
@@ -24,11 +24,11 @@ namespace UE::Editor::DataStorage::QueryStack
 		DECLARE_MULTICAST_DELEGATE_OneParam(FOnRowsChanged, FRowHandleArrayView /* InRows */);
 		
 		TEDSQUERYSTACK_API FRowMonitorNode(
-			ICoreProvider& Storage, TSharedPtr<IRowNode> ParentRow, TSharedPtr<IQueryNode> QueryNode);
+			ICoreProvider& Storage, TSharedPtr<IRowNode> ParentNode, TSharedPtr<IQueryNode> QueryNode);
 		TEDSQUERYSTACK_API FRowMonitorNode(
-			ICoreProvider& Storage, TSharedPtr<IRowNode> ParentRow, TArray<TObjectPtr<const UScriptStruct>> Columns);
+			ICoreProvider& Storage, TSharedPtr<IRowNode> ParentNode, TArray<TObjectPtr<const UScriptStruct>> Columns);
 		TEDSQUERYSTACK_API FRowMonitorNode(
-			ICoreProvider& Storage, TSharedPtr<IQueryNode> QueryNode, TSharedPtr<IRowNode> ParentRow, 
+			ICoreProvider& Storage, TSharedPtr<IQueryNode> QueryNode, TSharedPtr<IRowNode> ParentNode,
 			TArray<TObjectPtr<const UScriptStruct>> MonitoredColumns);
 		TEDSQUERYSTACK_API FRowMonitorNode(
 			ICoreProvider& Storage, TSharedPtr<IQueryNode> QueryNode);
@@ -42,7 +42,7 @@ namespace UE::Editor::DataStorage::QueryStack
 		TEDSQUERYSTACK_API virtual RevisionId GetRevision() const override;
 		TEDSQUERYSTACK_API virtual void Update() override;
 		TEDSQUERYSTACK_API virtual FRowHandleArrayView GetRows() const override;
-		TEDSQUERYSTACK_API virtual FRowHandleArray* GetMutableRows() override;
+		TEDSQUERYSTACK_API virtual FRowHandleArray& GetMutableRows() override;
 
 		// NOTE: These events are not meant to be a wholesale solution for tracking changes to this node, it only fires when monitored rows are modified
 		// through observers. They are also deferred to Update() and not fired immediately.
@@ -54,13 +54,9 @@ namespace UE::Editor::DataStorage::QueryStack
 		// Event that is called during Update() to allow systems to respond to monitored rows being removed
 		TEDSQUERYSTACK_API FOnRowsChanged& OnMonitoredRowsRemoved();
 	private:
-		const FRowHandleArray& ResolveRows() const;
-		FRowHandleArray& ResolveRows();
 		void UpdateColumnsFromQuery();
-		void UpdateRows();
 		void UpdateMonitoredColumns();
 
-		FRowHandleArray Rows;
 		FRowHandleArray AddedRows;
 		FRowHandleArray RemovedRows;
 
@@ -68,7 +64,7 @@ namespace UE::Editor::DataStorage::QueryStack
 		TArray<TObjectPtr<const UScriptStruct>> MonitoredColumns;
 		
 		TSharedPtr<IQueryNode> QueryNode;
-		TSharedPtr<IRowNode> ParentRow;
+		TSharedPtr<IRowNode> ParentNode;
 		ICoreProvider& Storage;
 		RevisionId QueryRevision = 0;
 		RevisionId ParentRevision = 0;
diff --git a/Engine/Plugins/Experimental/EditorDataStorageFeatures/Source/TedsQueryStack/Public/TedsRowOrderInversionNode.h b/Engine/Plugins/Experimental/EditorDataStorageFeatures/Source/TedsQueryStack/Public/TedsRowOrderInversionNode.h
--- a/Engine/Plugins/Experimental/EditorDataStorageFeatures/Source/TedsQueryStack/Public/TedsRowOrderInversionNode.h
+++ b/Engine/Plugins/Experimental/EditorDataStorageFeatures/Source/TedsQueryStack/Public/TedsRowOrderInversionNode.h
@@ -24,7 +24,7 @@ namespace UE::Editor::DataStorage::QueryStack
 		TEDSQUERYSTACK_API virtual RevisionId GetRevision() const override;
 		TEDSQUERYSTACK_API virtual void Update() override;
 		TEDSQUERYSTACK_API virtual FRowHandleArrayView GetRows() const override;
-		TEDSQUERYSTACK_API virtual FRowHandleArray* GetMutableRows() override;
+		TEDSQUERYSTACK_API virtual FRowHandleArray& GetMutableRows() override;
 
 	private:
 		void ApplyOrder();
diff --git a/Engine/Plugins/Experimental/EditorDataStorageFeatures/Source/TedsQueryStack/Public/TedsRowQueryResultsNode.h b/Engine/Plugins/Experimental/EditorDataStorageFeatures/Source/TedsQueryStack/Public/TedsRowQueryResultsNode.h
--- a/Engine/Plugins/Experimental/EditorDataStorageFeatures/Source/TedsQueryStack/Public/TedsRowQueryResultsNode.h
+++ b/Engine/Plugins/Experimental/EditorDataStorageFeatures/Source/TedsQueryStack/Public/TedsRowQueryResultsNode.h
@@ -48,7 +48,7 @@ namespace UE::Editor::DataStorage::QueryStack
 		TEDSQUERYSTACK_API virtual RevisionId GetRevision() const override;
 		TEDSQUERYSTACK_API virtual void Update() override;
 		TEDSQUERYSTACK_API virtual FRowHandleArrayView GetRows() const override;
-		TEDSQUERYSTACK_API virtual FRowHandleArray* GetMutableRows() override;
+		TEDSQUERYSTACK_API virtual FRowHandleArray& GetMutableRows() override;
 
 	private:
 		void RefreshInternal(FRowHandleArray& TargetRows);
diff --git a/Engine/Plugins/Experimental/EditorDataStorageFeatures/Source/TedsQueryStack/Public/TedsRowSortNode.h b/Engine/Plugins/Experimental/EditorDataStorageFeatures/Source/TedsQueryStack/Public/TedsRowSortNode.h
--- a/Engine/Plugins/Experimental/EditorDataStorageFeatures/Source/TedsQueryStack/Public/TedsRowSortNode.h
+++ b/Engine/Plugins/Experimental/EditorDataStorageFeatures/Source/TedsQueryStack/Public/TedsRowSortNode.h
@@ -31,9 +31,9 @@ namespace UE::Editor::DataStorage::QueryStack
 	class FRowSortNode : public IRowNode
 	{
 	public:
-		TEDSQUERYSTACK_API FRowSortNode(ICoreProvider& InStorage, TSharedPtr<IRowNode> InParent, 
+		TEDSQUERYSTACK_API FRowSortNode(ICoreProvider& InStorage, TSharedPtr<IRowNode> InParent,
 			FTimespan InMaxFrameDuration = FTimespan::FromMilliseconds(10));
-		TEDSQUERYSTACK_API FRowSortNode(ICoreProvider& InStorge, TSharedPtr<IRowNode> InParent, 
+		TEDSQUERYSTACK_API FRowSortNode(ICoreProvider& InStorge, TSharedPtr<IRowNode> InParent,
 			TSharedPtr<const FColumnSorterInterface> InColumnSorter, FTimespan InMaxFrameDuration = FTimespan::FromMilliseconds(10));
 		virtual ~FRowSortNode() override = default;
 
@@ -43,7 +43,7 @@ namespace UE::Editor::DataStorage::QueryStack
 		TEDSQUERYSTACK_API virtual RevisionId GetRevision() const override;
 		TEDSQUERYSTACK_API virtual void Update() override;
 		TEDSQUERYSTACK_API virtual FRowHandleArrayView GetRows() const override;
-		TEDSQUERYSTACK_API virtual FRowHandleArray* GetMutableRows() override;
+		TEDSQUERYSTACK_API virtual FRowHandleArray& GetMutableRows() override;
 
 	private:
 		void SetupFixedSize64();
diff --git a/Engine/Plugins/Experimental/EditorDataStorageFeatures/Source/TedsQueryStack/Public/TedsRowViewNode.h b/Engine/Plugins/Experimental/EditorDataStorageFeatures/Source/TedsQueryStack/Public/TedsRowViewNode.h
--- a/Engine/Plugins/Experimental/EditorDataStorageFeatures/Source/TedsQueryStack/Public/TedsRowViewNode.h
+++ b/Engine/Plugins/Experimental/EditorDataStorageFeatures/Source/TedsQueryStack/Public/TedsRowViewNode.h
@@ -8,26 +8,25 @@
 namespace UE::Editor::DataStorage::QueryStack
 {
 	/**
-	 * Stores a view to a list of rows. The container that view is pointing to needs to be kept alive for
+	 * Stores a reference to row handle array. The row handle array that this view is pointing to needs to be kept alive for
 	 * as long as this Query Stack node is alive.
 	 */
 	class FRowViewNode : public IRowNode
 	{
 	public:
-		FRowViewNode() = default;
-		TEDSQUERYSTACK_API explicit FRowViewNode(FRowHandleArrayView Rows);
+		TEDSQUERYSTACK_API explicit FRowViewNode(FRowHandleArray& Rows);
 		virtual ~FRowViewNode() override = default;
 
 		TEDSQUERYSTACK_API void MarkDirty();
-		TEDSQUERYSTACK_API void ResetView(FRowHandleArrayView InRows);
+		TEDSQUERYSTACK_API void ResetView(FRowHandleArray& InRows);
 
 		TEDSQUERYSTACK_API virtual RevisionId GetRevision() const override;
 		TEDSQUERYSTACK_API virtual void Update() override;
 		TEDSQUERYSTACK_API virtual FRowHandleArrayView GetRows() const override;
-		TEDSQUERYSTACK_API virtual FRowHandleArray* GetMutableRows() override;
+		TEDSQUERYSTACK_API virtual FRowHandleArray& GetMutableRows() override;
 
 	private:
-		FRowHandleArrayView Rows;
+		FRowHandleArray* Rows = nullptr;
 		RevisionId Revision = 0;
 	};
 } // namespace UE::Editor::DataStorage::QueryStack
diff --git a/Engine/Plugins/Experimental/EditorDataStorageFeatures/Source/TedsQueryStack/Tests/TedsRowMergeNodeTests.cpp b/Engine/Plugins/Experimental/EditorDataStorageFeatures/Source/TedsQueryStack/Tests/TedsRowMergeNodeTests.cpp
--- a/Engine/Plugins/Experimental/EditorDataStorageFeatures/Source/TedsQueryStack/Tests/TedsRowMergeNodeTests.cpp
+++ b/Engine/Plugins/Experimental/EditorDataStorageFeatures/Source/TedsQueryStack/Tests/TedsRowMergeNodeTests.cpp
@@ -22,11 +22,13 @@ TEST_CASE_NAMED(TEDS_QueryStack_RowMergeNode_Tests, "Editor::DataStorage::QueryS
 
 	SECTION("Append two nodes")
 	{
-		RowHandle ValueArray[] = {1, 2, 3};
-		FRowHandleArrayView Values = FRowHandleArrayView( ValueArray,
-			FRowHandleArrayView::EFlags::IsSorted | FRowHandleArrayView::EFlags::IsUnique);
-		TSharedPtr<IRowNode> View0 = MakeShared<FRowViewNode>(Values);
-		TSharedPtr<IRowNode> View1 = MakeShared<FRowViewNode>(Values);
+		FRowHandleArray Values0;
+		FRowHandleArray Values1;
+		Values0.Append({ 1, 2, 3 }, FRowHandleArray::EFlags::IsSorted | FRowHandleArray::EFlags::IsUnique);
+		Values1.Append({ 1, 2, 3 }, FRowHandleArray::EFlags::IsSorted | FRowHandleArray::EFlags::IsUnique);
+		
+		TSharedPtr<IRowNode> View0 = MakeShared<FRowViewNode>(Values0);
+		TSharedPtr<IRowNode> View1 = MakeShared<FRowViewNode>(Values1);
 		
 		FRowMergeNode Node({View0, View1}, FRowMergeNode::EMergeApproach::Append);
 		
@@ -43,11 +45,13 @@ TEST_CASE_NAMED(TEDS_QueryStack_RowMergeNode_Tests, "Editor::DataStorage::QueryS
 
 	SECTION("Merge two nodes sorted")
 	{
-		RowHandle ValueArray[] = { 1, 2, 3 }; 
-		FRowHandleArrayView Values = FRowHandleArrayView( ValueArray,
-			FRowHandleArrayView::EFlags::IsSorted | FRowHandleArrayView::EFlags::IsUnique);
-		TSharedPtr<IRowNode> View0 = MakeShared<FRowViewNode>(Values);
-		TSharedPtr<IRowNode> View1 = MakeShared<FRowViewNode>(Values);
+		FRowHandleArray Values0;
+		FRowHandleArray Values1;
+		Values0.Append({ 1, 2, 3 }, FRowHandleArray::EFlags::IsSorted | FRowHandleArray::EFlags::IsUnique);
+		Values1.Append({ 1, 2, 3 }, FRowHandleArray::EFlags::IsSorted | FRowHandleArray::EFlags::IsUnique);
+
+		TSharedPtr<IRowNode> View0 = MakeShared<FRowViewNode>(Values0);
+		TSharedPtr<IRowNode> View1 = MakeShared<FRowViewNode>(Values1);
 		
 		FRowMergeNode Node({ View0, View1 }, FRowMergeNode::EMergeApproach::Sorted);
 
@@ -64,11 +68,13 @@ TEST_CASE_NAMED(TEDS_QueryStack_RowMergeNode_Tests, "Editor::DataStorage::QueryS
 
 	SECTION("Uniquely merge two nodes")
 	{
-		RowHandle ValueArray[] = { 1, 2, 3 };
-		FRowHandleArrayView Values = FRowHandleArrayView( ValueArray,
-			FRowHandleArrayView::EFlags::IsSorted | FRowHandleArrayView::EFlags::IsUnique);
-		TSharedPtr<IRowNode> View0 = MakeShared<FRowViewNode>(Values);
-		TSharedPtr<IRowNode> View1 = MakeShared<FRowViewNode>(Values);
+		FRowHandleArray Values0;
+		FRowHandleArray Values1;
+		Values0.Append({ 1, 2, 3 }, FRowHandleArray::EFlags::IsSorted | FRowHandleArray::EFlags::IsUnique);
+		Values1.Append({ 1, 2, 3 }, FRowHandleArray::EFlags::IsSorted | FRowHandleArray::EFlags::IsUnique);
+
+		TSharedPtr<IRowNode> View0 = MakeShared<FRowViewNode>(Values0);
+		TSharedPtr<IRowNode> View1 = MakeShared<FRowViewNode>(Values1);
 
 		FRowMergeNode Node({ View0, View1 }, FRowMergeNode::EMergeApproach::Unique);
 
@@ -82,12 +88,11 @@ TEST_CASE_NAMED(TEDS_QueryStack_RowMergeNode_Tests, "Editor::DataStorage::QueryS
 
 	SECTION("Repeating of data in two nodes")
 	{
-		RowHandle ValueArray0[] = { 1, 2, 3 };
-		RowHandle ValueArray1[] = { 2, 3, 4 };
-		FRowHandleArrayView Values0 = FRowHandleArrayView( ValueArray0,
-			FRowHandleArrayView::EFlags::IsSorted | FRowHandleArrayView::EFlags::IsUnique);
-		FRowHandleArrayView Values1 = FRowHandleArrayView( ValueArray1,
-			FRowHandleArrayView::EFlags::IsSorted | FRowHandleArrayView::EFlags::IsUnique);
+		FRowHandleArray Values0;
+		FRowHandleArray Values1;
+		Values0.Append({ 1, 2, 3 }, FRowHandleArray::EFlags::IsSorted | FRowHandleArray::EFlags::IsUnique);
+		Values1.Append({ 2, 3, 4 }, FRowHandleArray::EFlags::IsSorted | FRowHandleArray::EFlags::IsUnique);
+
 		TSharedPtr<IRowNode> View0 = MakeShared<FRowViewNode>(Values0);
 		TSharedPtr<IRowNode> View1 = MakeShared<FRowViewNode>(Values1);
 
@@ -102,16 +107,13 @@ TEST_CASE_NAMED(TEDS_QueryStack_RowMergeNode_Tests, "Editor::DataStorage::QueryS
 
 	SECTION("Repeating of data in three nodes")
 	{
-		RowHandle ValueArray0[] = { 1, 2, 3 };
-		RowHandle ValueArray1[] = { 2, 3, 4 };
-		RowHandle ValueArray2[] = { 1, 3 };
-
-		FRowHandleArrayView Values0 = FRowHandleArrayView(ValueArray0,
-			FRowHandleArrayView::EFlags::IsSorted | FRowHandleArrayView::EFlags::IsUnique);
-		FRowHandleArrayView Values1 = FRowHandleArrayView(ValueArray1,
-			FRowHandleArrayView::EFlags::IsSorted | FRowHandleArrayView::EFlags::IsUnique);
-		FRowHandleArrayView Values2 = FRowHandleArrayView(ValueArray2,
-			FRowHandleArrayView::EFlags::IsSorted | FRowHandleArrayView::EFlags::IsUnique);
+		FRowHandleArray Values0;
+		FRowHandleArray Values1;
+		FRowHandleArray Values2;
+		Values0.Append({ 1, 2, 3 }, FRowHandleArray::EFlags::IsSorted | FRowHandleArray::EFlags::IsUnique);
+		Values1.Append({ 1, 2, 3 }, FRowHandleArray::EFlags::IsSorted | FRowHandleArray::EFlags::IsUnique);
+		Values2.Append({ 1, 3 }, FRowHandleArray::EFlags::IsSorted | FRowHandleArray::EFlags::IsUnique);
+
 		TSharedPtr<IRowNode> View0 = MakeShared<FRowViewNode>(Values0);
 		TSharedPtr<IRowNode> View1 = MakeShared<FRowViewNode>(Values1);
 		TSharedPtr<IRowNode> View2 = MakeShared<FRowViewNode>(Values2);
@@ -125,35 +127,6 @@ TEST_CASE_NAMED(TEDS_QueryStack_RowMergeNode_Tests, "Editor::DataStorage::QueryS
 		CHECK_EQUALS(TEXT("Rows[1]"), Rows[1], 2llu);
 		CHECK_EQUALS(TEXT("Rows[2]"), Rows[2], 3llu);
 	}
-
-	SECTION("Updating multiple nodes")
-	{
-		RowHandle Values[] = {1, 2, 3};
-
-		static constexpr int32 ViewCount = 42;
-		TArray<TSharedPtr<FRowViewNode>> Views;
-		Views.Reserve(ViewCount);
-		TArray<TSharedPtr<IRowNode>> Nodes;
-		Nodes.Reserve(ViewCount);
-		for (int32 Count = 0; Count < ViewCount; ++Count)
-		{
-			TSharedPtr<FRowViewNode> Node = MakeShared<FRowViewNode>(
-				FRowHandleArrayView(Values, FRowHandleArrayView::EFlags::IsSorted | FRowHandleArrayView::EFlags::IsUnique));
-			Views.Add(Node);
-			Nodes.Add(Node);
-		}
-
-		FRowMergeNode Node(Nodes, FRowMergeNode::EMergeApproach::Sorted);
-
-		Views[5]->MarkDirty();
-		Node.Update();
-
-		Views[14]->MarkDirty();
-		Node.Update();
-
-		Views[32]->MarkDirty();
-		Node.Update();
-	}
 }
 
 #endif // #if WITH_TESTS
diff --git a/Engine/Plugins/Experimental/EditorDataStorageFeatures/Source/TedsQueryStack/Tests/TedsRowViewNodeTests.cpp b/Engine/Plugins/Experimental/EditorDataStorageFeatures/Source/TedsQueryStack/Tests/TedsRowViewNodeTests.cpp
--- a/Engine/Plugins/Experimental/EditorDataStorageFeatures/Source/TedsQueryStack/Tests/TedsRowViewNodeTests.cpp
+++ b/Engine/Plugins/Experimental/EditorDataStorageFeatures/Source/TedsQueryStack/Tests/TedsRowViewNodeTests.cpp
@@ -10,17 +10,10 @@ TEST_CASE_NAMED(TEDS_QueryStack_RowViewNode_Tests, "Editor::DataStorage::QuerySt
 	using namespace UE::Editor::DataStorage;
 	using namespace UE::Editor::DataStorage::QueryStack;
 
-	SECTION("Empty view")
-	{
-		FRowViewNode View;
-		CHECK_EQUALS(TEXT("Revision"), View.GetRevision(), 0);
-		CHECK_MESSAGE(TEXT("Default for rows is not empty."), View.GetRows().IsEmpty());
-	}
-
 	SECTION("View with a few row handles")
 	{
-		RowHandle ValueArray[] = { 1, 2, 3 };
-		FRowHandleArrayView RowHandles(ValueArray, FRowHandleArrayView::EFlags::IsSorted | FRowHandleArrayView::EFlags::IsUnique);
+		FRowHandleArray RowHandles;
+		RowHandles.Append({ 1, 2, 3 }, FRowHandleArray::EFlags::IsSorted | FRowHandleArray::EFlags::IsUnique);
 		FRowViewNode View(RowHandles);
 		CHECK_EQUALS(TEXT("Revision"), View.GetRevision(), 0);
 		
@@ -34,7 +27,8 @@ TEST_CASE_NAMED(TEDS_QueryStack_RowViewNode_Tests, "Editor::DataStorage::QuerySt
 
 	SECTION("Mark dirty")
 	{
-		FRowViewNode View;
+		FRowHandleArray RowHandles;
+		FRowViewNode View(RowHandles);
 		CHECK_EQUALS(TEXT("Revision"), View.GetRevision(), 0);
 		
 		View.MarkDirty();
