Add default expansion state to STedsHierarchyViewer

diff --git a/Engine/Plugins/Experimental/EditorDataStorageFeatures/Source/TedsTableViewer/Private/Widgets/STedsHierarchyViewer.cpp b/Engine/Plugins/Experimental/EditorDataStorageFeatures/Source/TedsTableViewer/Private/Widgets/STedsHierarchyViewer.cpp
--- a/Engine/Plugins/Experimental/EditorDataStorageFeatures/Source/TedsTableViewer/Private/Widgets/STedsHierarchyViewer.cpp
+++ b/Engine/Plugins/Experimental/EditorDataStorageFeatures/Source/TedsTableViewer/Private/Widgets/STedsHierarchyViewer.cpp
@@ -22,6 +22,7 @@ namespace UE::Editor::DataStorage
 	{
 		OnSelectionChanged = InArgs._OnSelectionChanged;
 		EmptyRowsMessage = InArgs._EmptyRowsMessage;
+		DefaultExpansionState = InArgs._DefaultExpansionState;
 		
 		IUiProvider::FPurposeID CellWigetPurpose = InArgs._CellWidgetPurpose;
 
@@ -249,6 +250,8 @@ namespace UE::Editor::DataStorage
 
 	TSharedRef<ITableRow> SHierarchyViewer::MakeTableRowWidget(TableViewerItemPtr InItem, const TSharedRef<STableViewBase>& OwnerTable) const
 	{
+		TreeView->SetItemExpansion(InItem, DefaultExpansionState.Get() == EExpansionState::Expanded);
+
 		return SNew(SHierarchyViewerRow, OwnerTable, Model.ToSharedRef())
 				.Item(InItem);
 	}
diff --git a/Engine/Plugins/Experimental/EditorDataStorageFeatures/Source/TedsTableViewer/Public/Widgets/STedsHierarchyViewer.h b/Engine/Plugins/Experimental/EditorDataStorageFeatures/Source/TedsTableViewer/Public/Widgets/STedsHierarchyViewer.h
--- a/Engine/Plugins/Experimental/EditorDataStorageFeatures/Source/TedsTableViewer/Public/Widgets/STedsHierarchyViewer.h
+++ b/Engine/Plugins/Experimental/EditorDataStorageFeatures/Source/TedsTableViewer/Public/Widgets/STedsHierarchyViewer.h
@@ -38,6 +38,13 @@ namespace UE::Editor::DataStorage
 	{
 	public:
 		
+		enum class EExpansionState : uint8
+		{
+			Collapsed,
+			Expanded,
+			Invalid,
+		};
+
 		// Delegate fired when the selection in the table viewer changes
 		DECLARE_DELEGATE_OneParam(FOnSelectionChanged, RowHandle)
 
@@ -71,6 +78,9 @@ namespace UE::Editor::DataStorage
 		// Empty = simply show the column headers instead of a message
 		SLATE_ATTRIBUTE(FText, EmptyRowsMessage)
 
+		// Whether all rows start as collapsed or expanded by default
+		SLATE_ARGUMENT(EExpansionState, DefaultExpansionState)
+
 		SLATE_END_ARGS()
 
 	public:
@@ -144,6 +154,8 @@ namespace UE::Editor::DataStorage
 		// The message to show in place of the table viewer when there are no rows provided by the current query stack
 		TAttribute<FText> EmptyRowsMessage;
 
+		TAttribute<EExpansionState> DefaultExpansionState;
+
 		// A specialized query stack node to only contain the top level rows present in the actual query stack we are using.
 		TUniquePtr<QueryStack::FRowFilterNode<FTableRowParentColumn>> HierarchyNode;
 		
