Update OutlinerLabelWidget DisplayName

diff --git a/Engine/Plugins/Experimental/EditorDataStorageFeatures/Source/TedsOutliner/Private/Widgets/OutlinerLabelWidget.cpp b/Engine/Plugins/Experimental/EditorDataStorageFeatures/Source/TedsOutliner/Private/Widgets/OutlinerLabelWidget.cpp
--- a/Engine/Plugins/Experimental/EditorDataStorageFeatures/Source/TedsOutliner/Private/Widgets/OutlinerLabelWidget.cpp
+++ b/Engine/Plugins/Experimental/EditorDataStorageFeatures/Source/TedsOutliner/Private/Widgets/OutlinerLabelWidget.cpp
@@ -127,4 +127,10 @@ TSharedPtr<SWidget> FOutlinerLabelWidgetConstructor::CreateWidget(UE::Editor::Da
 	return SNullWidget::NullWidget;
 }
 
+FText FOutlinerLabelWidgetConstructor::CreateWidgetDisplayNameText(
+	UE::Editor::DataStorage::ICoreProvider* DataStorage, UE::Editor::DataStorage::RowHandle Row) const
+{
+	return LOCTEXT("OutlinerLabelDisplayText", "Item Label");
+}
+
 #undef LOCTEXT_NAMESPACE
diff --git a/Engine/Plugins/Experimental/EditorDataStorageFeatures/Source/TedsOutliner/Private/Widgets/OutlinerLabelWidget.h b/Engine/Plugins/Experimental/EditorDataStorageFeatures/Source/TedsOutliner/Private/Widgets/OutlinerLabelWidget.h
--- a/Engine/Plugins/Experimental/EditorDataStorageFeatures/Source/TedsOutliner/Private/Widgets/OutlinerLabelWidget.h
+++ b/Engine/Plugins/Experimental/EditorDataStorageFeatures/Source/TedsOutliner/Private/Widgets/OutlinerLabelWidget.h
@@ -32,6 +32,9 @@ public:
 	TEDSOUTLINER_API FOutlinerLabelWidgetConstructor();
 	~FOutlinerLabelWidgetConstructor() override = default;
 
+	TEDSOUTLINER_API virtual FText CreateWidgetDisplayNameText(
+		UE::Editor::DataStorage::ICoreProvider* DataStorage, UE::Editor::DataStorage::RowHandle Row = UE::Editor::DataStorage::InvalidRowHandle) const override;
+
 	TEDSOUTLINER_API virtual TSharedPtr<SWidget> CreateWidget(
 		UE::Editor::DataStorage::ICoreProvider* DataStorage,
 		UE::Editor::DataStorage::IUiProvider* DataStorageUi,
