[EditorDataStorageFeature] TEDS Asset Preview

diff --git a/Engine/Plugins/Experimental/EditorDataStorageFeatures/Source/TedsAssetData/Private/TedsAssetDataHelper.cpp b/Engine/Plugins/Experimental/EditorDataStorageFeatures/Source/TedsAssetData/Private/TedsAssetDataHelper.cpp
--- a/Engine/Plugins/Experimental/EditorDataStorageFeatures/Source/TedsAssetData/Private/TedsAssetDataHelper.cpp
+++ b/Engine/Plugins/Experimental/EditorDataStorageFeatures/Source/TedsAssetData/Private/TedsAssetDataHelper.cpp
@@ -11,12 +11,14 @@ namespace TedsAssetDataHelper
 		static const FLazyName ThumbnailHintText("bAllowHintText");
 		static const FLazyName ThumbnailRealTimeOnHovered("bAllowRealTimeOnHovered");
 		static const FLazyName ThumbnailSizeOffset("ThumbnailSizeOffset");
+		static const FLazyName ThumbnailCanDisplayEditModePrimitiveTools("CanDisplayEditModePrimitiveTools");
 
 		FName GetThumbnailStatusMetaDataName() { return ThumbnailStatus; }
 		FName GetThumbnailFadeInMetaDataName() { return ThumbnailFadeIn; }
 		FName GetThumbnailHintTextMetaDataName() { return ThumbnailHintText; }
 		FName GetThumbnailRealTimeOnHoveredMetaDataName() { return ThumbnailRealTimeOnHovered; }
 		FName GetThumbnailSizeOffsetMetaDataName() { return ThumbnailSizeOffset; }
+		FName GetThumbnailCanDisplayEditModePrimitiveTools() { return ThumbnailCanDisplayEditModePrimitiveTools; }
 	}
 
 	namespace TableView
diff --git a/Engine/Plugins/Experimental/EditorDataStorageFeatures/Source/TedsAssetData/Private/Widgets/AssetDataThumbnailWidget.cpp b/Engine/Plugins/Experimental/EditorDataStorageFeatures/Source/TedsAssetData/Private/Widgets/AssetDataThumbnailWidget.cpp
--- a/Engine/Plugins/Experimental/EditorDataStorageFeatures/Source/TedsAssetData/Private/Widgets/AssetDataThumbnailWidget.cpp
+++ b/Engine/Plugins/Experimental/EditorDataStorageFeatures/Source/TedsAssetData/Private/Widgets/AssetDataThumbnailWidget.cpp
@@ -26,6 +26,20 @@
 
 #define LOCTEXT_NAMESPACE "FAssetThumbnailWidgetConstructor"
 
+namespace ThumbnailHelper
+{
+	template <typename InColumnType>
+	UE::Editor::DataStorage::FAttributeBinder GetBinderForColumn(UE::Editor::DataStorage::ICoreProvider* DataStorage,
+		UE::Editor::DataStorage::RowHandle WidgetRow,
+		UE::Editor::DataStorage::RowHandle ParentWidgetRow)
+	{
+		
+		return DataStorage->IsRowAvailable(ParentWidgetRow) && !DataStorage->GetColumn<InColumnType>(WidgetRow) ?
+			UE::Editor::DataStorage::FAttributeBinder(ParentWidgetRow, DataStorage) :
+			UE::Editor::DataStorage::FAttributeBinder(WidgetRow, DataStorage);
+	}
+}
+
 namespace Private
 {
 	bool IsTopLevelFolder(const FStringView InFolderPath)
@@ -284,7 +298,6 @@ TSharedPtr<SWidget> FAssetThumbnailWidgetConstructor::CreateWidget(UE::Editor::D
 	{
 		ParentWidgetRowHandle = ParentWidgetRow->Parent;
 	}
-	UE::Editor::DataStorage::FAttributeBinder ParentWidgetBinder(ParentWidgetRowHandle, DataStorage);
 
 	float ThumbnailSizeOffset = 0.f;
 	UE::Editor::DataStorage::FMetaDataEntryView ThumbnailSizeOffsetMetaData = Arguments.FindGeneric(TedsAssetDataHelper::MetaDataNames::GetThumbnailSizeOffsetMetaDataName());
@@ -294,15 +307,19 @@ TSharedPtr<SWidget> FAssetThumbnailWidgetConstructor::CreateWidget(UE::Editor::D
 	}
 	constexpr int32 DefaultThumbnailSize = 64;
 
+
+	UE::Editor::DataStorage::FAttributeBinder PaddingBinder = ThumbnailHelper::GetBinderForColumn<FWidgetPaddingColumn_Experimental>(DataStorage, WidgetRow, ParentWidgetRowHandle);
+	UE::Editor::DataStorage::FAttributeBinder SizeBinder = ThumbnailHelper::GetBinderForColumn<FSizeValueColumn_Experimental>(DataStorage, WidgetRow, ParentWidgetRowHandle);
+
 	// Thumbnail Box container
 	const TSharedRef<SBox> ThumbnailBox = SNew(SBox)
-		.Padding(WidgetRowBinder.BindData(&FWidgetPaddingColumn_Experimental::Padding, FMargin(0.f)))
-		.WidthOverride(ParentWidgetBinder.BindData(&FSizeValueColumn_Experimental::SizeValue, [DefaultThumbnailSize, ThumbnailSizeOffset] (float InSizeValue)
+		.Padding(PaddingBinder.BindData(&FWidgetPaddingColumn_Experimental::Padding, FMargin(0.f)))
+		.WidthOverride(SizeBinder.BindData(&FSizeValueColumn_Experimental::SizeValue, [DefaultThumbnailSize, ThumbnailSizeOffset] (float InSizeValue)
 		{
 			const float Size = FMath::IsNearlyZero(InSizeValue) ? DefaultThumbnailSize : InSizeValue;
 			return FOptionalSize(Size + ThumbnailSizeOffset);
 		}))
-		.HeightOverride(ParentWidgetBinder.BindData(&FSizeValueColumn_Experimental::SizeValue, [DefaultThumbnailSize, ThumbnailSizeOffset] (float InSizeValue)
+		.HeightOverride(SizeBinder.BindData(&FSizeValueColumn_Experimental::SizeValue, [DefaultThumbnailSize, ThumbnailSizeOffset] (float InSizeValue)
 		{
 			const float Size = FMath::IsNearlyZero(InSizeValue) ? DefaultThumbnailSize : InSizeValue;
 			return FOptionalSize(Size + ThumbnailSizeOffset);
@@ -346,19 +363,17 @@ TSharedPtr<SWidget> FAssetThumbnailWidgetConstructor::CreateWidget(UE::Editor::D
 				return ReturnBrush;
 			});
 		}
+
 		// TODO: AssetItem->GetItem().GetItemTemporaryReason() == EContentBrowserItemFlags::Temporary_Creation;
 		// This was the previous check to assign bForceGenericThumbnail, for now it will use the Generic if the AssetData is not valid
 		ThumbnailConfig.bForceGenericThumbnail = !AssetDataToUse.IsValid();
 		ThumbnailConfig.AllowAssetSpecificThumbnailOverlay = !ThumbnailConfig.bForceGenericThumbnail;
 		ThumbnailConfig.ThumbnailLabel = EThumbnailLabel::ClassName;
 
-		if (DataStorage->HasColumns<FSizeValueColumn_Experimental>(WidgetRow))
+		ThumbnailConfig.GenericThumbnailSize = SizeBinder.BindData(&FSizeValueColumn_Experimental::SizeValue, [DefaultThumbnailSize] (float InThumbnailSize)
 		{
-			ThumbnailConfig.GenericThumbnailSize = ParentWidgetBinder.BindData(&FSizeValueColumn_Experimental::SizeValue, [DefaultThumbnailSize] (float InThumbnailSize)
-			{
-				return FMath::IsNearlyZero(InThumbnailSize) ? DefaultThumbnailSize : static_cast<int32>(InThumbnailSize);
-			});
-		}
+			return FMath::IsNearlyZero(InThumbnailSize) ? DefaultThumbnailSize : static_cast<int32>(InThumbnailSize);
+		});
 
 		// TODO: This data will need to be integrated in TEDS and later on remove this/update it to use TEDS instead to retrieve the data needed
 		// ThumbnailConfig.AssetSystemInfoProvider = nullptr;
@@ -369,8 +384,15 @@ TSharedPtr<SWidget> FAssetThumbnailWidgetConstructor::CreateWidget(UE::Editor::D
 		}
 		ThumbnailConfig.ShowAssetColor = true;
 		// TODO: If we are able to change the thumbnail in Teds view we would need to bind this to make it work
-		// ThumbnailConfig.bCanDisplayEditModePrimitiveTools = false;
-		ThumbnailConfig.IsEditModeVisible = ParentWidgetBinder.BindData(&FThumbnailEditModeColumn_Experimental::IsEditModeToggled, [] (bool IsEditModeToggled)
+
+		UE::Editor::DataStorage::FMetaDataEntryView bCanDisplayEditModePrimitiveTools = Arguments.FindGeneric(TedsAssetDataHelper::MetaDataNames::GetThumbnailCanDisplayEditModePrimitiveTools());
+		if (bCanDisplayEditModePrimitiveTools.IsSet() && bCanDisplayEditModePrimitiveTools.IsType<bool>())
+		{
+			ThumbnailConfig.bCanDisplayEditModePrimitiveTools = *bCanDisplayEditModePrimitiveTools.TryGetExact<bool>();
+		}
+
+		UE::Editor::DataStorage::FAttributeBinder ThumbnailEditModeBinder = ThumbnailHelper::GetBinderForColumn<FThumbnailEditModeColumn_Experimental>(DataStorage, WidgetRow, ParentWidgetRowHandle);
+		ThumbnailConfig.IsEditModeVisible = ThumbnailEditModeBinder.BindData(&FThumbnailEditModeColumn_Experimental::IsEditModeToggled, [] (bool IsEditModeToggled)
 		{
 			return IsEditModeToggled ? EVisibility::Visible : EVisibility::Collapsed;
 		});
@@ -457,15 +479,4 @@ TSharedPtr<SWidget> FAssetThumbnailWidgetConstructor::CreateWidget(UE::Editor::D
 	return ThumbnailBox;
 }
 
-TConstArrayView<const UScriptStruct*> FAssetThumbnailWidgetConstructor::GetAdditionalColumnsList() const
-{
-	using namespace UE::Editor::DataStorage;
-
-	static const TTypedElementColumnTypeList<
-		FSizeValueColumn_Experimental,
-		FThumbnailEditModeColumn_Experimental> Columns;
-
-	return Columns;
-}
-
 #undef LOCTEXT_NAMESPACE
diff --git a/Engine/Plugins/Experimental/EditorDataStorageFeatures/Source/TedsAssetData/Private/Widgets/AssetDataThumbnailWidget.h b/Engine/Plugins/Experimental/EditorDataStorageFeatures/Source/TedsAssetData/Private/Widgets/AssetDataThumbnailWidget.h
--- a/Engine/Plugins/Experimental/EditorDataStorageFeatures/Source/TedsAssetData/Private/Widgets/AssetDataThumbnailWidget.h
+++ b/Engine/Plugins/Experimental/EditorDataStorageFeatures/Source/TedsAssetData/Private/Widgets/AssetDataThumbnailWidget.h
@@ -41,7 +41,4 @@ public:
 		UE::Editor::DataStorage::RowHandle TargetRow,
 		UE::Editor::DataStorage::RowHandle WidgetRow, 
 		const UE::Editor::DataStorage::FMetaDataView& Arguments) override;
-
-protected:
-	virtual TConstArrayView<const UScriptStruct*> GetAdditionalColumnsList() const override;
 };
diff --git a/Engine/Plugins/Experimental/EditorDataStorageFeatures/Source/TedsAssetData/Private/Widgets/AssetPreview/AssetPreviewGeneralWidgetRegistrationFactory.cpp b/Engine/Plugins/Experimental/EditorDataStorageFeatures/Source/TedsAssetData/Private/Widgets/AssetPreview/AssetPreviewGeneralWidgetRegistrationFactory.cpp
new file mode 100644
--- /dev/null
+++ b/Engine/Plugins/Experimental/EditorDataStorageFeatures/Source/TedsAssetData/Private/Widgets/AssetPreview/AssetPreviewGeneralWidgetRegistrationFactory.cpp
@@ -0,0 +1,21 @@
+﻿// Copyright Epic Games, Inc. All Rights Reserved.
+
+#include "AssetPreviewGeneralWidgetRegistrationFactory.h"
+
+#include "Elements/Interfaces/TypedElementDataStorageUiInterface.h"
+
+#define LOCTEXT_NAMESPACE "AssetPreviewGeneralWidgetRegistrationFactory"
+
+void UAssetPreviewGeneralWidgetRegistrationFactory::RegisterWidgetPurposes(UE::Editor::DataStorage::IUiProvider& DataStorageUi) const
+{
+	using namespace UE::Editor::DataStorage;
+
+	const IUiProvider::FPurposeInfo AssetPreviewPurpose = IUiProvider::FPurposeInfo(
+		"AssetPreview", "Default", NAME_None,
+		IUiProvider::EPurposeType::UniqueByNameAndColumn,
+		LOCTEXT("GeneralAssetPreviewDefaultPurpose", "The default widget purpose for the AssetPreview."));
+
+	DataStorageUi.RegisterWidgetPurpose(AssetPreviewPurpose);
+}
+
+#undef LOCTEXT_NAMESPACE
diff --git a/Engine/Plugins/Experimental/EditorDataStorageFeatures/Source/TedsAssetData/Private/Widgets/AssetPreview/AssetPreviewGeneralWidgetRegistrationFactory.h b/Engine/Plugins/Experimental/EditorDataStorageFeatures/Source/TedsAssetData/Private/Widgets/AssetPreview/AssetPreviewGeneralWidgetRegistrationFactory.h
new file mode 100644
--- /dev/null
+++ b/Engine/Plugins/Experimental/EditorDataStorageFeatures/Source/TedsAssetData/Private/Widgets/AssetPreview/AssetPreviewGeneralWidgetRegistrationFactory.h
@@ -0,0 +1,19 @@
+﻿// Copyright Epic Games, Inc. All Rights Reserved.
+
+#pragma once
+
+#include "Elements/Interfaces/TypedElementDataStorageFactory.h"
+#include "UObject/ObjectMacros.h"
+
+#include "AssetPreviewGeneralWidgetRegistrationFactory.generated.h"
+
+UCLASS()
+class UAssetPreviewGeneralWidgetRegistrationFactory : public UEditorDataStorageFactory
+{
+	GENERATED_BODY()
+
+public:
+	~UAssetPreviewGeneralWidgetRegistrationFactory() override = default;
+
+	void RegisterWidgetPurposes(UE::Editor::DataStorage::IUiProvider& DataStorageUi) const override;
+};
diff --git a/Engine/Plugins/Experimental/EditorDataStorageFeatures/Source/TedsAssetData/Private/Widgets/AssetPreview/STedsAssetPreviewWidget.cpp b/Engine/Plugins/Experimental/EditorDataStorageFeatures/Source/TedsAssetData/Private/Widgets/AssetPreview/STedsAssetPreviewWidget.cpp
new file mode 100644
--- /dev/null
+++ b/Engine/Plugins/Experimental/EditorDataStorageFeatures/Source/TedsAssetData/Private/Widgets/AssetPreview/STedsAssetPreviewWidget.cpp
@@ -0,0 +1,333 @@
+// Copyright Epic Games, Inc. All Rights Reserved.
+
+#include "Widgets/AssetPreview/STedsAssetPreviewWidget.h"
+#include "Columns/SlateDelegateColumns.h"
+#include "Elements/Columns/TypedElementFolderColumns.h"
+#include "Elements/Columns/TypedElementMiscColumns.h"
+#include "Elements/Columns/TypedElementUIColumns.h"
+#include "Elements/Common/EditorDataStorageFeatures.h"
+#include "Elements/Framework/TypedElementDataStorageWidget.h"
+#include "Elements/Interfaces/TypedElementDataStorageInterface.h"
+#include "Elements/Interfaces/TypedElementDataStorageUiInterface.h"
+#include "TedsAssetDataColumns.h"
+#include "TedsAssetDataHelper.h"
+#include "TedsAssetDataWidgetColumns.h"
+#include "Widgets/AssetPreview/TedsAssetPreviewWidgetColumns.h"
+#include "Widgets/Input/SButton.h"
+#include "Widgets/Layout/SBorder.h"
+#include "Widgets/Layout/SBox.h"
+#include "Widgets/Images/SImage.h"
+#include "Widgets/SBoxPanel.h"
+
+void STedsAssetPreviewWidget::Construct(const FArguments& InArgs)
+{
+	WidgetPurpose = InArgs._WidgetPurpose;
+	WidgetHeaderPurpose = InArgs._HeaderWidgetPurpose;
+	ThumbnailWidgetPurpose = InArgs._ThumbnailWidgetPurpose;
+
+	// Create the internal Teds Widget
+	CreateTedsWidget();
+
+	// Default Border
+	PreviewPanel = SNew(SBorder);
+
+	TedsWidget->SetContent(PreviewPanel.ToSharedRef());
+
+	ChildSlot
+		[
+			TedsWidget->AsWidget()
+		];
+}
+
+UE::Editor::DataStorage::RowHandle STedsAssetPreviewWidget::GetWidgetRowHandle() const
+{
+	return TedsWidget->GetRowHandle();
+}
+
+void STedsAssetPreviewWidget::ReconstructTedsWidget()
+{
+	PreviewPanel->ClearContent();
+	ConstructWidget();
+}
+
+void STedsAssetPreviewWidget::SetTargetRow(UE::Editor::DataStorage::RowHandle InTargetRow)
+{
+	using namespace UE::Editor::DataStorage;
+
+	if (TargetRow != InTargetRow)
+	{
+		TargetRow = InTargetRow;
+		ReconstructTedsWidget();
+	}
+}
+
+void STedsAssetPreviewWidget::CreateTedsWidget()
+{
+	using namespace UE::Editor::DataStorage;
+
+	IUiProvider* DataStorageUI = GetMutableDataStorageFeature<IUiProvider>(UiFeatureName);
+	TedsWidget = DataStorageUI->CreateContainerTedsWidget(InvalidRowHandle);
+
+	const RowHandle WidgetRowHandle = TedsWidget->GetRowHandle();
+	ICoreProvider* DataStorage = GetMutableDataStorageFeature<ICoreProvider>(StorageFeatureName);
+
+	if (DataStorage->IsRowAvailable(WidgetRowHandle))
+	{
+		// External bound callback to get the ContextMenu of the Previewed Asset
+		DataStorage->AddColumn(WidgetRowHandle, FWidgetContextMenuColumn::StaticStruct());
+		DataStorage->AddColumn(WidgetRowHandle, FThumbnailEditModeColumn_Experimental::StaticStruct());
+	}
+}
+
+void STedsAssetPreviewWidget::ConstructWidget()
+{
+	using namespace UE::Editor::DataStorage;
+
+	ICoreProvider* DataStorage = GetMutableDataStorageFeature<ICoreProvider>(StorageFeatureName);
+
+	if (!DataStorage->IsRowAvailable(TargetRow))
+	{
+		return;
+	}
+
+	RowHandle WidgetRowPreview = TedsWidget->GetRowHandle();
+
+	IUiProvider* DataStorageUi = GetMutableDataStorageFeature<IUiProvider>(UiFeatureName);
+
+	TSharedRef<SVerticalBox> PreviewPanelVerticalBox = SNew(SVerticalBox);
+	TSharedRef<SHorizontalBox> ThumbnailHorizontalBox = SNew(SHorizontalBox);
+
+	FLazyName NamePurposeTableName = TEXT("Editor_WidgetTable");
+
+	// 1) Header
+	{
+		TSharedPtr<FTypedElementWidgetConstructor> OutHeaderButtonsWidgetConstructorPtr;
+		auto AssignHeaderButtonsWidgetToColumn = [&OutHeaderButtonsWidgetConstructorPtr] (TUniquePtr<FTypedElementWidgetConstructor> Constructor, TConstArrayView<TWeakObjectPtr<const UScriptStruct>>)
+		{
+			OutHeaderButtonsWidgetConstructorPtr = TSharedPtr<FTypedElementWidgetConstructor>(Constructor.Release());
+			return false;
+		};
+
+		TArray<TWeakObjectPtr<const UScriptStruct>> HeaderColumns = GetHeaderColumns();
+		DataStorageUi->CreateWidgetConstructors(DataStorageUi->FindPurpose(WidgetHeaderPurpose), IUiProvider::EMatchApproach::LongestMatch, HeaderColumns, FMetaDataView(), AssignHeaderButtonsWidgetToColumn);
+
+		TSharedPtr<SWidget> HeaderWidget = SNullWidget::NullWidget;
+		if (OutHeaderButtonsWidgetConstructorPtr)
+		{
+			if (HeaderWidgetRowHandle = DataStorage->AddRow(DataStorage->FindTable(NamePurposeTableName));
+				DataStorage->IsRowAvailable(HeaderWidgetRowHandle))
+			{
+				// Set the current TargetRow
+				DataStorage->AddColumn(HeaderWidgetRowHandle, FTypedElementRowReferenceColumn { .Row = TargetRow });
+
+				// External bound callback to get the ContextMenu of the Previewed Asset
+				if (FWidgetContextMenuColumn* ContextMenuWidgetRowColumn = DataStorage->GetColumn<FWidgetContextMenuColumn>(WidgetRowPreview))
+				{
+					DataStorage->AddColumn(HeaderWidgetRowHandle, FWidgetContextMenuColumn { .OnContextMenuOpening = ContextMenuWidgetRowColumn->OnContextMenuOpening });
+				}
+
+				// Set the current EditMode
+				DataStorage->AddColumn(HeaderWidgetRowHandle, FThumbnailEditModeColumn_Experimental { .IsEditModeToggled = IsEditModeEnabled });
+
+				// Set the external OnClicked event
+				DataStorage->AddColumn(HeaderWidgetRowHandle, FExternalWidgetOnClickedColumn_Experimental
+					{
+						.OnClicked = TDelegate<void()>::CreateLambda([this] ()
+						{
+							IsEditModeEnabled = !IsEditModeEnabled;
+							ICoreProvider* DataStorage = GetMutableDataStorageFeature<ICoreProvider>(StorageFeatureName);
+
+							// Update ThumbnailEditMode column
+							if (FThumbnailEditModeColumn_Experimental* ThumbnailEditModeColumn = DataStorage->GetColumn<FThumbnailEditModeColumn_Experimental>(ThumbnailWidgetRowHandle))
+							{
+								ThumbnailEditModeColumn->IsEditModeToggled = IsEditModeEnabled;
+							}
+						})
+					});
+
+				HeaderWidget = DataStorageUi->ConstructWidget(HeaderWidgetRowHandle, *OutHeaderButtonsWidgetConstructorPtr, FMetaDataView());
+			}
+		}
+
+		if (!HeaderWidget.IsValid())
+		{
+			HeaderWidget = SNullWidget::NullWidget;
+		}
+
+		PreviewPanelVerticalBox->AddSlot()
+			.AutoHeight()
+			[
+				HeaderWidget.ToSharedRef()
+			];
+	}
+
+	// 2) Thumbnail
+	{
+		TSharedPtr<FTypedElementWidgetConstructor> OutThumbnailWidgetConstructorPtr;
+		auto AssignThumbnailWidgetToColumn = [&OutThumbnailWidgetConstructorPtr] (TUniquePtr<FTypedElementWidgetConstructor> Constructor, TConstArrayView<TWeakObjectPtr<const UScriptStruct>>)
+		{
+			OutThumbnailWidgetConstructorPtr = TSharedPtr<FTypedElementWidgetConstructor>(Constructor.Release());
+			return false;
+		};
+
+		TArray<TWeakObjectPtr<const UScriptStruct>> ThumbnailColumns = GetThumbnailColumns();
+		DataStorageUi->CreateWidgetConstructors(DataStorageUi->FindPurpose(ThumbnailWidgetPurpose), IUiProvider::EMatchApproach::LongestMatch, ThumbnailColumns, FMetaDataView(), AssignThumbnailWidgetToColumn);
+
+		TSharedPtr<SWidget> ThumbnailWidget = SNullWidget::NullWidget;
+		if (OutThumbnailWidgetConstructorPtr)
+		{
+			if (ThumbnailWidgetRowHandle = DataStorage->AddRow(DataStorage->FindTable(NamePurposeTableName));
+				DataStorage->IsRowAvailable(ThumbnailWidgetRowHandle))
+			{
+				// Set the current TargetRow
+				DataStorage->AddColumn(ThumbnailWidgetRowHandle, FTypedElementRowReferenceColumn{ .Row = TargetRow });
+
+				// Set the current EditMode
+				DataStorage->AddColumn(ThumbnailWidgetRowHandle, FThumbnailEditModeColumn_Experimental{ .IsEditModeToggled = IsEditModeEnabled });
+
+				// Sizing is handled by the AssetPreview thumbnail widget
+
+				ThumbnailWidget = DataStorageUi->ConstructWidget(ThumbnailWidgetRowHandle, *OutThumbnailWidgetConstructorPtr, FMetaDataView());
+			}
+		}
+
+		if (!ThumbnailWidget.IsValid())
+		{
+			ThumbnailWidget = SNullWidget::NullWidget;
+		}
+
+		ThumbnailHorizontalBox->AddSlot()
+			.AutoWidth()
+			[
+				ThumbnailWidget.ToSharedRef()
+			];
+	}
+
+	// 3) Asset Name and Basic info
+	{
+		TSharedPtr<FTypedElementWidgetConstructor> OutBasicInfoWidgetConstructorPtr;
+		auto AssignBasicInfoWidgetToColumn = [&OutBasicInfoWidgetConstructorPtr] (TUniquePtr<FTypedElementWidgetConstructor> Constructor, TConstArrayView<TWeakObjectPtr<const UScriptStruct>>)
+		{
+			OutBasicInfoWidgetConstructorPtr = TSharedPtr<FTypedElementWidgetConstructor>(Constructor.Release());
+			return false;
+		};
+
+		TArray<TWeakObjectPtr<const UScriptStruct>> BasicInfoColumns = GetBasicInfoColumns();
+		DataStorageUi->CreateWidgetConstructors(DataStorageUi->FindPurpose(WidgetPurpose), IUiProvider::EMatchApproach::LongestMatch, BasicInfoColumns, FMetaDataView(), AssignBasicInfoWidgetToColumn);
+
+		TSharedPtr<SWidget> BasicInfoWidget = SNullWidget::NullWidget;
+		if (OutBasicInfoWidgetConstructorPtr)
+		{
+			if (BasicInfoWidgetRowHandle = DataStorage->AddRow(DataStorage->FindTable(NamePurposeTableName));
+				DataStorage->IsRowAvailable(BasicInfoWidgetRowHandle))
+			{
+				// Set the current TargetRow
+				DataStorage->AddColumn(BasicInfoWidgetRowHandle, FTypedElementRowReferenceColumn{ .Row = TargetRow });
+
+				BasicInfoWidget = DataStorageUi->ConstructWidget(BasicInfoWidgetRowHandle, *OutBasicInfoWidgetConstructorPtr, FMetaDataView());
+			}
+		}
+
+		if (!BasicInfoWidget.IsValid())
+		{
+			BasicInfoWidget = SNullWidget::NullWidget;
+		}
+
+		ThumbnailHorizontalBox->AddSlot()
+			[
+				BasicInfoWidget.ToSharedRef()
+			];
+
+		PreviewPanelVerticalBox->AddSlot()
+			.AutoHeight()
+			[
+				ThumbnailHorizontalBox
+			];
+	}
+
+	// 4) Full Details
+	{
+		TSharedPtr<FTypedElementWidgetConstructor> OutFullInfoWidgetConstructorPtr;
+		auto AssignFullInfoWidgetToColumn = [&OutFullInfoWidgetConstructorPtr] (TUniquePtr<FTypedElementWidgetConstructor> Constructor, TConstArrayView<TWeakObjectPtr<const UScriptStruct>>)
+		{
+			OutFullInfoWidgetConstructorPtr = TSharedPtr<FTypedElementWidgetConstructor>(Constructor.Release());
+			return false;
+		};
+
+		TArray<TWeakObjectPtr<const UScriptStruct>> FullInfoColumns = GetAdvancedInfoColumns();
+		DataStorageUi->CreateWidgetConstructors(DataStorageUi->FindPurpose(WidgetPurpose), IUiProvider::EMatchApproach::LongestMatch, FullInfoColumns, FMetaDataView(), AssignFullInfoWidgetToColumn);
+
+		TSharedPtr<SWidget> FullInfoWidget = SNullWidget::NullWidget;
+		if (OutFullInfoWidgetConstructorPtr)
+		{
+			if (AdvancedInfoWidgetRowHandle = DataStorage->AddRow(DataStorage->FindTable(NamePurposeTableName));
+				DataStorage->IsRowAvailable(AdvancedInfoWidgetRowHandle))
+			{
+				// Set the current TargetRow
+				DataStorage->AddColumn(AdvancedInfoWidgetRowHandle, FTypedElementRowReferenceColumn{ .Row = TargetRow });
+
+				FullInfoWidget = DataStorageUi->ConstructWidget(AdvancedInfoWidgetRowHandle, *OutFullInfoWidgetConstructorPtr, FMetaDataView());
+			}
+		}
+
+		if (!FullInfoWidget.IsValid())
+		{
+			FullInfoWidget = SNullWidget::NullWidget;
+		}
+
+		PreviewPanelVerticalBox->AddSlot()
+			.FillHeight(1.f)
+			[
+				FullInfoWidget.ToSharedRef()
+			];
+	}
+
+	// Assign to the Border the VerticalBox containing the TEDS Widgets
+	PreviewPanel->SetContent(PreviewPanelVerticalBox);
+}
+
+TArray<TWeakObjectPtr<const UScriptStruct>> STedsAssetPreviewWidget::GetHeaderColumns()
+{
+	static TArray<TWeakObjectPtr<const UScriptStruct>> HeaderColumns(
+	{
+		TWeakObjectPtr(FAssetTag::StaticStruct()),
+		TWeakObjectPtr(FFolderTag::StaticStruct())
+	});
+	return HeaderColumns;
+}
+
+TArray<TWeakObjectPtr<const UScriptStruct>> STedsAssetPreviewWidget::GetThumbnailColumns()
+{
+	static TArray<TWeakObjectPtr<const UScriptStruct>> ThumbnailColumns(
+		{
+			TWeakObjectPtr(FAssetTag::StaticStruct()),
+			TWeakObjectPtr(FFolderTag::StaticStruct())
+		});
+	return ThumbnailColumns;
+}
+
+TArray<TWeakObjectPtr<const UScriptStruct>> STedsAssetPreviewWidget::GetBasicInfoColumns()
+{
+	static TArray<TWeakObjectPtr<const UScriptStruct>> BasicInfoColumns(
+	{
+		TWeakObjectPtr(FAssetTag::StaticStruct()),
+		TWeakObjectPtr(FFolderTag::StaticStruct()),
+		TWeakObjectPtr(FNameColumn::StaticStruct())
+	});
+	return BasicInfoColumns;
+}
+
+TArray<TWeakObjectPtr<const UScriptStruct>> STedsAssetPreviewWidget::GetAdvancedInfoColumns()
+{
+	static TArray<TWeakObjectPtr<const UScriptStruct>> AdvancedInfoColumns(
+		{
+			TWeakObjectPtr(FAssetTag::StaticStruct()),
+			TWeakObjectPtr(FFolderTag::StaticStruct())
+		});
+	return AdvancedInfoColumns;
+}
+
+UE::Editor::DataStorage::RowHandle STedsAssetPreviewWidget::GetReferencedRowHandle() const
+{
+	return TargetRow;
+}
diff --git a/Engine/Plugins/Experimental/EditorDataStorageFeatures/Source/TedsAssetData/Private/Widgets/AssetPreview/Widgets/AdvancedInfo/AssetPreviewAdvancedInfoWidget.cpp b/Engine/Plugins/Experimental/EditorDataStorageFeatures/Source/TedsAssetData/Private/Widgets/AssetPreview/Widgets/AdvancedInfo/AssetPreviewAdvancedInfoWidget.cpp
new file mode 100644
--- /dev/null
+++ b/Engine/Plugins/Experimental/EditorDataStorageFeatures/Source/TedsAssetData/Private/Widgets/AssetPreview/Widgets/AdvancedInfo/AssetPreviewAdvancedInfoWidget.cpp
@@ -0,0 +1,49 @@
+﻿// Copyright Epic Games, Inc. All Rights Reserved.
+
+#include "AssetPreviewAdvancedInfoWidget.h"
+
+#include "Elements/Framework/TypedElementAttributeBinding.h"
+#include "Internationalization/Text.h"
+#include "TedsAssetDataColumns.h"
+#include "Elements/Columns/TypedElementFolderColumns.h"
+#include "UObject/ObjectMacros.h"
+#include "Widgets/Layout/SBox.h"
+#include "Widgets/Layout/SScrollBox.h"
+#include "Widgets/SBoxPanel.h"
+#include "Widgets/SRowDetails.h"
+#include "Widgets/Text/STextBlock.h"
+
+#define LOCTEXT_NAMESPACE "AssetPreviewAdvancedInfoWidget"
+
+void UAssetPreviewAdvancedInfoWidgetFactory::RegisterWidgetConstructors(UE::Editor::DataStorage::ICoreProvider& DataStorage,
+                                                                    UE::Editor::DataStorage::IUiProvider& DataStorageUi) const
+{
+	using namespace UE::Editor::DataStorage::Queries;
+	DataStorageUi.RegisterWidgetFactory<FAssetPreviewAdvancedInfoWidgetConstructor>(
+		DataStorageUi.FindPurpose(IUiProvider::FPurposeInfo("AssetPreview", "Default", NAME_None).GeneratePurposeID()),
+		TColumn<FAssetTag>() || TColumn<FFolderTag>());
+}
+
+FAssetPreviewAdvancedInfoWidgetConstructor::FAssetPreviewAdvancedInfoWidgetConstructor()
+	: FSimpleWidgetConstructor(StaticStruct())
+{
+}
+
+TSharedPtr<SWidget> FAssetPreviewAdvancedInfoWidgetConstructor::CreateWidget(UE::Editor::DataStorage::ICoreProvider* DataStorage,
+	UE::Editor::DataStorage::IUiProvider* DataStorageUi, UE::Editor::DataStorage::RowHandle TargetRow, UE::Editor::DataStorage::RowHandle WidgetRow,
+	const UE::Editor::DataStorage::FMetaDataView& Arguments)
+{
+	using namespace UE::Editor::DataStorage;
+
+	// For now Re-Use the RowDetails, but in the future we would probably need some customization or a specific new widget
+	TSharedRef<SRowDetails> RowDetails = SNew(SRowDetails).ShowAllDetails(true);
+	RowDetails->SetRow(TargetRow);
+
+	return SNew(SScrollBox)
+			+ SScrollBox::Slot()
+			[
+				RowDetails
+			];
+}
+
+#undef LOCTEXT_NAMESPACE
diff --git a/Engine/Plugins/Experimental/EditorDataStorageFeatures/Source/TedsAssetData/Private/Widgets/AssetPreview/Widgets/AdvancedInfo/AssetPreviewAdvancedInfoWidget.h b/Engine/Plugins/Experimental/EditorDataStorageFeatures/Source/TedsAssetData/Private/Widgets/AssetPreview/Widgets/AdvancedInfo/AssetPreviewAdvancedInfoWidget.h
new file mode 100644
--- /dev/null
+++ b/Engine/Plugins/Experimental/EditorDataStorageFeatures/Source/TedsAssetData/Private/Widgets/AssetPreview/Widgets/AdvancedInfo/AssetPreviewAdvancedInfoWidget.h
@@ -0,0 +1,42 @@
+﻿// Copyright Epic Games, Inc. All Rights Reserved.
+
+#pragma once
+
+#include "Elements/Interfaces/TypedElementDataStorageFactory.h"
+#include "Elements/Interfaces/TypedElementDataStorageUiInterface.h"
+
+#include "AssetPreviewAdvancedInfoWidget.generated.h"
+
+class IEditorDataStorageProvider;
+class UScriptStruct;
+class SWidget;
+
+UCLASS()
+class UAssetPreviewAdvancedInfoWidgetFactory : public UEditorDataStorageFactory
+{
+	GENERATED_BODY()
+
+public:
+	virtual ~UAssetPreviewAdvancedInfoWidgetFactory() override = default;
+
+	virtual void RegisterWidgetConstructors(UE::Editor::DataStorage::ICoreProvider& DataStorage,
+	                                        UE::Editor::DataStorage::IUiProvider& DataStorageUi) const override;
+};
+
+// Widget for the AssetPreview Advanced Info
+USTRUCT()
+struct FAssetPreviewAdvancedInfoWidgetConstructor : public FSimpleWidgetConstructor
+{
+	GENERATED_BODY()
+
+public:
+	FAssetPreviewAdvancedInfoWidgetConstructor();
+	virtual ~FAssetPreviewAdvancedInfoWidgetConstructor() override = default;
+
+	virtual TSharedPtr<SWidget> CreateWidget(
+		UE::Editor::DataStorage::ICoreProvider* DataStorage,
+		UE::Editor::DataStorage::IUiProvider* DataStorageUi,
+		UE::Editor::DataStorage::RowHandle TargetRow,
+		UE::Editor::DataStorage::RowHandle WidgetRow,
+		const UE::Editor::DataStorage::FMetaDataView& Arguments) override;
+};
diff --git a/Engine/Plugins/Experimental/EditorDataStorageFeatures/Source/TedsAssetData/Private/Widgets/AssetPreview/Widgets/BaseInfo/AssetPreviewBaseInfoWidget.cpp b/Engine/Plugins/Experimental/EditorDataStorageFeatures/Source/TedsAssetData/Private/Widgets/AssetPreview/Widgets/BaseInfo/AssetPreviewBaseInfoWidget.cpp
new file mode 100644
--- /dev/null
+++ b/Engine/Plugins/Experimental/EditorDataStorageFeatures/Source/TedsAssetData/Private/Widgets/AssetPreview/Widgets/BaseInfo/AssetPreviewBaseInfoWidget.cpp
@@ -0,0 +1,149 @@
+﻿// Copyright Epic Games, Inc. All Rights Reserved.
+
+#include "AssetPreviewBaseInfoWidget.h"
+
+#include "AssetDefinition.h"
+#include "AssetDefinitionRegistry.h"
+#include "ClassIconFinder.h"
+#include "Elements/Framework/TypedElementAttributeBinding.h"
+#include "Internationalization/Text.h"
+#include "TedsAssetDataColumns.h"
+#include "Elements/Columns/TypedElementFolderColumns.h"
+#include "Elements/Columns/TypedElementMiscColumns.h"
+#include "Elements/Columns/TypedElementSlateWidgetColumns.h"
+#include "UObject/ObjectMacros.h"
+#include "Widgets/Images/SImage.h"
+#include "Widgets/SBoxPanel.h"
+#include "Widgets/Text/STextBlock.h"
+
+#define LOCTEXT_NAMESPACE "AssetPreviewBaseInfoWidget"
+
+void UAssetPreviewBaseInfoWidgetFactory::RegisterWidgetConstructors(UE::Editor::DataStorage::ICoreProvider& DataStorage,
+                                                                    UE::Editor::DataStorage::IUiProvider& DataStorageUi) const
+{
+	using namespace UE::Editor::DataStorage::Queries;
+
+	DataStorageUi.RegisterWidgetFactory<FAssetPreviewBaseInfoWidgetConstructor>(
+		DataStorageUi.FindPurpose(IUiProvider::FPurposeInfo("AssetPreview", "Default", NAME_None).GeneratePurposeID()),
+		(TColumn<FAssetTag>() || TColumn<FFolderTag>()) && TColumn<FNameColumn>());
+}
+
+FAssetPreviewBaseInfoWidgetConstructor::FAssetPreviewBaseInfoWidgetConstructor()
+	: FSimpleWidgetConstructor(StaticStruct())
+{
+}
+
+TSharedPtr<SWidget> FAssetPreviewBaseInfoWidgetConstructor::CreateWidget(UE::Editor::DataStorage::ICoreProvider* DataStorage,
+	UE::Editor::DataStorage::IUiProvider* DataStorageUi, UE::Editor::DataStorage::RowHandle TargetRow, UE::Editor::DataStorage::RowHandle WidgetRow,
+	const UE::Editor::DataStorage::FMetaDataView& Arguments)
+{
+	using namespace UE::Editor::DataStorage;
+
+	TSharedPtr<SWidget> NameWidget = SNullWidget::NullWidget;
+
+	// Asset Name
+	TSharedRef<SVerticalBox> BaseInfoBox = SNew(SVerticalBox);
+	if (DataStorage->HasColumns<FNameColumn>(TargetRow))
+	{
+		const FName NamePurposeTableName = TEXT("Editor_WidgetTable");
+
+		TSharedPtr<FTypedElementWidgetConstructor> OutNameWidgetConstructorPtr;
+		auto AssignNameWidgetToColumn = [&OutNameWidgetConstructorPtr] (TUniquePtr<FTypedElementWidgetConstructor> Constructor, TConstArrayView<TWeakObjectPtr<const UScriptStruct>>)
+		{
+			OutNameWidgetConstructorPtr = TSharedPtr<FTypedElementWidgetConstructor>(Constructor.Release());
+			return false;
+		};
+
+		// MetaData
+		UE::Editor::DataStorage::FMetaData NameMeta;
+		NameMeta.AddOrSetMutableData(TEXT("OverflowPolicy"), 1);
+		NameMeta.AddOrSetMutableData(TEXT("Font_TypeFace_Name"), TEXT("Bold"));
+		NameMeta.AddOrSetMutableData(TEXT("Font_Size"), 15);
+
+		TArray<TWeakObjectPtr<const UScriptStruct>> BasicInfoColumns = GetLabelColumns();
+		DataStorageUi->CreateWidgetConstructors(DataStorageUi->FindPurpose(DataStorageUi->GetGeneralWidgetPurposeID()), IUiProvider::EMatchApproach::ExactMatch, BasicInfoColumns, FGenericMetaDataView(NameMeta), AssignNameWidgetToColumn);
+
+		RowHandle NameWidgetRowHandle = InvalidRowHandle;
+		if (OutNameWidgetConstructorPtr)
+		{
+			if (NameWidgetRowHandle = DataStorage->AddRow(DataStorage->FindTable(NamePurposeTableName));
+				DataStorage->IsRowAvailable(NameWidgetRowHandle))
+			{
+				// Get the current TargetRow based on the provider
+				DataStorage->AddColumn(NameWidgetRowHandle, FTypedElementRowReferenceColumn{ .Row = TargetRow });
+
+				// Add Columns before calling this if needed
+				NameWidget = DataStorageUi->ConstructWidget(NameWidgetRowHandle, *OutNameWidgetConstructorPtr, FGenericMetaDataView(NameMeta));
+			}
+		}
+
+		if (!NameWidget.IsValid())
+		{
+			NameWidget = SNullWidget::NullWidget;
+		}
+	}
+
+	FSlateIcon Icon = FSlateIcon();
+	FSlateColor IconColor = FSlateColor::UseForeground();
+	FText AssetClassName = FText::GetEmpty();
+	if (DataStorage->HasColumns<FAssetTag>(TargetRow))
+	{
+		FAssetData AssetData = DataStorage->GetColumn<FAssetDataColumn_Experimental>(TargetRow)->AssetData;
+		if (const UClass* AssetClass = AssetData.GetClass())
+		{
+			AssetClassName = AssetClass->GetDisplayNameText();
+			Icon = FSlateIconFinder::FindIconForClass(AssetClass);
+
+			if (DataStorage->HasColumns<FSlateColorColumn>(TargetRow))
+			{
+				IconColor = DataStorage->GetColumn<FSlateColorColumn>(TargetRow)->Color;
+			}
+			// If no Color is set use the Class specific one
+			else if (const UAssetDefinition* AssetDefinition = UAssetDefinitionRegistry::Get()->GetAssetDefinitionForClass(AssetClass))
+			{
+				IconColor = AssetDefinition->GetAssetColor();
+			}
+		}
+	}
+
+	TSharedPtr<SWidget> IconWidget = SNullWidget::NullWidget;
+	if (Icon.IsSet())
+	{
+		IconWidget = SNew(SImage).Image(Icon.GetIcon()).ColorAndOpacity(IconColor);
+	}
+
+	// Add Icon and AssetName
+	BaseInfoBox->AddSlot()
+		.AutoHeight()
+		[
+			SNew(SHorizontalBox)
+			+ SHorizontalBox::Slot()
+			.AutoWidth()
+			[
+				IconWidget.ToSharedRef()
+			]
+			+ SHorizontalBox::Slot()
+			.AutoWidth()
+			[
+				NameWidget.ToSharedRef()
+			]
+		];
+
+	// Class Name
+	BaseInfoBox->AddSlot()
+		.AutoHeight()
+		[
+			SNew(STextBlock)
+			.Text(AssetClassName)
+		];
+
+	return BaseInfoBox;
+}
+
+TArray<TWeakObjectPtr<const UScriptStruct>> FAssetPreviewBaseInfoWidgetConstructor::GetLabelColumns()
+{
+	static TArray<TWeakObjectPtr<const UScriptStruct>> LabelColumns({ TWeakObjectPtr(FNameColumn::StaticStruct()) });
+	return LabelColumns;
+}
+
+#undef LOCTEXT_NAMESPACE
diff --git a/Engine/Plugins/Experimental/EditorDataStorageFeatures/Source/TedsAssetData/Private/Widgets/AssetPreview/Widgets/BaseInfo/AssetPreviewBaseInfoWidget.h b/Engine/Plugins/Experimental/EditorDataStorageFeatures/Source/TedsAssetData/Private/Widgets/AssetPreview/Widgets/BaseInfo/AssetPreviewBaseInfoWidget.h
new file mode 100644
--- /dev/null
+++ b/Engine/Plugins/Experimental/EditorDataStorageFeatures/Source/TedsAssetData/Private/Widgets/AssetPreview/Widgets/BaseInfo/AssetPreviewBaseInfoWidget.h
@@ -0,0 +1,43 @@
+﻿// Copyright Epic Games, Inc. All Rights Reserved.
+
+#pragma once
+
+#include "Elements/Interfaces/TypedElementDataStorageFactory.h"
+#include "Elements/Interfaces/TypedElementDataStorageUiInterface.h"
+
+#include "AssetPreviewBaseInfoWidget.generated.h"
+
+class UScriptStruct;
+
+UCLASS()
+class UAssetPreviewBaseInfoWidgetFactory : public UEditorDataStorageFactory
+{
+	GENERATED_BODY()
+
+public:
+	virtual ~UAssetPreviewBaseInfoWidgetFactory() override = default;
+
+	virtual void RegisterWidgetConstructors(UE::Editor::DataStorage::ICoreProvider& DataStorage,
+	                                        UE::Editor::DataStorage::IUiProvider& DataStorageUi) const override;
+};
+
+// Widget for the AssetPreview Basic Info
+USTRUCT()
+struct FAssetPreviewBaseInfoWidgetConstructor : public FSimpleWidgetConstructor
+{
+	GENERATED_BODY()
+
+public:
+	FAssetPreviewBaseInfoWidgetConstructor();
+	virtual ~FAssetPreviewBaseInfoWidgetConstructor() override = default;
+
+	virtual TSharedPtr<SWidget> CreateWidget(
+		UE::Editor::DataStorage::ICoreProvider* DataStorage,
+		UE::Editor::DataStorage::IUiProvider* DataStorageUi,
+		UE::Editor::DataStorage::RowHandle TargetRow,
+		UE::Editor::DataStorage::RowHandle WidgetRow,
+		const UE::Editor::DataStorage::FMetaDataView& Arguments) override;
+
+private:
+	TArray<TWeakObjectPtr<const UScriptStruct>> GetLabelColumns();
+};
diff --git a/Engine/Plugins/Experimental/EditorDataStorageFeatures/Source/TedsAssetData/Private/Widgets/AssetPreview/Widgets/Header/AssetPreviewHeaderWidget.cpp b/Engine/Plugins/Experimental/EditorDataStorageFeatures/Source/TedsAssetData/Private/Widgets/AssetPreview/Widgets/Header/AssetPreviewHeaderWidget.cpp
new file mode 100644
--- /dev/null
+++ b/Engine/Plugins/Experimental/EditorDataStorageFeatures/Source/TedsAssetData/Private/Widgets/AssetPreview/Widgets/Header/AssetPreviewHeaderWidget.cpp
@@ -0,0 +1,200 @@
+// Copyright Epic Games, Inc. All Rights Reserved.
+
+#include "AssetPreviewHeaderWidget.h"
+
+#include "AssetThumbnail.h"
+#include "AssetDefinition.h"
+#include "Elements/Framework/TypedElementAttributeBinding.h"
+#include "TedsAssetDataColumns.h"
+#include "TedsAssetDataWidgetColumns.h"
+#include "Columns/SlateDelegateColumns.h"
+#include "Elements/Columns/TypedElementFolderColumns.h"
+#include "Elements/Columns/TypedElementMiscColumns.h"
+#include "Widgets/AssetPreview/TedsAssetPreviewWidgetColumns.h"
+#include "Widgets/Images/SImage.h"
+#include "Widgets/Layout/SBox.h"
+#include "Widgets/Layout/SSpacer.h"
+#include "Widgets/Text/STextBlock.h"
+#include "Widgets/SBoxPanel.h"
+
+#define LOCTEXT_NAMESPACE "FAssetPreviewHeaderWidgetConstructor"
+
+void UAssetPreviewHeaderWidgetFactory::RegisterWidgetPurposes(UE::Editor::DataStorage::IUiProvider& DataStorageUi) const
+{
+	using namespace UE::Editor::DataStorage::Queries;
+	
+	DataStorageUi.RegisterWidgetPurpose(
+		IUiProvider::FPurposeInfo("AssetPreview", "Header", "Default",
+		IUiProvider::EPurposeType::UniqueByNameAndColumn,
+		LOCTEXT("GeneralAssetPreviewDefaultHeaderPurpose", "The default widget purpose for the AssetPreview header.")));
+}
+
+void UAssetPreviewHeaderWidgetFactory::RegisterWidgetConstructors(UE::Editor::DataStorage::ICoreProvider& DataStorage,
+	UE::Editor::DataStorage::IUiProvider& DataStorageUi) const
+{
+	using namespace UE::Editor::DataStorage::Queries;
+
+	DataStorageUi.RegisterWidgetFactory<FAssetPreviewHeaderWidgetConstructor>(
+		DataStorageUi.FindPurpose(IUiProvider::FPurposeInfo("AssetPreview", "Header", "Default").GeneratePurposeID()),
+		TColumn<FAssetTag>() || TColumn<FFolderTag>());
+}
+
+FAssetPreviewHeaderWidgetConstructor::FAssetPreviewHeaderWidgetConstructor()
+	: FSimpleWidgetConstructor(StaticStruct())
+{
+}
+
+TSharedPtr<SWidget> FAssetPreviewHeaderWidgetConstructor::CreateWidget(UE::Editor::DataStorage::ICoreProvider* DataStorage,
+	UE::Editor::DataStorage::IUiProvider* DataStorageUi, UE::Editor::DataStorage::RowHandle TargetRow, UE::Editor::DataStorage::RowHandle WidgetRow,
+	const UE::Editor::DataStorage::FMetaDataView& Arguments)
+{
+	using namespace UE::Editor::DataStorage;
+
+	const FName NamePurposeTableName = TEXT("Editor_WidgetTable");
+
+	TSharedPtr<FTypedElementWidgetConstructor> OutEditModeButtonWidgetConstructorPtr;
+	TSharedPtr<FTypedElementWidgetConstructor> OutContextMenuWidgetConstructorPtr;
+
+	auto AssignEditModeButtonWidgetToColumn = [&OutEditModeButtonWidgetConstructorPtr] (TUniquePtr<FTypedElementWidgetConstructor> Constructor, TConstArrayView<TWeakObjectPtr<const UScriptStruct>>)
+	{
+		OutEditModeButtonWidgetConstructorPtr = TSharedPtr<FTypedElementWidgetConstructor>(Constructor.Release());
+		return false;
+	};
+
+	auto AssignContextMenuWidgetToColumn = [&OutContextMenuWidgetConstructorPtr] (TUniquePtr<FTypedElementWidgetConstructor> Constructor, TConstArrayView<TWeakObjectPtr<const UScriptStruct>>)
+	{
+		OutContextMenuWidgetConstructorPtr = TSharedPtr<FTypedElementWidgetConstructor>(Constructor.Release());
+		return false;
+	};
+
+	if (DataStorage->HasColumns<FAssetTag>(TargetRow))
+	{
+		TArray<TWeakObjectPtr<const UScriptStruct>> EditModeButtonColumns = GetEditModeColumns();
+		DataStorageUi->CreateWidgetConstructors(DataStorageUi->FindPurpose(
+			IUiProvider::FPurposeInfo("AssetPreview", "Header", "EditMode").GeneratePurposeID()),
+			IUiProvider::EMatchApproach::ExactMatch,
+			EditModeButtonColumns,
+			Arguments,
+			AssignEditModeButtonWidgetToColumn);
+	}
+
+	TArray<TWeakObjectPtr<const UScriptStruct>> ContextMenuButtonColumns = GetItemContextMenuColumns();
+	DataStorageUi->CreateWidgetConstructors(DataStorageUi->FindPurpose(
+		IUiProvider::FPurposeInfo("AssetPreview", "Header", "ContextMenu").GeneratePurposeID()),
+		IUiProvider::EMatchApproach::ExactMatch,
+		ContextMenuButtonColumns,
+		Arguments,
+		AssignContextMenuWidgetToColumn);
+
+	const FMargin ColumnItemPadding(8, 0);
+	constexpr float ThumbnailNameHorizontalPadding = 8.f;
+
+	TSharedPtr<SWidget> EditModeWidget = SNullWidget::NullWidget;
+	TSharedPtr<SWidget> ContextMenuWidget = SNullWidget::NullWidget;
+
+	RowHandle EditModeWidgetRowHandle = InvalidRowHandle;
+	if (OutEditModeButtonWidgetConstructorPtr)
+	{
+		if (EditModeWidgetRowHandle = DataStorage->AddRow(DataStorage->FindTable(NamePurposeTableName));
+			DataStorage->IsRowAvailable(EditModeWidgetRowHandle))
+		{
+			// Get the current TargetRow based on the provider
+			DataStorage->AddColumn(EditModeWidgetRowHandle, FTypedElementRowReferenceColumn{ .Row = TargetRow });
+
+			// Assign the same FThumbnailEditModeColumn_Experimental initial value that was assigned to this widget row
+			if (const FThumbnailEditModeColumn_Experimental* IsEditModeToggledColumn = DataStorage->GetColumn<FThumbnailEditModeColumn_Experimental>(WidgetRow))
+			{
+				DataStorage->AddColumn(EditModeWidgetRowHandle, FThumbnailEditModeColumn_Experimental{ .IsEditModeToggled = IsEditModeToggledColumn->IsEditModeToggled });
+			}
+
+			// Assign the same OnClicked that was assigned to this widget row
+			if (const FExternalWidgetOnClickedColumn_Experimental* OnWidgetClickedColumn = DataStorage->GetColumn<FExternalWidgetOnClickedColumn_Experimental>(WidgetRow))
+			{
+				DataStorage->AddColumn(EditModeWidgetRowHandle, FExternalWidgetOnClickedColumn_Experimental{ .OnClicked = OnWidgetClickedColumn->OnClicked });
+			}
+
+			// Add Columns before calling this if needed
+			EditModeWidget = DataStorageUi->ConstructWidget(EditModeWidgetRowHandle, *OutEditModeButtonWidgetConstructorPtr, Arguments);
+		}
+	}
+
+	RowHandle ContextMenuWidgetRowHandle = InvalidRowHandle;
+	if (OutContextMenuWidgetConstructorPtr)
+	{
+		if (ContextMenuWidgetRowHandle = DataStorage->AddRow(DataStorage->FindTable(NamePurposeTableName));
+			DataStorage->IsRowAvailable(ContextMenuWidgetRowHandle))
+		{
+			// Get the current TargetRow based on the provider
+			DataStorage->AddColumn(ContextMenuWidgetRowHandle, FTypedElementRowReferenceColumn{ .Row = TargetRow });
+
+			// Assign the same ContextMenuOpening that was assigned to this widget row
+			if (const FWidgetContextMenuColumn* WidgetContextMenuColumn = DataStorage->GetColumn<FWidgetContextMenuColumn>(WidgetRow))
+			{
+				DataStorage->AddColumn(ContextMenuWidgetRowHandle, FWidgetContextMenuColumn{ .OnContextMenuOpening = WidgetContextMenuColumn->OnContextMenuOpening });
+			}
+
+			// Add Columns before calling this if needed
+			ContextMenuWidget = DataStorageUi->ConstructWidget(ContextMenuWidgetRowHandle, *OutContextMenuWidgetConstructorPtr, Arguments);
+		}
+	}
+
+	if (!EditModeWidget.IsValid())
+	{
+		EditModeWidget = SNullWidget::NullWidget;
+	}
+
+	if (!ContextMenuWidget.IsValid())
+	{
+		ContextMenuWidget = SNullWidget::NullWidget;
+	}
+
+	return SNew(SBox)
+			.Padding(ColumnItemPadding)
+			[
+				SNew(SHorizontalBox)
+
+				+SHorizontalBox::Slot()
+				.FillWidth(1.f)
+				.VAlign(VAlign_Center)
+				[
+					EditModeWidget.ToSharedRef()
+				]
+
+				+SHorizontalBox::Slot()
+				.HAlign(HAlign_Right)
+				.AutoWidth()
+				.VAlign(VAlign_Center)
+				.Padding(ThumbnailNameHorizontalPadding, 0.f, 0.f, 0.f)
+				[
+					ContextMenuWidget.ToSharedRef()
+				]
+			];
+}
+
+TConstArrayView<const UScriptStruct*> FAssetPreviewHeaderWidgetConstructor::GetAdditionalColumnsList() const
+{
+	using namespace UE::Editor::DataStorage;
+	static TTypedElementColumnTypeList<FThumbnailEditModeColumn_Experimental, FExternalWidgetOnClickedColumn_Experimental, FWidgetContextMenuColumn> Columns;
+	return Columns;
+}
+
+TArray<TWeakObjectPtr<const UScriptStruct>> FAssetPreviewHeaderWidgetConstructor::GetEditModeColumns()
+{
+	static TArray<TWeakObjectPtr<const UScriptStruct>> EditModeColumns(
+		{
+			TWeakObjectPtr(FAssetTag::StaticStruct())
+		});
+	return EditModeColumns;
+}
+
+TArray<TWeakObjectPtr<const UScriptStruct>> FAssetPreviewHeaderWidgetConstructor::GetItemContextMenuColumns()
+{
+	static TArray<TWeakObjectPtr<const UScriptStruct>> MenuColumns(
+		{
+			TWeakObjectPtr(FAssetTag::StaticStruct()),
+			TWeakObjectPtr(FFolderTag::StaticStruct())
+		});
+	return MenuColumns;
+}
+
+#undef LOCTEXT_NAMESPACE
diff --git a/Engine/Plugins/Experimental/EditorDataStorageFeatures/Source/TedsAssetData/Private/Widgets/AssetPreview/Widgets/Header/AssetPreviewHeaderWidget.h b/Engine/Plugins/Experimental/EditorDataStorageFeatures/Source/TedsAssetData/Private/Widgets/AssetPreview/Widgets/Header/AssetPreviewHeaderWidget.h
new file mode 100644
--- /dev/null
+++ b/Engine/Plugins/Experimental/EditorDataStorageFeatures/Source/TedsAssetData/Private/Widgets/AssetPreview/Widgets/Header/AssetPreviewHeaderWidget.h
@@ -0,0 +1,48 @@
+// Copyright Epic Games, Inc. All Rights Reserved.
+
+#pragma once
+
+#include "Elements/Interfaces/TypedElementDataStorageFactory.h"
+#include "Elements/Interfaces/TypedElementDataStorageUiInterface.h"
+
+#include "AssetPreviewHeaderWidget.generated.h"
+
+class IEditorDataStorageProvider;
+class UScriptStruct;
+
+UCLASS()
+class UAssetPreviewHeaderWidgetFactory : public UEditorDataStorageFactory
+{
+	GENERATED_BODY()
+
+public:
+	virtual ~UAssetPreviewHeaderWidgetFactory() override = default;
+
+	virtual void RegisterWidgetPurposes(UE::Editor::DataStorage::IUiProvider& DataStorageUi) const override;
+	virtual void RegisterWidgetConstructors(UE::Editor::DataStorage::ICoreProvider& DataStorage,
+	                                        UE::Editor::DataStorage::IUiProvider& DataStorageUi) const override;
+};
+
+// Widget for the AssetPreview Header
+USTRUCT()
+struct FAssetPreviewHeaderWidgetConstructor : public FSimpleWidgetConstructor
+{
+	GENERATED_BODY()
+
+public:
+	FAssetPreviewHeaderWidgetConstructor();
+	virtual ~FAssetPreviewHeaderWidgetConstructor() override = default;
+
+	virtual TSharedPtr<SWidget> CreateWidget(
+		UE::Editor::DataStorage::ICoreProvider* DataStorage,
+		UE::Editor::DataStorage::IUiProvider* DataStorageUi,
+		UE::Editor::DataStorage::RowHandle TargetRow,
+		UE::Editor::DataStorage::RowHandle WidgetRow,
+		const UE::Editor::DataStorage::FMetaDataView& Arguments) override;
+
+	virtual TConstArrayView<const UScriptStruct*> GetAdditionalColumnsList() const override;
+
+private:
+	TArray<TWeakObjectPtr<const UScriptStruct>> GetEditModeColumns();
+	TArray<TWeakObjectPtr<const UScriptStruct>> GetItemContextMenuColumns();
+};
diff --git a/Engine/Plugins/Experimental/EditorDataStorageFeatures/Source/TedsAssetData/Private/Widgets/AssetPreview/Widgets/Header/EditModeToggleHeaderWidget.cpp b/Engine/Plugins/Experimental/EditorDataStorageFeatures/Source/TedsAssetData/Private/Widgets/AssetPreview/Widgets/Header/EditModeToggleHeaderWidget.cpp
new file mode 100644
--- /dev/null
+++ b/Engine/Plugins/Experimental/EditorDataStorageFeatures/Source/TedsAssetData/Private/Widgets/AssetPreview/Widgets/Header/EditModeToggleHeaderWidget.cpp
@@ -0,0 +1,82 @@
+// Copyright Epic Games, Inc. All Rights Reserved.
+
+#include "EditModeToggleHeaderWidget.h"
+
+#include "Elements/Framework/TypedElementAttributeBinding.h"
+#include "Internationalization/Text.h"
+#include "TedsAssetDataColumns.h"
+#include "TedsAssetDataWidgetColumns.h"
+#include "UObject/ObjectMacros.h"
+#include "Widgets/AssetPreview/TedsAssetPreviewWidgetColumns.h"
+#include "Widgets/Input/SButton.h"
+#include "Widgets/Layout/SBox.h"
+#include "Widgets/Text/STextBlock.h"
+
+#define LOCTEXT_NAMESPACE "EditModeToggleHeaderWidget"
+
+void UEditModeToggleHeaderWidgetFactory::RegisterWidgetPurposes(UE::Editor::DataStorage::IUiProvider& DataStorageUi) const
+{
+	using namespace UE::Editor::DataStorage::Queries;
+	
+	DataStorageUi.RegisterWidgetPurpose(
+		IUiProvider::FPurposeInfo("AssetPreview", "Header", "EditMode",
+		IUiProvider::EPurposeType::Generic,
+		LOCTEXT("AssetPreview_Header", "Widget that display the Header of the Preview Panel")));
+}
+
+void UEditModeToggleHeaderWidgetFactory::RegisterWidgetConstructors(UE::Editor::DataStorage::ICoreProvider& DataStorage,
+                                                                    UE::Editor::DataStorage::IUiProvider& DataStorageUi) const
+{
+	using namespace UE::Editor::DataStorage::Queries;
+	DataStorageUi.RegisterWidgetFactory<FEditModeToggleHeaderWidgetConstructor>(
+		DataStorageUi.FindPurpose(IUiProvider::FPurposeInfo("AssetPreview", "Header", "EditMode").GeneratePurposeID()),
+		TColumn<FAssetTag>());
+}
+
+FEditModeToggleHeaderWidgetConstructor::FEditModeToggleHeaderWidgetConstructor()
+	: FSimpleWidgetConstructor(StaticStruct())
+{
+}
+
+TSharedPtr<SWidget> FEditModeToggleHeaderWidgetConstructor::CreateWidget(UE::Editor::DataStorage::ICoreProvider* DataStorage,
+	UE::Editor::DataStorage::IUiProvider* DataStorageUi, UE::Editor::DataStorage::RowHandle TargetRow, UE::Editor::DataStorage::RowHandle WidgetRow,
+	const UE::Editor::DataStorage::FMetaDataView& Arguments)
+{
+	using namespace UE::Editor::DataStorage;
+
+	FAttributeBinder WidgetRowBinder(WidgetRow, DataStorage);
+
+	return SNew(SButton)
+		.Text(LOCTEXT("EditModeButton_Toggle", "Edit..."))
+		.ButtonColorAndOpacity(WidgetRowBinder.BindData(&FThumbnailEditModeColumn_Experimental::IsEditModeToggled, [] (bool IsEditMode)
+		{
+			if (IsEditMode)
+			{
+				return FSlateColor(FLinearColor::Blue);
+			}
+			return FSlateColor::UseForeground();
+		}))
+		.OnClicked_Lambda([WidgetRow, DataStorage] ()
+		{
+			FAttributeBinder WidgetRowBinder(WidgetRow, DataStorage);
+			const TDelegate<void()> Event = WidgetRowBinder.BindEvent(&FExternalWidgetOnClickedColumn_Experimental::OnClicked);
+			if (Event.IsBound())
+			{
+				Event.Execute();
+				if (FThumbnailEditModeColumn_Experimental* IsEditModeToggledColumn = DataStorage->GetColumn<FThumbnailEditModeColumn_Experimental>(WidgetRow))
+				{
+					IsEditModeToggledColumn->IsEditModeToggled = !IsEditModeToggledColumn->IsEditModeToggled;
+				}
+			}
+			return FReply::Handled();
+		});
+}
+
+TConstArrayView<const UScriptStruct*> FEditModeToggleHeaderWidgetConstructor::GetAdditionalColumnsList() const
+{
+	using namespace UE::Editor::DataStorage;
+	static TTypedElementColumnTypeList<FThumbnailEditModeColumn_Experimental, FExternalWidgetOnClickedColumn_Experimental> Columns;
+	return Columns;
+}
+
+#undef LOCTEXT_NAMESPACE
diff --git a/Engine/Plugins/Experimental/EditorDataStorageFeatures/Source/TedsAssetData/Private/Widgets/AssetPreview/Widgets/Header/EditModeToggleHeaderWidget.h b/Engine/Plugins/Experimental/EditorDataStorageFeatures/Source/TedsAssetData/Private/Widgets/AssetPreview/Widgets/Header/EditModeToggleHeaderWidget.h
new file mode 100644
--- /dev/null
+++ b/Engine/Plugins/Experimental/EditorDataStorageFeatures/Source/TedsAssetData/Private/Widgets/AssetPreview/Widgets/Header/EditModeToggleHeaderWidget.h
@@ -0,0 +1,44 @@
+// Copyright Epic Games, Inc. All Rights Reserved.
+
+#pragma once
+
+#include "Elements/Interfaces/TypedElementDataStorageFactory.h"
+#include "Elements/Interfaces/TypedElementDataStorageUiInterface.h"
+
+#include "EditModeToggleHeaderWidget.generated.h"
+
+class IEditorDataStorageProvider;
+class UScriptStruct;
+
+UCLASS()
+class UEditModeToggleHeaderWidgetFactory : public UEditorDataStorageFactory
+{
+	GENERATED_BODY()
+
+public:
+	virtual ~UEditModeToggleHeaderWidgetFactory() override = default;
+
+	virtual void RegisterWidgetPurposes(UE::Editor::DataStorage::IUiProvider& DataStorageUi) const override;
+	virtual void RegisterWidgetConstructors(UE::Editor::DataStorage::ICoreProvider& DataStorage,
+		UE::Editor::DataStorage::IUiProvider& DataStorageUi) const override;
+};
+
+// Widget that display the Edit button for the AssetPreview panel
+USTRUCT()
+struct FEditModeToggleHeaderWidgetConstructor : public FSimpleWidgetConstructor
+{
+	GENERATED_BODY()
+
+public:
+	FEditModeToggleHeaderWidgetConstructor();
+	virtual ~FEditModeToggleHeaderWidgetConstructor() override = default;
+
+	virtual TSharedPtr<SWidget> CreateWidget(
+		UE::Editor::DataStorage::ICoreProvider* DataStorage,
+		UE::Editor::DataStorage::IUiProvider* DataStorageUi,
+		UE::Editor::DataStorage::RowHandle TargetRow,
+		UE::Editor::DataStorage::RowHandle WidgetRow,
+		const UE::Editor::DataStorage::FMetaDataView& Arguments) override;
+
+	virtual TConstArrayView<const UScriptStruct*> GetAdditionalColumnsList() const override;
+};
diff --git a/Engine/Plugins/Experimental/EditorDataStorageFeatures/Source/TedsAssetData/Private/Widgets/AssetPreview/Widgets/Header/ItemContextMenuWidget.cpp b/Engine/Plugins/Experimental/EditorDataStorageFeatures/Source/TedsAssetData/Private/Widgets/AssetPreview/Widgets/Header/ItemContextMenuWidget.cpp
new file mode 100644
--- /dev/null
+++ b/Engine/Plugins/Experimental/EditorDataStorageFeatures/Source/TedsAssetData/Private/Widgets/AssetPreview/Widgets/Header/ItemContextMenuWidget.cpp
@@ -0,0 +1,80 @@
+// Copyright Epic Games, Inc. All Rights Reserved.
+
+#include "ItemContextMenuWidget.h"
+
+#include "Columns/SlateDelegateColumns.h"
+#include "Elements/Columns/TypedElementFolderColumns.h"
+#include "Elements/Framework/TypedElementAttributeBinding.h"
+#include "Framework/Application/SlateApplication.h"
+#include "Internationalization/Text.h"
+#include "TedsAssetDataColumns.h"
+#include "UObject/ObjectMacros.h"
+#include "Widgets/Images/SImage.h"
+#include "Widgets/Input/SButton.h"
+#include "Widgets/Layout/SBox.h"
+#include "Widgets/SBoxPanel.h"
+#include "Widgets/Text/STextBlock.h"
+
+#define LOCTEXT_NAMESPACE "ItemContextMenuWidget"
+
+void UItemContextMenuWidgetFactory::RegisterWidgetPurposes(UE::Editor::DataStorage::IUiProvider& DataStorageUi) const
+{
+	using namespace UE::Editor::DataStorage::Queries;
+
+	DataStorageUi.RegisterWidgetPurpose(
+		IUiProvider::FPurposeInfo("AssetPreview", "Header", "ContextMenu",
+		IUiProvider::EPurposeType::UniqueByName,
+		LOCTEXT("AssetPreview_Header", "Widget that display the ContextMenu of the selected Asset/Folder")));
+}
+
+void UItemContextMenuWidgetFactory::RegisterWidgetConstructors(UE::Editor::DataStorage::ICoreProvider& DataStorage,
+                                                               UE::Editor::DataStorage::IUiProvider& DataStorageUi) const
+{
+	using namespace UE::Editor::DataStorage::Queries;
+
+	DataStorageUi.RegisterWidgetFactory<FItemContextMenuWidgetConstructor>(
+		DataStorageUi.FindPurpose(IUiProvider::FPurposeInfo("AssetPreview", "Header", "ContextMenu").GeneratePurposeID()),
+		TColumn<FAssetTag>() || TColumn<FFolderTag>());
+}
+
+FItemContextMenuWidgetConstructor::FItemContextMenuWidgetConstructor()
+	: FSimpleWidgetConstructor(StaticStruct())
+{
+}
+
+TSharedPtr<SWidget> FItemContextMenuWidgetConstructor::CreateWidget(UE::Editor::DataStorage::ICoreProvider* DataStorage,
+	UE::Editor::DataStorage::IUiProvider* DataStorageUi, UE::Editor::DataStorage::RowHandle TargetRow, UE::Editor::DataStorage::RowHandle WidgetRow,
+	const UE::Editor::DataStorage::FMetaDataView& Arguments)
+{
+	TSharedRef<SButton> Button = SNew(SButton)
+		[
+			SNew(SImage)
+			.Image(FAppStyle::GetBrush("Icons.EllipsisVerticalNarrow"))
+		];
+
+	Button->SetOnClicked(FOnClicked::CreateLambda([Button, DataStorage, WidgetRow] ()
+	{
+		if (FWidgetContextMenuColumn* ContextMenuColumn = DataStorage->GetColumn<FWidgetContextMenuColumn>(WidgetRow))
+		{
+			if (ContextMenuColumn->OnContextMenuOpening.IsBound())
+			{
+				const FVector2f SummonLocation = FSlateApplication::Get().GetCursorPos();
+				FWidgetPath ButtonWidgetPath = FWidgetPath();
+				FSlateApplication::Get().FindPathToWidget(Button, ButtonWidgetPath);
+				FSlateApplication::Get().PushMenu(Button->AsShared(), ButtonWidgetPath, ContextMenuColumn->OnContextMenuOpening.Execute().ToSharedRef(), SummonLocation, FPopupTransitionEffect(FPopupTransitionEffect::ContextMenu));
+			}
+		}
+		return FReply::Handled();
+	}));
+
+	return Button;
+}
+
+TConstArrayView<const UScriptStruct*> FItemContextMenuWidgetConstructor::GetAdditionalColumnsList() const
+{
+	using namespace UE::Editor::DataStorage;
+	static TTypedElementColumnTypeList<FWidgetContextMenuColumn> Columns;
+	return Columns;
+}
+
+#undef LOCTEXT_NAMESPACE
diff --git a/Engine/Plugins/Experimental/EditorDataStorageFeatures/Source/TedsAssetData/Private/Widgets/AssetPreview/Widgets/Header/ItemContextMenuWidget.h b/Engine/Plugins/Experimental/EditorDataStorageFeatures/Source/TedsAssetData/Private/Widgets/AssetPreview/Widgets/Header/ItemContextMenuWidget.h
new file mode 100644
--- /dev/null
+++ b/Engine/Plugins/Experimental/EditorDataStorageFeatures/Source/TedsAssetData/Private/Widgets/AssetPreview/Widgets/Header/ItemContextMenuWidget.h
@@ -0,0 +1,44 @@
+// Copyright Epic Games, Inc. All Rights Reserved.
+
+#pragma once
+
+#include "Elements/Interfaces/TypedElementDataStorageFactory.h"
+#include "Elements/Interfaces/TypedElementDataStorageUiInterface.h"
+
+#include "ItemContextMenuWidget.generated.h"
+
+class IEditorDataStorageProvider;
+class UScriptStruct;
+
+UCLASS()
+class UItemContextMenuWidgetFactory : public UEditorDataStorageFactory
+{
+	GENERATED_BODY()
+
+public:
+	virtual ~UItemContextMenuWidgetFactory() override = default;
+
+	virtual void RegisterWidgetPurposes(UE::Editor::DataStorage::IUiProvider& DataStorageUi) const override;
+	virtual void RegisterWidgetConstructors(UE::Editor::DataStorage::ICoreProvider& DataStorage,
+	                                        UE::Editor::DataStorage::IUiProvider& DataStorageUi) const override;
+};
+
+// Widget that display the ContextMenu button for the AssetPreview panel
+USTRUCT()
+struct FItemContextMenuWidgetConstructor : public FSimpleWidgetConstructor
+{
+	GENERATED_BODY()
+
+public:
+	FItemContextMenuWidgetConstructor();
+	virtual ~FItemContextMenuWidgetConstructor() override = default;
+
+	virtual TSharedPtr<SWidget> CreateWidget(
+		UE::Editor::DataStorage::ICoreProvider* DataStorage,
+		UE::Editor::DataStorage::IUiProvider* DataStorageUi,
+		UE::Editor::DataStorage::RowHandle TargetRow,
+		UE::Editor::DataStorage::RowHandle WidgetRow,
+		const UE::Editor::DataStorage::FMetaDataView& Arguments) override;
+
+	virtual TConstArrayView<const UScriptStruct*> GetAdditionalColumnsList() const override;
+};
diff --git a/Engine/Plugins/Experimental/EditorDataStorageFeatures/Source/TedsAssetData/Private/Widgets/AssetPreview/Widgets/Thumbnail/AssetPreviewThumbnailWidget.cpp b/Engine/Plugins/Experimental/EditorDataStorageFeatures/Source/TedsAssetData/Private/Widgets/AssetPreview/Widgets/Thumbnail/AssetPreviewThumbnailWidget.cpp
new file mode 100644
--- /dev/null
+++ b/Engine/Plugins/Experimental/EditorDataStorageFeatures/Source/TedsAssetData/Private/Widgets/AssetPreview/Widgets/Thumbnail/AssetPreviewThumbnailWidget.cpp
@@ -0,0 +1,203 @@
+﻿// Copyright Epic Games, Inc. All Rights Reserved.
+
+#include "AssetPreviewThumbnailWidget.h"
+
+#include "AssetThumbnail.h"
+#include "Elements/Framework/TypedElementAttributeBinding.h"
+#include "TedsAssetDataColumns.h"
+#include "TedsAssetDataHelper.h"
+#include "TedsAssetDataWidgetColumns.h"
+#include "Elements/Columns/TypedElementFolderColumns.h"
+#include "Elements/Columns/TypedElementHiearchyColumns.h"
+#include "Elements/Columns/TypedElementMiscColumns.h"
+#include "Styling/StyleColors.h"
+#include "Widgets/Images/SImage.h"
+#include "Widgets/Input/SButton.h"
+#include "Widgets/Layout/SBox.h"
+#include "Widgets/SOverlay.h"
+
+
+#define LOCTEXT_NAMESPACE "AssetPreviewThumbnailWidget"
+
+// used to make the AssetPreview expand state global
+static bool bIsThumbnailExpandedGlobal = false;
+
+void UAssetPreviewThumbnailWidgetFactory::RegisterWidgetPurposes(UE::Editor::DataStorage::IUiProvider& DataStorageUi) const
+{
+	using namespace UE::Editor::DataStorage;
+	
+	DataStorageUi.RegisterWidgetPurpose(
+		IUiProvider::FPurposeInfo("AssetPreview", "Thumbnail", NAME_None,
+		IUiProvider::EPurposeType::UniqueByNameAndColumn,
+		LOCTEXT("AssetPreviewThumbnailPurpose", "Specific purpose display thumbnails in the AssetPreview."),
+		IUiProvider::FPurposeInfo("ContentBrowser", "Thumbnail", NAME_None).GeneratePurposeID()));
+}
+
+
+void UAssetPreviewThumbnailWidgetFactory::RegisterWidgetConstructors(UE::Editor::DataStorage::ICoreProvider& DataStorage,
+                                                              UE::Editor::DataStorage::IUiProvider& DataStorageUi) const
+{
+	using namespace UE::Editor::DataStorage::Queries;
+	DataStorageUi.RegisterWidgetFactory<FAssetPreviewThumbnailWidgetConstructor>(
+		DataStorageUi.FindPurpose(IUiProvider::FPurposeInfo("AssetPreview", "Thumbnail", NAME_None).GeneratePurposeID()),
+		TColumn<FAssetTag>() || TColumn<FFolderTag>());
+}
+
+FAssetPreviewThumbnailWidgetConstructor::FAssetPreviewThumbnailWidgetConstructor()
+	: FSimpleWidgetConstructor(StaticStruct())
+{
+}
+
+TSharedPtr<SWidget> FAssetPreviewThumbnailWidgetConstructor::CreateWidget(UE::Editor::DataStorage::ICoreProvider* DataStorage,
+	UE::Editor::DataStorage::IUiProvider* DataStorageUi,
+	UE::Editor::DataStorage::RowHandle TargetRow, UE::Editor::DataStorage::RowHandle WidgetRow,
+	const UE::Editor::DataStorage::FMetaDataView& Arguments)
+{
+	using namespace UE::Editor::DataStorage;
+
+	bool bIsAsset = DataStorage->HasColumns<FAssetTag>(TargetRow);
+	const FName NamePurposeTableName = TEXT("Editor_WidgetTable");
+
+	if (FSizeValueColumn_Experimental* SizeValueColumn = DataStorage->GetColumn<FSizeValueColumn_Experimental>(WidgetRow))
+	{
+		SizeValueColumn->SizeValue = bIsThumbnailExpandedGlobal ? 256.f : 128.f;
+	}
+
+	TSharedPtr<FTypedElementWidgetConstructor> OutThumbnailWidgetConstructorPtr;
+    auto AssignThumbnailWidgetToColumn = [&OutThumbnailWidgetConstructorPtr] (TUniquePtr<FTypedElementWidgetConstructor> Constructor, TConstArrayView<TWeakObjectPtr<const UScriptStruct>>)
+    {
+    	OutThumbnailWidgetConstructorPtr = TSharedPtr<FTypedElementWidgetConstructor>(Constructor.Release());
+    	return false;
+    };
+
+	// Thumbnail Config MetaData
+	UE::Editor::DataStorage::FMetaData ThumbnailMeta;
+	ThumbnailMeta.AddOrSetMutableData(TedsAssetDataHelper::MetaDataNames::GetThumbnailStatusMetaDataName(), true);
+	ThumbnailMeta.AddOrSetMutableData(TedsAssetDataHelper::MetaDataNames::GetThumbnailFadeInMetaDataName(), true);
+	ThumbnailMeta.AddOrSetMutableData(TedsAssetDataHelper::MetaDataNames::GetThumbnailHintTextMetaDataName(), false);
+	ThumbnailMeta.AddOrSetMutableData(TedsAssetDataHelper::MetaDataNames::GetThumbnailRealTimeOnHoveredMetaDataName(), false);
+	ThumbnailMeta.AddOrSetMutableData(TedsAssetDataHelper::MetaDataNames::GetThumbnailCanDisplayEditModePrimitiveTools(), true);
+	const FGenericMetaDataView ThumbnailMetaView = FGenericMetaDataView(ThumbnailMeta);
+
+	TArray<TWeakObjectPtr<const UScriptStruct>> ThumbnailColumns = GetThumbnailColumns();
+	DataStorageUi->CreateWidgetConstructors(DataStorageUi->FindPurpose(
+		IUiProvider::FPurposeInfo("ContentBrowser", "Thumbnail", NAME_None).GeneratePurposeID()),
+		IUiProvider::EMatchApproach::LongestMatch,
+		ThumbnailColumns,
+		ThumbnailMetaView,
+		AssignThumbnailWidgetToColumn);
+
+	TSharedPtr<SWidget> ThumbnailWidget = SNullWidget::NullWidget;
+
+	RowHandle ThumbnailWidgetRowHandle = InvalidRowHandle;
+	if (OutThumbnailWidgetConstructorPtr)
+	{
+		if (ThumbnailWidgetRowHandle = DataStorage->AddRow(DataStorage->FindTable(NamePurposeTableName));
+			DataStorage->IsRowAvailable(ThumbnailWidgetRowHandle))
+		{
+			// Set the current TargetRow
+			DataStorage->AddColumn(ThumbnailWidgetRowHandle, FTypedElementRowReferenceColumn{ .Row = TargetRow });
+
+			// Set this as the parent so that it will use the EditMode and Size of this TEDS widget
+			DataStorage->AddColumn(ThumbnailWidgetRowHandle, FTableRowParentColumn{ .Parent = WidgetRow });
+
+			ThumbnailWidget = DataStorageUi->ConstructWidget(ThumbnailWidgetRowHandle, *OutThumbnailWidgetConstructorPtr, ThumbnailMetaView);
+		}
+	}
+
+	if (!ThumbnailWidget.IsValid())
+	{
+		ThumbnailWidget = SNullWidget::NullWidget;
+	}
+
+	// If the widget creating this is interested in the ThumbnailTooltip it has to add the column itself to avoid adding them if not used
+	if (FLocalWidgetTooltipColumn_Experimental* WidgetTooltipColumn = DataStorage->GetColumn<FLocalWidgetTooltipColumn_Experimental>(WidgetRow))
+	{
+		WidgetTooltipColumn->Tooltip = ThumbnailWidget->GetToolTip();
+	}
+
+	// For assets, grab the color from the asset definition
+	if (bIsAsset)
+	{
+		return SNew(SBox)
+			[
+				SNew(SOverlay)
+
+				+ SOverlay::Slot()
+				[
+					ThumbnailWidget.ToSharedRef()
+				]
+
+				// Overlay for incrementing Thumbnail Size
+				+ SOverlay::Slot()
+				.Padding(2.f)
+				.HAlign(HAlign_Right)
+				.VAlign(VAlign_Top)
+				[
+					SNew(SBox)
+					.WidthOverride(20.f)
+					.HeightOverride(20.f)
+					[
+						SNew(SButton)
+						.ContentPadding(0.f)
+						.ButtonStyle(&FAppStyle::Get().GetWidgetStyle<FButtonStyle>("SimpleButton"))
+						.OnClicked_Lambda([DataStorage, WidgetRow] ()
+						{
+							using namespace UE::Editor::DataStorage;
+
+							bIsThumbnailExpandedGlobal = !bIsThumbnailExpandedGlobal;
+							if (FSizeValueColumn_Experimental* ThumbnailSizeValueColumn = DataStorage->GetColumn<FSizeValueColumn_Experimental>(WidgetRow))
+							{
+								if (bIsThumbnailExpandedGlobal)
+								{
+									ThumbnailSizeValueColumn->SizeValue = 256.f;
+								}
+								else
+								{
+									ThumbnailSizeValueColumn->SizeValue = 128.f;
+								}
+							}
+							return FReply::Handled();
+						})
+						[
+							SNew(SImage)
+							.Image_Lambda([] ()
+							{
+								if (bIsThumbnailExpandedGlobal)
+								{
+									return FAppStyle::GetBrush("Icons.Minus");
+								}
+								return FAppStyle::GetBrush("Icons.Plus");
+							})
+						]
+					]
+				]
+			];
+	}
+
+	return ThumbnailWidget;
+}
+
+TConstArrayView<const UScriptStruct*> FAssetPreviewThumbnailWidgetConstructor::GetAdditionalColumnsList() const
+{
+	using namespace UE::Editor::DataStorage;
+
+	static const TTypedElementColumnTypeList<
+		FSizeValueColumn_Experimental,
+		FThumbnailEditModeColumn_Experimental,
+		FLocalWidgetTooltipColumn_Experimental> Columns;
+
+	return Columns;
+}
+
+TArray<TWeakObjectPtr<const UScriptStruct>> FAssetPreviewThumbnailWidgetConstructor::GetThumbnailColumns()
+{
+	static TArray<TWeakObjectPtr<const UScriptStruct>> ThumbnailColumns(
+	{
+		TWeakObjectPtr(FAssetTag::StaticStruct()),
+		TWeakObjectPtr(FFolderTag::StaticStruct())
+	});
+	return ThumbnailColumns;
+}
+
+#undef LOCTEXT_NAMESPACE
diff --git a/Engine/Plugins/Experimental/EditorDataStorageFeatures/Source/TedsAssetData/Private/Widgets/AssetPreview/Widgets/Thumbnail/AssetPreviewThumbnailWidget.h b/Engine/Plugins/Experimental/EditorDataStorageFeatures/Source/TedsAssetData/Private/Widgets/AssetPreview/Widgets/Thumbnail/AssetPreviewThumbnailWidget.h
new file mode 100644
--- /dev/null
+++ b/Engine/Plugins/Experimental/EditorDataStorageFeatures/Source/TedsAssetData/Private/Widgets/AssetPreview/Widgets/Thumbnail/AssetPreviewThumbnailWidget.h
@@ -0,0 +1,49 @@
+﻿// Copyright Epic Games, Inc. All Rights Reserved.
+
+#pragma once
+
+#include "Elements/Interfaces/TypedElementDataStorageFactory.h"
+#include "Elements/Interfaces/TypedElementDataStorageUiInterface.h"
+#include "Internationalization/Text.h"
+#include "UObject/ObjectMacros.h"
+
+#include "AssetPreviewThumbnailWidget.generated.h"
+
+class UScriptStruct;
+
+UCLASS()
+class UAssetPreviewThumbnailWidgetFactory : public UEditorDataStorageFactory
+{
+	GENERATED_BODY()
+
+public:
+	virtual ~UAssetPreviewThumbnailWidgetFactory() override = default;
+
+	virtual void RegisterWidgetPurposes(UE::Editor::DataStorage::IUiProvider& DataStorageUi) const override;
+	virtual void RegisterWidgetConstructors(UE::Editor::DataStorage::ICoreProvider& DataStorage,
+		UE::Editor::DataStorage::IUiProvider& DataStorageUi) const override;
+};
+
+// Thumbnail widget for the asset preview
+USTRUCT()
+struct FAssetPreviewThumbnailWidgetConstructor : public FSimpleWidgetConstructor
+{
+	GENERATED_BODY()
+
+public:
+	FAssetPreviewThumbnailWidgetConstructor();
+	virtual ~FAssetPreviewThumbnailWidgetConstructor() override = default;
+
+	virtual TSharedPtr<SWidget> CreateWidget(
+		UE::Editor::DataStorage::ICoreProvider* DataStorage,
+		UE::Editor::DataStorage::IUiProvider* DataStorageUi,
+		UE::Editor::DataStorage::RowHandle TargetRow,
+		UE::Editor::DataStorage::RowHandle WidgetRow, 
+		const UE::Editor::DataStorage::FMetaDataView& Arguments) override;
+
+	virtual TConstArrayView<const UScriptStruct*> GetAdditionalColumnsList() const override;
+
+private:
+	/** Get the columns used for the Thumbnail TEDS widget */
+	TArray<TWeakObjectPtr<const UScriptStruct>> GetThumbnailColumns();
+};
diff --git a/Engine/Plugins/Experimental/EditorDataStorageFeatures/Source/TedsAssetData/Public/TedsAssetDataHelper.h b/Engine/Plugins/Experimental/EditorDataStorageFeatures/Source/TedsAssetData/Public/TedsAssetDataHelper.h
--- a/Engine/Plugins/Experimental/EditorDataStorageFeatures/Source/TedsAssetData/Public/TedsAssetDataHelper.h
+++ b/Engine/Plugins/Experimental/EditorDataStorageFeatures/Source/TedsAssetData/Public/TedsAssetDataHelper.h
@@ -14,6 +14,7 @@ namespace TedsAssetDataHelper
 		TEDSASSETDATA_API FName GetThumbnailHintTextMetaDataName();
 		TEDSASSETDATA_API FName GetThumbnailRealTimeOnHoveredMetaDataName();
 		TEDSASSETDATA_API FName GetThumbnailSizeOffsetMetaDataName();
+		TEDSASSETDATA_API FName GetThumbnailCanDisplayEditModePrimitiveTools();
 	}
 
 	namespace TableView
diff --git a/Engine/Plugins/Experimental/EditorDataStorageFeatures/Source/TedsAssetData/Public/Widgets/AssetPreview/STedsAssetPreviewWidget.h b/Engine/Plugins/Experimental/EditorDataStorageFeatures/Source/TedsAssetData/Public/Widgets/AssetPreview/STedsAssetPreviewWidget.h
new file mode 100644
--- /dev/null
+++ b/Engine/Plugins/Experimental/EditorDataStorageFeatures/Source/TedsAssetData/Public/Widgets/AssetPreview/STedsAssetPreviewWidget.h
@@ -0,0 +1,109 @@
+// Copyright Epic Games, Inc. All Rights Reserved.
+
+#pragma once
+
+#include "Elements/Common/TypedElementHandles.h"
+#include "Elements/Interfaces/TypedElementDataStorageUiInterface.h"
+#include "Widgets/SCompoundWidget.h"
+
+namespace UE::Editor::DataStorage
+{
+	class ITedsWidget;
+}
+
+class SBorder;
+
+class STedsAssetPreviewWidget : public SCompoundWidget
+{
+public:
+	SLATE_BEGIN_ARGS(STedsAssetPreviewWidget)
+		: _WidgetPurpose(UE::Editor::DataStorage::IUiProvider::FPurposeInfo("AssetPreview", "Default", NAME_None).GeneratePurposeID())
+		, _HeaderWidgetPurpose(UE::Editor::DataStorage::IUiProvider::FPurposeInfo("AssetPreview", "Header", "Default").GeneratePurposeID())
+		, _ThumbnailWidgetPurpose(UE::Editor::DataStorage::IUiProvider::FPurposeInfo("AssetPreview", "Thumbnail", NAME_None).GeneratePurposeID())
+		, _TargetRow(UE::Editor::DataStorage::InvalidRowHandle)
+	{}
+
+		// The widget purpose to use to create the tile widget
+		SLATE_ARGUMENT(UE::Editor::DataStorage::IUiProvider::FPurposeID, WidgetPurpose)
+
+		// The widget header purpose to use to create the tile widget
+		SLATE_ARGUMENT(UE::Editor::DataStorage::IUiProvider::FPurposeID, HeaderWidgetPurpose)
+
+		// The widget thumbnail purpose to use to create the tile widget
+		SLATE_ARGUMENT(UE::Editor::DataStorage::IUiProvider::FPurposeID, ThumbnailWidgetPurpose)
+
+		// Target Row of the Item we want to create the preview of
+		SLATE_ARGUMENT(UE::Editor::DataStorage::RowHandle, TargetRow)
+
+	SLATE_END_ARGS()
+
+	TEDSASSETDATA_API void Construct(const FArguments& InArgs);
+
+	/** Get the Widget Row Handle of this TedsWidget */
+	TEDSASSETDATA_API UE::Editor::DataStorage::RowHandle GetWidgetRowHandle() const;
+
+	/** Reconstructs all the TedsWidget used in this widget */
+	TEDSASSETDATA_API void ReconstructTedsWidget();
+
+	/** Set the current Target Row to use, needs ReconstructTedsWidget to change */
+	TEDSASSETDATA_API void SetTargetRow(UE::Editor::DataStorage::RowHandle InTargetRow);
+
+private:
+	/** Create and initialize the ITedsWidget */
+	void CreateTedsWidget();
+
+	/** Create all Teds widgets for this widget based on the current TargetRow */
+	void ConstructWidget();
+
+	/** Get the columns used for the Asset Preview Header TEDS widget */
+	TArray<TWeakObjectPtr<const UScriptStruct>> GetHeaderColumns();
+
+	/** Get the columns used for the Thumbnail TEDS widget */
+    TArray<TWeakObjectPtr<const UScriptStruct>> GetThumbnailColumns();
+
+	/** Get the columns used for the Asset Preview BasicInfo TEDS widget */
+	TArray<TWeakObjectPtr<const UScriptStruct>> GetBasicInfoColumns();
+
+	/** Get the columns used for the Asset Preview AdvancedInfo TEDS widget */
+	TArray<TWeakObjectPtr<const UScriptStruct>> GetAdvancedInfoColumns();
+
+	/** Get the current Target Row */
+	UE::Editor::DataStorage::RowHandle GetReferencedRowHandle() const;
+
+private:
+	/** Header Row Handle */
+	UE::Editor::DataStorage::RowHandle HeaderWidgetRowHandle = UE::Editor::DataStorage::InvalidRowHandle;
+
+	/** Thumbnail Row Handle */
+	UE::Editor::DataStorage::RowHandle ThumbnailWidgetRowHandle = UE::Editor::DataStorage::InvalidRowHandle;
+
+	/** Basic Info Row Handle */
+	UE::Editor::DataStorage::RowHandle BasicInfoWidgetRowHandle = UE::Editor::DataStorage::InvalidRowHandle;
+
+	/** Advanced Info Row Handle */
+	UE::Editor::DataStorage::RowHandle AdvancedInfoWidgetRowHandle = UE::Editor::DataStorage::InvalidRowHandle;
+
+	/** Target Row Handle */
+	UE::Editor::DataStorage::RowHandle TargetRow = UE::Editor::DataStorage::InvalidRowHandle;
+
+	/** Asset preview Purpose */
+	UE::Editor::DataStorage::IUiProvider::FPurposeID WidgetPurpose;
+
+	/** Asset preview Header Purpose */
+	UE::Editor::DataStorage::IUiProvider::FPurposeID WidgetHeaderPurpose;
+
+	/** Asset preview Thumbnail Purpose */
+	UE::Editor::DataStorage::IUiProvider::FPurposeID ThumbnailWidgetPurpose;
+
+	/** Whether the thumbnail is big or small */
+	bool bIsThumbnailExpanded = false;
+
+	/** Whether the thumbnail edit mode is enabled or not */
+	bool IsEditModeEnabled = false;
+
+	/** Our teds widget */
+	TSharedPtr<UE::Editor::DataStorage::ITedsWidget> TedsWidget;
+
+	/** Border containing the widget (child of the TedsWdiget) */
+	TSharedPtr<SBorder> PreviewPanel;
+};
diff --git a/Engine/Plugins/Experimental/EditorDataStorageFeatures/Source/TedsAssetData/Public/Widgets/AssetPreview/TedsAssetPreviewWidgetColumns.h b/Engine/Plugins/Experimental/EditorDataStorageFeatures/Source/TedsAssetData/Public/Widgets/AssetPreview/TedsAssetPreviewWidgetColumns.h
new file mode 100644
--- /dev/null
+++ b/Engine/Plugins/Experimental/EditorDataStorageFeatures/Source/TedsAssetData/Public/Widgets/AssetPreview/TedsAssetPreviewWidgetColumns.h
@@ -0,0 +1,23 @@
+// Copyright Epic Games, Inc. All Rights Reserved.
+
+#pragma once
+
+#include <Elements/Common/TypedElementCommonTypes.h>
+#include "Framework/SlateDelegates.h"
+
+#include "TedsAssetPreviewWidgetColumns.generated.h"
+
+
+/**
+ * Column added to a widget row when an external widget manages On Clicked behavior
+ */
+USTRUCT(meta = (DisplayName = "Widget with externally managed On Clicked behavior"))
+struct FExternalWidgetOnClickedColumn_Experimental : public FEditorDataStorageColumn
+{
+	GENERATED_BODY()
+
+	/** Delegate to execute during the On Click of the widget */
+	// TODO: BindEvent currently do not support custom return value, since FReply do not have a default constructor it will fail
+	// FOnClicked OnClicked;
+	TDelegate<void()> OnClicked;
+};
diff --git a/Engine/Plugins/Experimental/EditorDataStorageFeatures/Source/TedsAssetData/TedsAssetData.build.cs b/Engine/Plugins/Experimental/EditorDataStorageFeatures/Source/TedsAssetData/TedsAssetData.build.cs
--- a/Engine/Plugins/Experimental/EditorDataStorageFeatures/Source/TedsAssetData/TedsAssetData.build.cs
+++ b/Engine/Plugins/Experimental/EditorDataStorageFeatures/Source/TedsAssetData/TedsAssetData.build.cs
@@ -32,11 +32,10 @@ public class TedsAssetData : ModuleRules
 					"Slate",
 					"SlateCore",
 					"TedsAlerts",
+					"TedsTableViewer",
 					"ToolWidgets",
 					"TypedElementFramework",
 					"UnrealEd",
-					"UnrealEd",
-					"AssetDefinition",
 				});
 
 			PublicDependencyModuleNames.AddRange(
diff --git a/Engine/Plugins/Experimental/EditorDataStorageFeatures/Source/TedsContentBrowser/Private/TedsContentBrowserModule.cpp b/Engine/Plugins/Experimental/EditorDataStorageFeatures/Source/TedsContentBrowser/Private/TedsContentBrowserModule.cpp
--- a/Engine/Plugins/Experimental/EditorDataStorageFeatures/Source/TedsContentBrowser/Private/TedsContentBrowserModule.cpp
+++ b/Engine/Plugins/Experimental/EditorDataStorageFeatures/Source/TedsContentBrowser/Private/TedsContentBrowserModule.cpp
@@ -23,6 +23,7 @@
 #include "Elements/Columns/TypedElementVisibilityColumns.h"
 #include "Elements/Framework/TypedElementQueryBuilder.h"
 #include "Widgets/STedsTableViewer.h"
+#include "Widgets/AssetPreview/STedsAssetPreviewWidget.h"
 #include "Widgets/STedsTileViewer.h"
 #include "Widgets/Views/SListView.h"
 
@@ -238,6 +239,8 @@ namespace UE::Editor::ContentBrowser
 					.ListSelectionMode(ESelectionMode::Multi)
 					.OnSelectionChanged_Lambda([this](RowHandle Row)
 					{
+						UpdateAssetPreviewTargetRow();
+
 						if(TSharedPtr<FAssetViewItem> AssetViewItem = GetAssetViewItemFromRow(Row))
 						{
 							// CB 2.0 TODO: Does the CB use ESelectInfo and we need to propagate it from the table viewer?
@@ -266,7 +269,9 @@ namespace UE::Editor::ContentBrowser
 					.ItemHeight_Lambda([this] () { return GetTileViewItemHeight(); })
 					.OnSelectionChanged_Lambda([this](RowHandle Row)
 					{
-						if(TSharedPtr<FAssetViewItem> AssetViewItem = GetAssetViewItemFromRow(Row))
+						UpdateAssetPreviewTargetRow();
+
+						if (TSharedPtr<FAssetViewItem> AssetViewItem = GetAssetViewItemFromRow(Row))
 						{
 							// CB 2.0 TODO: Does the CB use ESelectInfo and we need to propagate it from the table viewer?
 							OnSelectionChangedDelegate.Execute(AssetViewItem, ESelectInfo::Direct);
@@ -321,6 +326,57 @@ namespace UE::Editor::ContentBrowser
 		}
 	}
 
+	void FTedsContentBrowserViewExtender::CreateAssetPreview()
+	{
+		AssetPreviewWidget = SNew(STedsAssetPreviewWidget);
+
+		BindAssetPreviewColumns();
+
+		UpdateAssetPreviewTargetRow();
+	}
+
+	void FTedsContentBrowserViewExtender::BindAssetPreviewColumns()
+	{
+		if (!AssetPreviewWidget.IsValid())
+		{
+			return;
+		}
+
+		UE::Editor::DataStorage::RowHandle AssetPreviewRowHandle = AssetPreviewWidget->GetWidgetRowHandle();
+		if (FWidgetContextMenuColumn* ContextMenuColumn = DataStorage->GetColumn<FWidgetContextMenuColumn>(AssetPreviewRowHandle))
+		{
+			ContextMenuColumn->OnContextMenuOpening.BindLambda([this]()
+			{
+				return OnContextMenuOpenedDelegate.Execute();
+			});
+		}
+	}
+
+	void FTedsContentBrowserViewExtender::UpdateAssetPreviewTargetRow() const
+	{
+		if (AssetPreviewWidget.IsValid())
+		{
+			TArray<FTedsRowHandle> SelectedItems;
+			TableViewer->ForEachSelectedRow([this, &SelectedItems](RowHandle InRow)
+			{
+				SelectedItems.Add(FTedsRowHandle(InRow));
+			});
+
+			if (SelectedItems.Num())
+			{
+				AssetPreviewWidget->SetTargetRow(SelectedItems.Last());
+			}
+			else
+			{
+				// Nothing is selected so don't show anything
+				AssetPreviewWidget->SetTargetRow(InvalidRowHandle);
+			}
+
+			// Reconstruct Teds Widgets for the AssetPreview since Purpose and TargetRow may change
+			AssetPreviewWidget->ReconstructTedsWidget();
+		}
+	}
+
 	float FTedsContentBrowserViewExtender::GetTileViewTypeNameHeight() const
 	{
 		if (CurrentThumbnailSize == EThumbnailSize::Tiny)
@@ -452,10 +508,25 @@ namespace UE::Editor::ContentBrowser
 		default:
 			CreateListView();
 		}
-		
+
+		CreateAssetPreview();
+
 		RefreshRows(InItemsSource);
-		
-		return TableViewer->AsWidget();
+
+		return SNew(SSplitter)
+			.Orientation(Orient_Horizontal)
+
+			+ SSplitter::Slot()
+			.Resizable(true)
+			[
+				TableViewer->AsWidget()
+			]
+
+			+ SSplitter::Slot()
+			.Resizable(true)
+			[
+				AssetPreviewWidget.ToSharedRef()
+			];
 	}
 	
 	void FTedsContentBrowserViewExtender::OnItemListChanged(TArray<TSharedPtr<FAssetViewItem>>* InItemsSource)
diff --git a/Engine/Plugins/Experimental/EditorDataStorageFeatures/Source/TedsContentBrowser/Private/TedsContentBrowserModule.h b/Engine/Plugins/Experimental/EditorDataStorageFeatures/Source/TedsContentBrowser/Private/TedsContentBrowserModule.h
--- a/Engine/Plugins/Experimental/EditorDataStorageFeatures/Source/TedsContentBrowser/Private/TedsContentBrowserModule.h
+++ b/Engine/Plugins/Experimental/EditorDataStorageFeatures/Source/TedsContentBrowser/Private/TedsContentBrowserModule.h
@@ -22,6 +22,7 @@ namespace UE::Editor::DataStorage
 }
 
 class FAssetViewItem;
+class STedsAssetPreviewWidget;
 
 namespace UE::Editor::ContentBrowser
 {
@@ -83,6 +84,15 @@ namespace UE::Editor::ContentBrowser
 		// Bind the table viewer's columns to CB delegates
 		void BindViewColumns();
 
+		// Create the asset preview widget
+		void CreateAssetPreview();
+
+		// Bind the AssetPreview columns to CB delegates
+		void BindAssetPreviewColumns();
+
+		// Update the AssetPreview Target Row
+		void UpdateAssetPreviewTargetRow() const;
+
 	private:
 		/** Get the name area total height */
 		float GetTileViewTypeNameHeight() const;
@@ -115,6 +125,9 @@ namespace UE::Editor::ContentBrowser
 		void UpdateEditMode() const;
 
 	private:
+		/** AssetPreview widget */
+		TSharedPtr<STedsAssetPreviewWidget> AssetPreviewWidget;
+
 		/** Custom view type */
 		ETableViewMode::Type CustomViewType = ETableViewMode::List;
 
