[TEDS] Add back ParentNode->GetRevision to the RowMonitorNode since it was not entering the loop needed to broadcast row adding and removing

diff --git a/Engine/Plugins/Experimental/EditorDataStorageFeatures/Source/TedsQueryStack/Private/TedsRowMonitorNode.cpp b/Engine/Plugins/Experimental/EditorDataStorageFeatures/Source/TedsQueryStack/Private/TedsRowMonitorNode.cpp
--- a/Engine/Plugins/Experimental/EditorDataStorageFeatures/Source/TedsQueryStack/Private/TedsRowMonitorNode.cpp
+++ b/Engine/Plugins/Experimental/EditorDataStorageFeatures/Source/TedsQueryStack/Private/TedsRowMonitorNode.cpp
@@ -97,6 +97,7 @@ namespace UE::Editor::DataStorage::QueryStack
 			ParentNode->Update();
 			if (ParentNode->GetRevision() != ParentRevision)
 			{
+				ParentRevision = ParentNode->GetRevision();
 				Revision++;
 				MergeChanges = false;
 			}
