Teds Outliner: Unregister any previously existing row mappings to work around an ensure in MapRow because of old outliner instances not being destructed yet due to dangling references

diff --git a/Engine/Plugins/Experimental/EditorDataStorageFeatures/Source/TedsOutliner/Private/TedsOutlinerImpl.cpp b/Engine/Plugins/Experimental/EditorDataStorageFeatures/Source/TedsOutliner/Private/TedsOutlinerImpl.cpp
--- a/Engine/Plugins/Experimental/EditorDataStorageFeatures/Source/TedsOutliner/Private/TedsOutlinerImpl.cpp
+++ b/Engine/Plugins/Experimental/EditorDataStorageFeatures/Source/TedsOutliner/Private/TedsOutlinerImpl.cpp
@@ -445,6 +445,11 @@ void FTedsOutlinerImpl::RegisterTedsOutliner()
 
 		if (!Identifier.IsNone())
 		{
+			// Any old outliner instances could still have dangling references if external systems are holding onto them - which means it might only
+			// be destructed after this new instance is init and MapRow ensures on duplicate keys. At this point it is safe to remove the old mapping
+			// since all the data will be fine because the old row exists till the destructor is called and we want any systems using the mapping
+			// to start looking at the new instance anyways
+			Storage->RemoveRowMapping(MappingDomain, DataStorage::FMapKey(SceneOutliner->GetOutlinerIdentifier()));
 			Storage->MapRow(MappingDomain, DataStorage::FMapKey(SceneOutliner->GetOutlinerIdentifier()), OutlinerRowHandle);
 		}
 	}
