TedsSettings : Add a FSettingsNameColumn and use it instead of the shared FNameColumn.

This fixes a regression where the FNameColumn behavior changed for UObject rows.

diff --git a/Engine/Plugins/Experimental/EditorDataStorageFeatures/Source/TedsSettings/Private/TedsSettingsManager.cpp b/Engine/Plugins/Experimental/EditorDataStorageFeatures/Source/TedsSettings/Private/TedsSettingsManager.cpp
--- a/Engine/Plugins/Experimental/EditorDataStorageFeatures/Source/TedsSettings/Private/TedsSettingsManager.cpp
+++ b/Engine/Plugins/Experimental/EditorDataStorageFeatures/Source/TedsSettings/Private/TedsSettingsManager.cpp
@@ -41,7 +41,7 @@ namespace UE::Editor::Settings::Private
 			&& ColumnType != FSettingsInactiveSectionTag::StaticStruct()
 			&& ColumnType != FSettingsContainerReferenceColumn::StaticStruct()
 			&& ColumnType != FSettingsCategoryReferenceColumn::StaticStruct()
-			&& ColumnType != FNameColumn::StaticStruct()
+			&& ColumnType != FSettingsNameColumn::StaticStruct()
 			&& ColumnType != FDisplayNameColumn::StaticStruct()
 			&& ColumnType != FDescriptionColumn::StaticStruct()
 			&& ColumnType != FTypedElementUObjectColumn::StaticStruct()
@@ -154,7 +154,7 @@ UE::Editor::DataStorage::RowHandle FTedsSettingsManager::FindOrAddSettingsSectio
 				SectionRow = DataStorage->AddRow(SettingsInactiveSectionTable);
 
 				DataStorage->AddColumn<FSettingsInactiveSectionTag>(SectionRow);
-				DataStorage->AddColumn<FNameColumn>(SectionRow, { .Name = SectionName });
+				DataStorage->AddColumn<FSettingsNameColumn>(SectionRow, { .Name = SectionName });
 				DataStorage->AddColumn<FSettingsContainerReferenceColumn>(SectionRow, { .ContainerName = ContainerName, .ContainerRow = InvalidRowHandle });
 				DataStorage->AddColumn<FSettingsCategoryReferenceColumn>(SectionRow, { .CategoryName = CategoryName, .CategoryRow = InvalidRowHandle });
 
@@ -180,7 +180,7 @@ bool FTedsSettingsManager::GetSettingsSectionFromRow(UE::Editor::DataStorage::Ro
 			{
 				const FSettingsContainerReferenceColumn* ContainerReferenceColumn = DataStorage->GetColumn<FSettingsContainerReferenceColumn>(Row);
 				const FSettingsCategoryReferenceColumn* CategoryReferenceColumn = DataStorage->GetColumn<FSettingsCategoryReferenceColumn>(Row);
-				const FNameColumn* NameColumn = DataStorage->GetColumn<FNameColumn>(Row);
+				const FSettingsNameColumn* NameColumn = DataStorage->GetColumn<FSettingsNameColumn>(Row);
 
 				if (ContainerReferenceColumn && CategoryReferenceColumn && NameColumn)
 				{
@@ -203,21 +203,21 @@ void FTedsSettingsManager::RegisterTables(UE::Editor::DataStorage::ICoreProvider
 	if (SettingsContainerTable == InvalidTableHandle)
 	{
 		SettingsContainerTable = DataStorage.RegisterTable(
-			TTypedElementColumnTypeList<FNameColumn, FDisplayNameColumn, FDescriptionColumn, FSettingsContainerTag>(),
+			TTypedElementColumnTypeList<FSettingsNameColumn, FDisplayNameColumn, FDescriptionColumn, FSettingsContainerTag>(),
 			FName(TEXT("Editor_SettingsContainerTable")));
 	}
 
 	if (SettingsCategoryTable == InvalidTableHandle)
 	{
 		SettingsCategoryTable = DataStorage.RegisterTable(
-			TTypedElementColumnTypeList<FSettingsContainerReferenceColumn, FNameColumn, FDisplayNameColumn, FDescriptionColumn, FSettingsCategoryTag>(),
+			TTypedElementColumnTypeList<FSettingsContainerReferenceColumn, FSettingsNameColumn, FDisplayNameColumn, FDescriptionColumn, FSettingsCategoryTag>(),
 			FName(TEXT("Editor_SettingsCategoryTable")));
 	}
 
 	if (SettingsInactiveSectionTable == InvalidTableHandle)
 	{
 		SettingsInactiveSectionTable = DataStorage.RegisterTable(
-			TTypedElementColumnTypeList<FSettingsContainerReferenceColumn, FSettingsCategoryReferenceColumn, FNameColumn, FSettingsInactiveSectionTag>(),
+			TTypedElementColumnTypeList<FSettingsContainerReferenceColumn, FSettingsCategoryReferenceColumn, FSettingsNameColumn, FSettingsInactiveSectionTag>(),
 			FName(TEXT("Editor_SettingsInactiveSectionTable")));
 	}
 }
@@ -230,7 +230,7 @@ void FTedsSettingsManager::RegisterQueries(UE::Editor::DataStorage::ICoreProvide
 	{
 		SelectAllActiveSettingsQuery = DataStorage.RegisterQuery(
 			Select()
-				.ReadOnly<FSettingsContainerReferenceColumn, FSettingsCategoryReferenceColumn, FNameColumn>()
+				.ReadOnly<FSettingsContainerReferenceColumn, FSettingsCategoryReferenceColumn, FSettingsNameColumn>()
 			.Where()
 				.All<FSettingsSectionTag>()
 			.Compile());
@@ -240,7 +240,7 @@ void FTedsSettingsManager::RegisterQueries(UE::Editor::DataStorage::ICoreProvide
 	{
 		SelectAllInactiveSettingsQuery = DataStorage.RegisterQuery(
 			Select()
-				.ReadOnly<FSettingsContainerReferenceColumn, FSettingsCategoryReferenceColumn, FNameColumn>()
+				.ReadOnly<FSettingsContainerReferenceColumn, FSettingsCategoryReferenceColumn, FSettingsNameColumn>()
 			.Where()
 				.All<FSettingsInactiveSectionTag>()
 			.Compile());
@@ -295,7 +295,7 @@ void FTedsSettingsManager::RegisterSettingsContainer(const FName& ContainerName)
 	if (ContainerRow == InvalidRowHandle)
 	{
 		ContainerRow = DataStorage->AddRow(SettingsContainerTable);
-		DataStorage->AddColumn<FNameColumn>(ContainerRow, { .Name = ContainerName });
+		DataStorage->AddColumn<FSettingsNameColumn>(ContainerRow, { .Name = ContainerName });
 		DataStorage->AddColumn<FDisplayNameColumn>(ContainerRow, { .DisplayName = ContainerPtr->GetDisplayName() });
 		DataStorage->AddColumn<FDescriptionColumn>(ContainerRow, { .Description = ContainerPtr->GetDescription() });
 		DataStorage->AddColumn<FSettingsContainerTag>(ContainerRow);
@@ -427,7 +427,7 @@ void FTedsSettingsManager::UpdateSettingsCategory(TSharedPtr<ISettingsCategory>
 	ICompatibilityProvider* DataStorageCompatibility = GetMutableDataStorageFeature<ICompatibilityProvider>(CompatibilityFeatureName);
 	check(DataStorageCompatibility);
 
-	const FName& ContainerName = DataStorage->GetColumn<FNameColumn>(ContainerRow)->Name;
+	const FName& ContainerName = DataStorage->GetColumn<FSettingsNameColumn>(ContainerRow)->Name;
 	const FName& CategoryName = SettingsCategory->GetName();
 
 	UE_LOG(LogTedsSettings, Verbose, TEXT("Update Settings Category: '%s->%s'"), *ContainerName.ToString(), *CategoryName.ToString());
@@ -440,7 +440,7 @@ void FTedsSettingsManager::UpdateSettingsCategory(TSharedPtr<ISettingsCategory>
 		CategoryRow = DataStorage->AddRow(SettingsCategoryTable);
 
 		DataStorage->AddColumn<FSettingsContainerReferenceColumn>(CategoryRow, { .ContainerName = ContainerName, .ContainerRow = ContainerRow });
-		DataStorage->AddColumn<FNameColumn>(CategoryRow, { .Name = CategoryName });
+		DataStorage->AddColumn<FSettingsNameColumn>(CategoryRow, { .Name = CategoryName });
 		DataStorage->AddColumn<FDisplayNameColumn>(CategoryRow, { .DisplayName = SettingsCategory->GetDisplayName() });
 		DataStorage->AddColumn<FDescriptionColumn>(CategoryRow, { .Description = SettingsCategory->GetDescription() });
 		DataStorage->AddColumn<FSettingsCategoryTag>(CategoryRow);
@@ -461,7 +461,7 @@ void FTedsSettingsManager::UpdateSettingsCategory(TSharedPtr<ISettingsCategory>
 				IDirectQueryContext& Context,
 				const FSettingsContainerReferenceColumn* ContainerColumns,
 				const FSettingsCategoryReferenceColumn* CategoryColumns,
-				const FNameColumn* SectionNameColumns)
+				const FSettingsNameColumn* SectionNameColumns)
 			{
 				const uint32 RowCount = Context.GetRowCount();
 
@@ -520,7 +520,7 @@ void FTedsSettingsManager::UpdateSettingsCategory(TSharedPtr<ISettingsCategory>
 			DataStorage->AddColumn<FSettingsSectionTag>(NewSectionRow);
 			DataStorage->AddColumn<FSettingsContainerReferenceColumn>(NewSectionRow, { .ContainerName = ContainerName, .ContainerRow = ContainerRow });
 			DataStorage->AddColumn<FSettingsCategoryReferenceColumn>(NewSectionRow, { .CategoryName = CategoryName, .CategoryRow = CategoryRow });
-			DataStorage->AddColumn<FNameColumn>(NewSectionRow, { .Name = SectionName });
+			DataStorage->AddColumn<FSettingsNameColumn>(NewSectionRow, { .Name = SectionName });
 			DataStorage->AddColumn<FDisplayNameColumn>(NewSectionRow, { .DisplayName = SectionPtr->GetDisplayName() });
 			DataStorage->AddColumn<FDescriptionColumn>(NewSectionRow, { .Description = SectionPtr->GetDescription() });
 			
@@ -558,7 +558,7 @@ void FTedsSettingsManager::UpdateSettingsCategory(TSharedPtr<ISettingsCategory>
 		RowHandle NewSectionRow = DataStorage->AddRow(SettingsInactiveSectionTable);
 
 		DataStorage->AddColumn<FSettingsInactiveSectionTag>(NewSectionRow);
-		DataStorage->AddColumn<FNameColumn>(NewSectionRow, { .Name = OldSectionName });
+		DataStorage->AddColumn<FSettingsNameColumn>(NewSectionRow, { .Name = OldSectionName });
 		DataStorage->AddColumn<FSettingsContainerReferenceColumn>(NewSectionRow, { .ContainerName = ContainerName, .ContainerRow = InvalidRowHandle });
 		DataStorage->AddColumn<FSettingsCategoryReferenceColumn>(NewSectionRow, { .CategoryName = CategoryName, .CategoryRow = InvalidRowHandle });
 
diff --git a/Engine/Plugins/Experimental/EditorDataStorageFeatures/Source/TedsSettings/Public/TedsSettingsColumns.h b/Engine/Plugins/Experimental/EditorDataStorageFeatures/Source/TedsSettings/Public/TedsSettingsColumns.h
--- a/Engine/Plugins/Experimental/EditorDataStorageFeatures/Source/TedsSettings/Public/TedsSettingsColumns.h
+++ b/Engine/Plugins/Experimental/EditorDataStorageFeatures/Source/TedsSettings/Public/TedsSettingsColumns.h
@@ -36,6 +36,17 @@ struct FSettingsCategoryReferenceColumn final : public FEditorDataStorageColumn
 	UE::Editor::DataStorage::RowHandle CategoryRow;
 };
 
+// A generic name column for use within the settings domain.
+// This is added to all container, category, and section rows.
+USTRUCT(meta = (DisplayName = "Name"))
+struct FSettingsNameColumn final : public FEditorDataStorageColumn
+{
+	GENERATED_BODY()
+
+	UPROPERTY(meta = (Searchable))
+	FName Name;
+};
+
 USTRUCT(meta = (DisplayName = "Settings Container"))
 struct FSettingsContainerTag final : public FEditorDataStorageTag
 {
diff --git a/Engine/Plugins/Experimental/EditorDataStorageFeatures/Source/TedsSettings/Tests/TedsSettingsTests.cpp b/Engine/Plugins/Experimental/EditorDataStorageFeatures/Source/TedsSettings/Tests/TedsSettingsTests.cpp
--- a/Engine/Plugins/Experimental/EditorDataStorageFeatures/Source/TedsSettings/Tests/TedsSettingsTests.cpp
+++ b/Engine/Plugins/Experimental/EditorDataStorageFeatures/Source/TedsSettings/Tests/TedsSettingsTests.cpp
@@ -137,7 +137,7 @@ namespace UE::Editor::Settings::Tests
 
 					TestEqual(TEXT("ContainerName"), DataStorage->GetColumn<FSettingsContainerReferenceColumn>(RowHandle)->ContainerName, ContainerName);
 					TestEqual(TEXT("CategoryName"), DataStorage->GetColumn<FSettingsCategoryReferenceColumn>(RowHandle)->CategoryName, CategoryName);
-					TestEqual(TEXT("SectionName"), DataStorage->GetColumn<FNameColumn>(RowHandle)->Name, SectionName);
+					TestEqual(TEXT("SectionName"), DataStorage->GetColumn<FSettingsNameColumn>(RowHandle)->Name, SectionName);
 				});
 			});
 		});
@@ -237,7 +237,7 @@ namespace UE::Editor::Settings::Tests
 
 					TestEqual(TEXT("ContainerName"), DataStorage->GetColumn<FSettingsContainerReferenceColumn>(RowHandle)->ContainerName, ContainerName);
 					TestEqual(TEXT("CategoryName"), DataStorage->GetColumn<FSettingsCategoryReferenceColumn>(RowHandle)->CategoryName, CategoryName);
-					TestEqual(TEXT("SectionName"), DataStorage->GetColumn<FNameColumn>(RowHandle)->Name, SectionName);
+					TestEqual(TEXT("SectionName"), DataStorage->GetColumn<FSettingsNameColumn>(RowHandle)->Name, SectionName);
 				});
 			});
 		});
@@ -271,7 +271,7 @@ namespace UE::Editor::Settings::Tests
 
 					TestEqual(TEXT("ContainerName"), DataStorage->GetColumn<FSettingsContainerReferenceColumn>(RowHandle)->ContainerName, ContainerName);
 					TestEqual(TEXT("CategoryName"), DataStorage->GetColumn<FSettingsCategoryReferenceColumn>(RowHandle)->CategoryName, CategoryName);
-					TestEqual(TEXT("SectionName"), DataStorage->GetColumn<FNameColumn>(RowHandle)->Name, SectionName);
+					TestEqual(TEXT("SectionName"), DataStorage->GetColumn<FSettingsNameColumn>(RowHandle)->Name, SectionName);
 
 					TestTrue(TEXT("FSettingsInactiveSectionTag"), DataStorage->HasColumns<FSettingsInactiveSectionTag>(RowHandle));
 				});
@@ -324,7 +324,7 @@ namespace UE::Editor::Settings::Tests
 
 					TestEqual(TEXT("ContainerName"), DataStorage->GetColumn<FSettingsContainerReferenceColumn>(RowHandle)->ContainerName, ContainerName);
 					TestEqual(TEXT("CategoryName"), DataStorage->GetColumn<FSettingsCategoryReferenceColumn>(RowHandle)->CategoryName, CategoryName);
-					TestEqual(TEXT("SectionName"), DataStorage->GetColumn<FNameColumn>(RowHandle)->Name, SectionName);
+					TestEqual(TEXT("SectionName"), DataStorage->GetColumn<FSettingsNameColumn>(RowHandle)->Name, SectionName);
 
 					TestFalse(TEXT("FSettingsInactiveSectionTag"), DataStorage->HasColumns<FSettingsInactiveSectionTag>(RowHandle));
 					TestTrue(TEXT("FSettingsSectionTag"), DataStorage->HasColumns<FSettingsSectionTag>(RowHandle));
