[TEDS Actor Compatibility] Create a new module for the EditorDataStorageFeatures called TedsActorCompatibility and move actor specific queries out of TedsCore

diff --git a/Engine/Plugins/Experimental/EditorDataStorage/Source/TedsCore/Private/Compatibility/TypedElementActorIconOverrideQueries.cpp b/Engine/Plugins/Experimental/EditorDataStorage/Source/TedsCore/Private/Compatibility/TypedElementActorIconOverrideQueries.cpp
deleted file mode 100644
diff --git a/Engine/Plugins/Experimental/EditorDataStorage/Source/TedsCore/Private/Compatibility/TypedElementActorIconOverrideQueries.h b/Engine/Plugins/Experimental/EditorDataStorage/Source/TedsCore/Private/Compatibility/TypedElementActorIconOverrideQueries.h
deleted file mode 100644
diff --git a/Engine/Plugins/Experimental/EditorDataStorage/Source/TedsCore/Private/Compatibility/TypedElementActorLabelQueries.cpp b/Engine/Plugins/Experimental/EditorDataStorage/Source/TedsCore/Private/Compatibility/TypedElementActorLabelQueries.cpp
deleted file mode 100644
diff --git a/Engine/Plugins/Experimental/EditorDataStorage/Source/TedsCore/Private/Compatibility/TypedElementActorLabelQueries.h b/Engine/Plugins/Experimental/EditorDataStorage/Source/TedsCore/Private/Compatibility/TypedElementActorLabelQueries.h
deleted file mode 100644
diff --git a/Engine/Plugins/Experimental/EditorDataStorage/Source/TedsCore/Private/Compatibility/TypedElementActorParentQueries.cpp b/Engine/Plugins/Experimental/EditorDataStorage/Source/TedsCore/Private/Compatibility/TypedElementActorParentQueries.cpp
deleted file mode 100644
diff --git a/Engine/Plugins/Experimental/EditorDataStorage/Source/TedsCore/Private/Compatibility/TypedElementActorParentQueries.h b/Engine/Plugins/Experimental/EditorDataStorage/Source/TedsCore/Private/Compatibility/TypedElementActorParentQueries.h
deleted file mode 100644
diff --git a/Engine/Plugins/Experimental/EditorDataStorage/Source/TedsCore/Private/Compatibility/TypedElementActorTransformQueries.cpp b/Engine/Plugins/Experimental/EditorDataStorage/Source/TedsCore/Private/Compatibility/TypedElementActorTransformQueries.cpp
deleted file mode 100644
diff --git a/Engine/Plugins/Experimental/EditorDataStorage/Source/TedsCore/Private/Compatibility/TypedElementActorTransformQueries.h b/Engine/Plugins/Experimental/EditorDataStorage/Source/TedsCore/Private/Compatibility/TypedElementActorTransformQueries.h
deleted file mode 100644
diff --git a/Engine/Plugins/Experimental/EditorDataStorage/Source/TedsCore/Private/Compatibility/TypedElementActorViewportProcessors.cpp b/Engine/Plugins/Experimental/EditorDataStorage/Source/TedsCore/Private/Compatibility/TypedElementActorViewportProcessors.cpp
deleted file mode 100644
diff --git a/Engine/Plugins/Experimental/EditorDataStorage/Source/TedsCore/Private/Compatibility/TypedElementActorViewportProcessors.h b/Engine/Plugins/Experimental/EditorDataStorage/Source/TedsCore/Private/Compatibility/TypedElementActorViewportProcessors.h
deleted file mode 100644
diff --git a/Engine/Plugins/Experimental/EditorDataStorage/Source/TedsCore/Private/Compatibility/TypedElementActorVisibilityQueries.cpp b/Engine/Plugins/Experimental/EditorDataStorage/Source/TedsCore/Private/Compatibility/TypedElementActorVisibilityQueries.cpp
deleted file mode 100644
diff --git a/Engine/Plugins/Experimental/EditorDataStorage/Source/TedsCore/Private/Compatibility/TypedElementActorVisibilityQueries.h b/Engine/Plugins/Experimental/EditorDataStorage/Source/TedsCore/Private/Compatibility/TypedElementActorVisibilityQueries.h
deleted file mode 100644
diff --git a/Engine/Plugins/Experimental/EditorDataStorage/Source/TedsCore/Private/TypedElementDatabaseCompatibility.cpp b/Engine/Plugins/Experimental/EditorDataStorage/Source/TedsCore/Private/TypedElementDatabaseCompatibility.cpp
--- a/Engine/Plugins/Experimental/EditorDataStorage/Source/TedsCore/Private/TypedElementDatabaseCompatibility.cpp
+++ b/Engine/Plugins/Experimental/EditorDataStorage/Source/TedsCore/Private/TypedElementDatabaseCompatibility.cpp
@@ -85,11 +85,6 @@ void UEditorDataStorageCompatibility::Initialize(UEditorDataStorage* InStorage)
 	ObjectReinstancedDelegateHandle = FCoreUObjectDelegates::OnObjectsReinstanced.AddUObject(this, &UEditorDataStorageCompatibility::OnObjectReinstanced);
 	
 	PostGcUnreachableAnalysisHandle = FCoreUObjectDelegates::PostReachabilityAnalysis.AddUObject(this, &UEditorDataStorageCompatibility::OnPostGcUnreachableAnalysis);
-	// Used to get all the worlds and register the actor create/destroy handles on them.
-	PostWorldInitializationDelegateHandle = FWorldDelegates::OnPostWorldInitialization.AddUObject(this, &UEditorDataStorageCompatibility::OnPostWorldInitialization);
-	PreWorldFinishDestroyDelegateHandle = FWorldDelegates::OnPreWorldFinishDestroy.AddUObject(this, &UEditorDataStorageCompatibility::OnPreWorldFinishDestroy);
-
-	ActorOuterChangedDelegateHandle = GEngine->OnLevelActorOuterChanged().AddUObject(this, &UEditorDataStorageCompatibility::OnActorOuterChanged);
 }
 
 void UEditorDataStorageCompatibility::PostInitialize(UEditorDataStorage* InStorage)
@@ -102,15 +97,6 @@ void UEditorDataStorageCompatibility::Deinitialize()
 
 	FScopedExclusiveLock Lock(EGlobalLockScope::Public);
 	
-	for (TPair<UWorld*, FDelegateHandle>& It : ActorDestroyedDelegateHandles)
-	{
-		It.Key->RemoveOnActorDestroyedHandler(It.Value);
-	}
-
-	GEngine->OnLevelActorOuterChanged().RemoveAll(this);
-
-	FWorldDelegates::OnPreWorldFinishDestroy.Remove(PreWorldFinishDestroyDelegateHandle);
-	FWorldDelegates::OnPostWorldInitialization.Remove(PostWorldInitializationDelegateHandle);
 	FCoreUObjectDelegates::PostReachabilityAnalysis.Remove(PostGcUnreachableAnalysisHandle);
 
 	FCoreUObjectDelegates::OnObjectsReinstanced.Remove(ObjectReinstancedDelegateHandle);
@@ -357,16 +343,6 @@ void UEditorDataStorageCompatibility::CreateStandardArchetypes()
 {
 	// Thread-safe as this is only called from a function that has an exclusive lock.
 	using namespace UE::Editor::DataStorage;
-	StandardActorTable = Storage->RegisterTable(TTypedElementColumnTypeList<
-			FTypedElementUObjectColumn, FTypedElementUObjectIdColumn, FTypedElementClassTypeInfoColumn,
-			FTypedElementLabelColumn, FTypedElementLabelHashColumn, FNameColumn, FTypedElementActorTag,
-			FTypedElementSyncFromWorldTag>(), 
-		FName("Editor_StandardActorTable"));
-
-	StandardActorWithTransformTable = Storage->RegisterTable(StandardActorTable,
-		TTypedElementColumnTypeList<FTypedElementLocalTransformColumn>(),
-		FName("Editor_StandardActorWithTransformTable"));
-
 	StandardUObjectTable = Storage->RegisterTable(TTypedElementColumnTypeList<
 			FTypedElementUObjectColumn, FTypedElementUObjectIdColumn, FTypedElementClassTypeInfoColumn,
 			FNameColumn, FTypedElementSyncFromWorldTag>(), 
@@ -377,7 +353,6 @@ void UEditorDataStorageCompatibility::CreateStandardArchetypes()
 			FTypedElementSyncFromWorldTag>(), 
 		FName("Editor_StandardExternalObjectTable"));
 
-	RegisterTypeTableAssociation(AActor::StaticClass(), StandardActorTable);
 	RegisterTypeTableAssociation(UObject::StaticClass(), StandardUObjectTable);
 }
 
@@ -1170,49 +1145,6 @@ void UEditorDataStorageCompatibility::OnPostGcUnreachableAnalysis()
 	}
 }
 
-void UEditorDataStorageCompatibility::OnPostWorldInitialization(UWorld* World, const UWorld::InitializationValues InitializationValues)
-{
-	using namespace UE::Editor::DataStorage;
-
-	FScopedExclusiveLock Lock(EGlobalLockScope::Public);
-
-	FDelegateHandle Handle = World->AddOnActorDestroyedHandler(
-		FOnActorDestroyed::FDelegate::CreateUObject(this, &UEditorDataStorageCompatibility::OnActorDestroyed));
-	ActorDestroyedDelegateHandles.Add(World, Handle);
-}
-
-void UEditorDataStorageCompatibility::OnPreWorldFinishDestroy(UWorld* World)
-{
-	using namespace UE::Editor::DataStorage;
-
-	FScopedExclusiveLock Lock(EGlobalLockScope::Public);
-
-	FDelegateHandle Handle;
-	if (ActorDestroyedDelegateHandles.RemoveAndCopyValue(World, Handle))
-	{
-		World->RemoveOnActorDestroyedHandler(Handle);
-	}
-}
-
-void UEditorDataStorageCompatibility::OnActorDestroyed(AActor* Actor)
-{
-	// The only function called is already thread safe.
-	RemoveCompatibleObjectExplicit(Actor);
-}
-
-void UEditorDataStorageCompatibility::OnActorOuterChanged(AActor* Actor, UObject* Outer)
-{
-	// We only want to register actors outer'd to a level in TEDS - so if the outer changes we add/remove the object based on that status
-	if (Actor->GetLevel())
-	{
-		AddCompatibleObjectExplicit(Actor);
-	}
-	else
-	{
-		RemoveCompatibleObjectExplicit(Actor);
-	}
-}
-
 SIZE_T GetTypeHash(const UEditorDataStorageCompatibility::FSyncTagInfo& Column)
 {
 	// Thread safe as it only uses local data.
diff --git a/Engine/Plugins/Experimental/EditorDataStorage/Source/TedsCore/Private/TypedElementDatabaseCompatibility.h b/Engine/Plugins/Experimental/EditorDataStorage/Source/TedsCore/Private/TypedElementDatabaseCompatibility.h
--- a/Engine/Plugins/Experimental/EditorDataStorage/Source/TedsCore/Private/TypedElementDatabaseCompatibility.h
+++ b/Engine/Plugins/Experimental/EditorDataStorage/Source/TedsCore/Private/TypedElementDatabaseCompatibility.h
@@ -132,12 +132,6 @@ private:
 
 	UE_API void OnPostGcUnreachableAnalysis();
 
-	UE_API void OnPostWorldInitialization(UWorld* World, const UWorld::InitializationValues InitializationValues);
-	UE_API void OnPreWorldFinishDestroy(UWorld* World);
-
-	UE_API void OnActorDestroyed(AActor* Actor);
-	UE_API void OnActorOuterChanged(AActor* Actor, UObject* Outer);
-
 	struct FPendingTypeInformationUpdate
 	{
 	public:
@@ -206,8 +200,6 @@ private:
 	TArray<TPair<UE::Editor::DataStorage::ObjectAddedCallback, FDelegateHandle>> ObjectAddedCallbackList;
 	TArray<TPair<UE::Editor::DataStorage::ObjectRemovedCallback, FDelegateHandle>> PreObjectRemovedCallbackList;
 
-	UE::Editor::DataStorage::TableHandle StandardActorTable{ UE::Editor::DataStorage::InvalidTableHandle };
-	UE::Editor::DataStorage::TableHandle StandardActorWithTransformTable{ UE::Editor::DataStorage::InvalidTableHandle };
 	UE::Editor::DataStorage::TableHandle StandardUObjectTable{ UE::Editor::DataStorage::InvalidTableHandle };
 	UE::Editor::DataStorage::TableHandle StandardExternalObjectTable{ UE::Editor::DataStorage::InvalidTableHandle };
 	UE::Editor::DataStorage::ICoreProvider* Storage{ nullptr };
@@ -231,15 +223,11 @@ private:
 	using ObjectsNeedingSyncTagsMap = TMap<ObjectsNeedingSyncTagsMapKey, ObjectsNeedingSyncTagsMapValue>;
 	ObjectsNeedingSyncTagsMap ObjectsNeedingSyncTags;
 
-	TMap<UWorld*, FDelegateHandle> ActorDestroyedDelegateHandles;
 	FDelegateHandle PreEditChangePropertyDelegateHandle;
 	FDelegateHandle PostEditChangePropertyDelegateHandle;
 	FDelegateHandle ObjectModifiedDelegateHandle;
-	FDelegateHandle PostWorldInitializationDelegateHandle;
-	FDelegateHandle PreWorldFinishDestroyDelegateHandle;
 	FDelegateHandle ObjectReinstancedDelegateHandle;
 	FDelegateHandle PostGcUnreachableAnalysisHandle;
-	FDelegateHandle ActorOuterChangedDelegateHandle;
 	
 	TSharedPtr<UE::Editor::DataStorage::FEnvironment> Environment;
 	UE::Editor::DataStorage::QueryHandle ClassTypeInfoQuery;
diff --git a/Engine/Plugins/Experimental/EditorDataStorageFeatures/EditorDataStorageFeatures.uplugin b/Engine/Plugins/Experimental/EditorDataStorageFeatures/EditorDataStorageFeatures.uplugin
--- a/Engine/Plugins/Experimental/EditorDataStorageFeatures/EditorDataStorageFeatures.uplugin
+++ b/Engine/Plugins/Experimental/EditorDataStorageFeatures/EditorDataStorageFeatures.uplugin
@@ -15,6 +15,10 @@
 	"IsExperimentalVersion": true,
 	"Installed": false,
 	"Modules": [
+		{
+			"Name": "TedsActorCompatibility",
+			"Type": "Editor"
+		},
 		{
 			"Name": "TedsAlerts",
 			"Type": "Editor"
diff --git a/Engine/Plugins/Experimental/EditorDataStorageFeatures/Source/TedsActorCompatibility/Private/Compatibility/TedsActorIconOverrideQueries.cpp b/Engine/Plugins/Experimental/EditorDataStorageFeatures/Source/TedsActorCompatibility/Private/Compatibility/TedsActorIconOverrideQueries.cpp
new file mode 100644
--- /dev/null
+++ b/Engine/Plugins/Experimental/EditorDataStorageFeatures/Source/TedsActorCompatibility/Private/Compatibility/TedsActorIconOverrideQueries.cpp
@@ -0,0 +1,60 @@
+// Copyright Epic Games, Inc. All Rights Reserved.
+
+#include "Compatibility/TedsActorIconOverrideQueries.h"
+
+#include "Elements/Columns/TypedElementCompatibilityColumns.h"
+#include "Elements/Columns/TypedElementIconOverrideColumns.h"
+#include "Elements/Columns/TypedElementMiscColumns.h"
+#include "Elements/Framework/TypedElementQueryBuilder.h"
+#include "GameFramework/Actor.h"
+
+void UActorIconOverrideDataStorageFactory::RegisterQueries(UE::Editor::DataStorage::ICoreProvider& DataStorage)
+{
+	using namespace UE::Editor::DataStorage::Queries;
+
+	DataStorage.RegisterQuery(
+		Select(
+			TEXT("Add icon override column to actor"),
+			FProcessor(EQueryTickPhase::PrePhysics, DataStorage.GetQueryTickGroupName(EQueryTickGroups::SyncExternalToDataStorage))
+				.SetExecutionMode(EExecutionMode::GameThread),
+			[](IQueryContext& Context, RowHandle Row, const FTypedElementUObjectColumn& ActorColumn)
+			{
+				if (const AActor* Actor = Cast<AActor>(ActorColumn.Object))
+				{
+					if (FName IconName = Actor->GetCustomIconName(); !IconName.IsNone())
+					{
+						Context.AddColumn(Row, FTypedElementIconOverrideColumn{ .IconName = IconName });
+					}
+				}
+			})
+		.Where()
+			.All<FTypedElementSyncFromWorldTag, FTypedElementActorTag>()
+			.None<FTypedElementIconOverrideColumn>()
+		.Compile());
+
+	DataStorage.RegisterQuery(
+		Select(
+			TEXT("Update/remove icon override column to actor"),
+			FProcessor(EQueryTickPhase::PrePhysics, DataStorage.GetQueryTickGroupName(EQueryTickGroups::SyncExternalToDataStorage))
+				.SetExecutionMode(EExecutionMode::GameThread),
+			[]( IQueryContext& Context, 
+				RowHandle Row, 
+				const FTypedElementUObjectColumn& ActorColumn, 
+				FTypedElementIconOverrideColumn& IconColumn)
+			{
+				if (const AActor* Actor = Cast<AActor>(ActorColumn.Object))
+				{
+					if (FName IconName = Actor->GetCustomIconName(); !IconName.IsNone())
+					{
+						IconColumn.IconName = IconName;
+					}
+					else
+					{
+						Context.RemoveColumns<FTypedElementIconOverrideColumn>(Row);
+					}
+				}
+			})
+		.Where()
+			.All<FTypedElementSyncFromWorldTag, FTypedElementActorTag>()
+		.Compile());
+}
diff --git a/Engine/Plugins/Experimental/EditorDataStorageFeatures/Source/TedsActorCompatibility/Private/Compatibility/TedsActorIconOverrideQueries.h b/Engine/Plugins/Experimental/EditorDataStorageFeatures/Source/TedsActorCompatibility/Private/Compatibility/TedsActorIconOverrideQueries.h
new file mode 100644
--- /dev/null
+++ b/Engine/Plugins/Experimental/EditorDataStorageFeatures/Source/TedsActorCompatibility/Private/Compatibility/TedsActorIconOverrideQueries.h
@@ -0,0 +1,19 @@
+// Copyright Epic Games, Inc. All Rights Reserved.
+
+#pragma once
+
+#include "Elements/Interfaces/TypedElementDataStorageFactory.h"
+#include "UObject/ObjectMacros.h"
+
+#include "TedsActorIconOverrideQueries.generated.h"
+
+UCLASS()
+class UActorIconOverrideDataStorageFactory : public UEditorDataStorageFactory
+{
+	GENERATED_BODY()
+
+public:
+	virtual ~UActorIconOverrideDataStorageFactory() override = default;
+
+	virtual void RegisterQueries(UE::Editor::DataStorage::ICoreProvider& DataStorage) override;
+};
diff --git a/Engine/Plugins/Experimental/EditorDataStorageFeatures/Source/TedsActorCompatibility/Private/Compatibility/TedsActorLabelQueries.cpp b/Engine/Plugins/Experimental/EditorDataStorageFeatures/Source/TedsActorCompatibility/Private/Compatibility/TedsActorLabelQueries.cpp
new file mode 100644
--- /dev/null
+++ b/Engine/Plugins/Experimental/EditorDataStorageFeatures/Source/TedsActorCompatibility/Private/Compatibility/TedsActorLabelQueries.cpp
@@ -0,0 +1,97 @@
+// Copyright Epic Games, Inc. All Rights Reserved.
+
+#include "Compatibility/TedsActorLabelQueries.h"
+
+#include "Editor/EditorEngine.h"
+#include "Elements/Columns/TypedElementCompatibilityColumns.h"
+#include "Elements/Columns/TypedElementMiscColumns.h"
+#include "Elements/Columns/TypedElementLabelColumns.h"
+#include "Elements/Framework/TypedElementQueryBuilder.h"
+#include "Elements/Framework/TypedElementRegistry.h"
+#include "Hash/CityHash.h"
+#include "ScopedTransaction.h"
+
+#define LOCTEXT_NAMESPACE "TedsCore"
+
+void UActorLabelDataStorageFactory::RegisterQueries(UE::Editor::DataStorage::ICoreProvider& DataStorage)
+{
+	RegisterActorLabelToColumnQuery(DataStorage);
+	RegisterLabelColumnToActorQuery(DataStorage);
+}
+
+void UActorLabelDataStorageFactory::RegisterActorLabelToColumnQuery(UE::Editor::DataStorage::ICoreProvider& DataStorage) const
+{
+	using namespace UE::Editor::DataStorage::Queries;
+	
+	DataStorage.RegisterQuery(
+		Select(
+			TEXT("Sync actor label to column"),
+			FProcessor(EQueryTickPhase::PrePhysics, DataStorage.GetQueryTickGroupName(EQueryTickGroups::SyncExternalToDataStorage))
+				.SetExecutionMode(EExecutionMode::GameThread),
+			[](const FTypedElementUObjectColumn& Actor, FTypedElementLabelColumn& Label, FTypedElementLabelHashColumn& LabelHash)
+			{
+				if (const AActor* ActorInstance = Cast<AActor>(Actor.Object); ActorInstance != nullptr)
+				{
+					const FString& ActorLabel = ActorInstance->GetActorLabel();
+					uint64 ActorLabelHash = CityHash64(reinterpret_cast<const char*>(*ActorLabel), ActorLabel.Len() * sizeof(**ActorLabel));
+					if (LabelHash.LabelHash != ActorLabelHash)
+					{
+						Label.Label = ActorLabel;
+						LabelHash.LabelHash = ActorLabelHash;
+					}
+				}
+			}
+		)
+		.Where()
+			.All<FTypedElementActorTag, FTypedElementSyncFromWorldTag>()
+		.Compile());
+}
+
+void UActorLabelDataStorageFactory::RegisterLabelColumnToActorQuery(UE::Editor::DataStorage::ICoreProvider& DataStorage) const
+{
+	using namespace UE::Editor::DataStorage::Queries;
+
+	struct FRenameActorCommand
+	{
+		void operator()()
+		{
+			TStrongObjectPtr<AActor> PinnedActor = Actor.Pin();
+			if (PinnedActor)
+			{
+				const FScopedTransaction Transaction(LOCTEXT("RenameActorTransaction", "Rename Actor"));
+				FActorLabelUtilities::RenameExistingActor(PinnedActor.Get(), NewLabel);
+			}
+			
+		}
+		TWeakObjectPtr<AActor> Actor;
+		FString NewLabel;
+	};
+	
+	DataStorage.RegisterQuery(
+		Select(
+			TEXT("Sync label column to actor"),
+			FProcessor(EQueryTickPhase::FrameEnd, DataStorage.GetQueryTickGroupName(EQueryTickGroups::SyncDataStorageToExternal))
+				.SetExecutionMode(EExecutionMode::GameThread),
+			[](IQueryContext& Context, FTypedElementUObjectColumn& Actor, const FTypedElementLabelColumn& Label, const FTypedElementLabelHashColumn& LabelHash)
+			{
+				if (AActor* ActorInstance = Cast<AActor>(Actor.Object); ActorInstance != nullptr)
+				{
+					const FString& ActorLabel = ActorInstance->GetActorLabel(false);
+					uint64 ActorLabelHash = CityHash64(reinterpret_cast<const char*>(*ActorLabel), ActorLabel.Len() * sizeof(**ActorLabel));
+					if (LabelHash.LabelHash != ActorLabelHash)
+					{
+						Context.PushCommand(FRenameActorCommand
+							{
+								.Actor = ActorInstance,
+								.NewLabel = Label.Label
+							});
+					}
+				}
+			}
+		)
+		.Where()
+			.All<FTypedElementActorTag, FTypedElementSyncBackToWorldTag>()
+		.Compile());
+}
+
+#undef LOCTEXT_NAMESPACE
diff --git a/Engine/Plugins/Experimental/EditorDataStorageFeatures/Source/TedsActorCompatibility/Private/Compatibility/TedsActorLabelQueries.h b/Engine/Plugins/Experimental/EditorDataStorageFeatures/Source/TedsActorCompatibility/Private/Compatibility/TedsActorLabelQueries.h
new file mode 100644
--- /dev/null
+++ b/Engine/Plugins/Experimental/EditorDataStorageFeatures/Source/TedsActorCompatibility/Private/Compatibility/TedsActorLabelQueries.h
@@ -0,0 +1,31 @@
+// Copyright Epic Games, Inc. All Rights Reserved.
+
+#pragma once
+
+#include "Elements/Interfaces/TypedElementDataStorageFactory.h"
+#include "UObject/ObjectMacros.h"
+
+#include "TedsActorLabelQueries.generated.h"
+
+
+UCLASS()
+class UActorLabelDataStorageFactory : public UEditorDataStorageFactory
+{
+	GENERATED_BODY()
+
+public:
+	~UActorLabelDataStorageFactory() override = default;
+
+	void RegisterQueries(UE::Editor::DataStorage::ICoreProvider& DataStorage) override;
+
+private:
+	/**
+	 * Takes the label set on an actor and copies it to the Data Storage if they differ.
+	 */
+	void RegisterActorLabelToColumnQuery(UE::Editor::DataStorage::ICoreProvider& DataStorage) const;
+	/**
+	 * Takes the label stored in the Data Storage and copies it to the actor's label if the FTypedElementSyncBackToWorldTag
+	 * has been set and the labels differ.
+	 */
+	void RegisterLabelColumnToActorQuery(UE::Editor::DataStorage::ICoreProvider& DataStorage) const;
+};
\ No newline at end of file
diff --git a/Engine/Plugins/Experimental/EditorDataStorageFeatures/Source/TedsActorCompatibility/Private/Compatibility/TedsActorParentQueries.cpp b/Engine/Plugins/Experimental/EditorDataStorageFeatures/Source/TedsActorCompatibility/Private/Compatibility/TedsActorParentQueries.cpp
new file mode 100644
--- /dev/null
+++ b/Engine/Plugins/Experimental/EditorDataStorageFeatures/Source/TedsActorCompatibility/Private/Compatibility/TedsActorParentQueries.cpp
@@ -0,0 +1,105 @@
+// Copyright Epic Games, Inc. All Rights Reserved.
+
+#include "Compatibility/TedsActorParentQueries.h"
+
+#include "Elements/Columns/TypedElementCompatibilityColumns.h"
+#include "Elements/Columns/TypedElementHiearchyColumns.h"
+#include "Elements/Columns/TypedElementMiscColumns.h"
+#include "Elements/Common/TypedElementMapKey.h"
+#include "Elements/Framework/TypedElementQueryBuilder.h"
+#include "GameFramework/Actor.h"
+
+namespace UE::Editor::DataStorage::Queries::Private
+{
+	static bool bAddParentColumnToActors = false;
+
+	// Cvar to allow the TEDS-Outliner to automatically take over the level editor's 4 Outliners instead of appearing as a separate tab
+	static FAutoConsoleVariableRef CVarUseTEDSOutliner(
+		TEXT("TEDS.AddParentColumnToActors"),
+		bAddParentColumnToActors,
+		TEXT("Mirror parent information for actors to TEDS (only works when set on startup)"));
+};
+
+void UActorParentDataStorageFactory::RegisterQueries(UE::Editor::DataStorage::ICoreProvider& DataStorage)
+{
+	if (UE::Editor::DataStorage::Queries::Private::bAddParentColumnToActors)
+	{
+		RegisterAddParentColumn(DataStorage);
+		RegisterUpdateOrRemoveParentColumn(DataStorage);
+	}
+}
+
+void UActorParentDataStorageFactory::RegisterAddParentColumn(UE::Editor::DataStorage::ICoreProvider& DataStorage) const
+{
+	using namespace UE::Editor::DataStorage;
+	using namespace UE::Editor::DataStorage::Queries;
+
+	DataStorage.RegisterQuery(
+		Select(
+			TEXT("Add parent column to actor"),
+			FProcessor(EQueryTickPhase::PrePhysics, DataStorage.GetQueryTickGroupName(EQueryTickGroups::SyncExternalToDataStorage))
+			.SetExecutionMode(EExecutionMode::GameThread),
+			[](IQueryContext& Context, RowHandle Row, const FTypedElementUObjectColumn& Actor)
+			{
+				if (const AActor* ActorInstance = Cast<AActor>(Actor.Object))
+				{
+					if (AActor* Parent = ActorInstance->GetAttachParentActor())
+					{
+						FMapKeyView IdKey = FMapKeyView(Parent);
+						RowHandle ParentRow = Context.LookupMappedRow(ICompatibilityProvider::ObjectMappingDomain, IdKey);
+						if (Context.IsRowAvailable(ParentRow))
+						{
+							Context.AddColumn(Row, FTableRowParentColumn{ .Parent = ParentRow });
+						}
+						else
+						{
+							Context.AddColumn(Row, FUnresolvedTableRowParentColumn{ .ParentIdKey = FMapKey(Parent) });
+						}
+					}
+				}
+			})
+		.Where()
+		.All<FTypedElementSyncFromWorldTag, FTypedElementActorTag>()
+		.None<FTableRowParentColumn, FUnresolvedTableRowParentColumn>()
+		.Compile());
+}
+
+void UActorParentDataStorageFactory::RegisterUpdateOrRemoveParentColumn(UE::Editor::DataStorage::ICoreProvider& DataStorage) const
+{
+	using namespace UE::Editor::DataStorage::Queries;
+
+	DataStorage.RegisterQuery(
+		Select(
+			TEXT("Sync actor's parent to column"),
+			FProcessor(EQueryTickPhase::PrePhysics, DataStorage.GetQueryTickGroupName(EQueryTickGroups::SyncExternalToDataStorage))
+			.SetExecutionMode(EExecutionMode::GameThread),
+			[](IQueryContext& Context, RowHandle Row, const FTypedElementUObjectColumn& Actor, FTableRowParentColumn& Parent)
+			{
+				if (const AActor* ActorInstance = Cast<AActor>(Actor.Object))
+				{
+					if (AActor* ParentActor = ActorInstance->GetAttachParentActor())
+					{
+						FMapKeyView IdKey = FMapKeyView(ParentActor);
+						RowHandle ParentRow = Context.LookupMappedRow(ICompatibilityProvider::ObjectMappingDomain, IdKey);
+
+						if (Parent.Parent != ParentRow)
+						{
+							if (Context.IsRowAvailable(ParentRow))
+							{
+								Parent.Parent = ParentRow;
+							}
+							else
+							{
+								Context.RemoveColumns<FTableRowParentColumn>(Row);
+								Context.AddColumn(Row, FUnresolvedTableRowParentColumn{ .ParentIdKey = FMapKey(ParentActor) });
+							}
+						}
+						return;
+					}
+				}
+				Context.RemoveColumns<FTableRowParentColumn>(Row);
+			})
+		.Where()
+		.All<FTypedElementActorTag, FTypedElementSyncFromWorldTag>()
+		.Compile());
+}
diff --git a/Engine/Plugins/Experimental/EditorDataStorageFeatures/Source/TedsActorCompatibility/Private/Compatibility/TedsActorParentQueries.h b/Engine/Plugins/Experimental/EditorDataStorageFeatures/Source/TedsActorCompatibility/Private/Compatibility/TedsActorParentQueries.h
new file mode 100644
--- /dev/null
+++ b/Engine/Plugins/Experimental/EditorDataStorageFeatures/Source/TedsActorCompatibility/Private/Compatibility/TedsActorParentQueries.h
@@ -0,0 +1,31 @@
+// Copyright Epic Games, Inc. All Rights Reserved.
+
+#pragma once
+
+#include "Elements/Interfaces/TypedElementDataStorageFactory.h"
+#include "UObject/ObjectMacros.h"
+
+#include "TedsActorParentQueries.generated.h"
+
+UCLASS()
+class UActorParentDataStorageFactory : public UEditorDataStorageFactory
+{
+	GENERATED_BODY()
+
+public:
+	~UActorParentDataStorageFactory() override = default;
+
+	void RegisterQueries(UE::Editor::DataStorage::ICoreProvider& DataStorage) override;
+
+private:
+	/**
+	 * Checks rows with actors that don't have a parent column yet if one needs to be added whenever
+	 * the row is marked for updates.
+	 */
+	void RegisterAddParentColumn(UE::Editor::DataStorage::ICoreProvider& DataStorage) const;
+	/**
+	 * Updates the parent column with the parent from the actor or removes it if there's no parent associated
+	 * with the actor anymore.
+	 */
+	void RegisterUpdateOrRemoveParentColumn(UE::Editor::DataStorage::ICoreProvider& DataStorage) const;
+};
diff --git a/Engine/Plugins/Experimental/EditorDataStorageFeatures/Source/TedsActorCompatibility/Private/Compatibility/TedsActorTransformQueries.cpp b/Engine/Plugins/Experimental/EditorDataStorageFeatures/Source/TedsActorCompatibility/Private/Compatibility/TedsActorTransformQueries.cpp
new file mode 100644
--- /dev/null
+++ b/Engine/Plugins/Experimental/EditorDataStorageFeatures/Source/TedsActorCompatibility/Private/Compatibility/TedsActorTransformQueries.cpp
@@ -0,0 +1,102 @@
+// Copyright Epic Games, Inc. All Rights Reserved.
+
+#include "Compatibility/TedsActorTransformQueries.h"
+
+#include "Elements/Columns/TypedElementCompatibilityColumns.h"
+#include "Elements/Columns/TypedElementMiscColumns.h"
+#include "Elements/Columns/TypedElementTransformColumns.h"
+#include "Elements/Framework/TypedElementQueryBuilder.h"
+#include "GameFramework/Actor.h"
+
+namespace UE::Editor::DataStorage::ActorTransformQueries::Private
+{
+	static bool bSyncActorTransformsToWorld = false;
+	
+	// Cvar to allow the TEDS-Outliner to automatically take over the level editor's 4 Outliners instead of appearing as a separate tab
+	static FAutoConsoleVariableRef CvarSyncActorTransforms(
+		TEXT("TEDS.Feature.SyncActorTransformsToWorld"),
+		bSyncActorTransformsToWorld,
+		TEXT("Sync actor transform changes from TEDS back to the world (only works when set on startup)"));
+}
+
+void UActorTransformDataStorageFactory::RegisterQueries(UE::Editor::DataStorage::ICoreProvider& DataStorage)
+{
+	RegisterActorAddTransformColumn(DataStorage);
+	RegisterActorLocalTransformToColumn(DataStorage);
+	RegisterLocalTransformColumnToActor(DataStorage);
+}
+
+void UActorTransformDataStorageFactory::RegisterActorAddTransformColumn(UE::Editor::DataStorage::ICoreProvider& DataStorage) const
+{
+	using namespace UE::Editor::DataStorage::Queries;
+
+	
+	DataStorage.RegisterQuery(
+		Select(
+			TEXT("Add transform column to actor"),
+			FProcessor(EQueryTickPhase::PrePhysics, DataStorage.GetQueryTickGroupName(EQueryTickGroups::SyncExternalToDataStorage))
+				.SetExecutionMode(EExecutionMode::GameThread),
+			[](IQueryContext& Context, RowHandle Row, const FTypedElementUObjectColumn& Actor)
+			{
+				if (const AActor* ActorInstance = Cast<AActor>(Actor.Object); ActorInstance != nullptr && ActorInstance->GetRootComponent())
+				{
+					Context.AddColumn(Row, FTypedElementLocalTransformColumn{ .Transform = ActorInstance->GetActorTransform() });
+				}
+			})
+		.Where()
+			.All<FTypedElementSyncFromWorldTag, FTypedElementActorTag>()
+			.None<FTypedElementLocalTransformColumn>()
+		.Compile());
+}
+
+void UActorTransformDataStorageFactory::RegisterActorLocalTransformToColumn(UE::Editor::DataStorage::ICoreProvider& DataStorage) const
+{
+	using namespace UE::Editor::DataStorage::Queries;
+
+	
+	DataStorage.RegisterQuery(
+		Select(
+			TEXT("Sync actor transform to column"),
+			FProcessor(EQueryTickPhase::PostPhysics, DataStorage.GetQueryTickGroupName(EQueryTickGroups::SyncExternalToDataStorage))
+				.SetExecutionMode(EExecutionMode::GameThread),
+			[](IQueryContext& Context, RowHandle Row, const FTypedElementUObjectColumn& Actor, FTypedElementLocalTransformColumn& Transform)
+			{
+				if (const AActor* ActorInstance = Cast<AActor>(Actor.Object); ActorInstance != nullptr && ActorInstance->GetRootComponent() != nullptr)
+				{
+					Transform.Transform = ActorInstance->GetActorTransform();
+				}
+				else
+				{
+					Context.RemoveColumns<FTypedElementLocalTransformColumn>(Row);
+				}
+			})
+		.Where()
+			.All<FTypedElementActorTag>()
+			.Any<FTypedElementSyncFromWorldTag, FTypedElementSyncFromWorldInteractiveTag>()
+		.Compile());
+}
+
+void UActorTransformDataStorageFactory::RegisterLocalTransformColumnToActor(UE::Editor::DataStorage::ICoreProvider& DataStorage) const
+{
+	using namespace UE::Editor::DataStorage::Queries;
+
+	if (ActorTransformQueries::Private::bSyncActorTransformsToWorld)
+	{
+		DataStorage.RegisterQuery(
+		Select(
+			TEXT("Sync transform column to actor"),
+			FProcessor(EQueryTickPhase::FrameEnd, DataStorage.GetQueryTickGroupName(EQueryTickGroups::SyncDataStorageToExternal))
+				.SetExecutionMode(EExecutionMode::GameThread),
+			[](FTypedElementUObjectColumn& Actor, const FTypedElementLocalTransformColumn& Transform)
+			{
+				if (AActor* ActorInstance = Cast<AActor>(Actor.Object); ActorInstance != nullptr)
+				{
+					ActorInstance->SetActorTransform(Transform.Transform);
+				}
+			})
+		.Where()
+			.All<FTypedElementActorTag, FTypedElementSyncBackToWorldTag>()
+		.Compile());
+	}
+	
+}
diff --git a/Engine/Plugins/Experimental/EditorDataStorageFeatures/Source/TedsActorCompatibility/Private/Compatibility/TedsActorTransformQueries.h b/Engine/Plugins/Experimental/EditorDataStorageFeatures/Source/TedsActorCompatibility/Private/Compatibility/TedsActorTransformQueries.h
new file mode 100644
--- /dev/null
+++ b/Engine/Plugins/Experimental/EditorDataStorageFeatures/Source/TedsActorCompatibility/Private/Compatibility/TedsActorTransformQueries.h
@@ -0,0 +1,36 @@
+// Copyright Epic Games, Inc. All Rights Reserved.
+
+#pragma once
+
+#include "Elements/Interfaces/TypedElementDataStorageFactory.h"
+#include "UObject/ObjectMacros.h"
+
+#include "TedsActorTransformQueries.generated.h"
+
+UCLASS()
+class UActorTransformDataStorageFactory : public UEditorDataStorageFactory
+{
+	GENERATED_BODY()
+
+public:
+	~UActorTransformDataStorageFactory() override = default;
+
+	void RegisterQueries(UE::Editor::DataStorage::ICoreProvider& DataStorage) override;
+
+private:
+	/**
+	 * Checks actors that don't have a transform column and adds one if an actor has been
+	 * assigned a transform.
+	 */
+	void RegisterActorAddTransformColumn(UE::Editor::DataStorage::ICoreProvider& DataStorage) const;
+	/**
+	 * Takes the transform set on an actor and copies it to the Data Storage or removes the
+	 * transform column if there's not transform available anymore.
+	 */
+	void RegisterActorLocalTransformToColumn(UE::Editor::DataStorage::ICoreProvider& DataStorage) const;
+	/**
+	 * Takes the transform stored in the Data Storage and copies it to the actor's tranform if
+	 * the FTypedElementSyncBackToWorldTag has been set.
+	 */
+	void RegisterLocalTransformColumnToActor(UE::Editor::DataStorage::ICoreProvider& DataStorage) const;
+};
diff --git a/Engine/Plugins/Experimental/EditorDataStorageFeatures/Source/TedsActorCompatibility/Private/Compatibility/TedsActorViewportProcessors.cpp b/Engine/Plugins/Experimental/EditorDataStorageFeatures/Source/TedsActorCompatibility/Private/Compatibility/TedsActorViewportProcessors.cpp
new file mode 100644
--- /dev/null
+++ b/Engine/Plugins/Experimental/EditorDataStorageFeatures/Source/TedsActorCompatibility/Private/Compatibility/TedsActorViewportProcessors.cpp
@@ -0,0 +1,246 @@
+// Copyright Epic Games, Inc. All Rights Reserved.
+
+#include "Compatibility/TedsActorViewportProcessors.h"
+
+#include "Components/PrimitiveComponent.h"
+#include "Elements/Columns/TypedElementCompatibilityColumns.h"
+#include "Elements/Columns/TypedElementViewportColumns.h"
+#include "Elements/Columns/TypedElementMiscColumns.h"
+#include "Elements/Common/EditorDataStorageFeatures.h"
+#include "Elements/Framework/TypedElementQueryBuilder.h"
+#include "Elements/Interfaces/TypedElementDataStorageInterface.h"
+#include "HAL/IConsoleManager.h"
+#include "Elements/Columns/TypedElementSelectionColumns.h"
+#include "ProfilingDebugging/CpuProfilerTrace.h"
+
+namespace UE::Editor::DataStorage::Private
+{
+	FAutoConsoleCommandWithArgsAndOutputDevice SetOutlineColorConsoleCommand(
+		TEXT("TEDS.Debug.SetOutlineColor"),
+		TEXT("Adds an outline color to selected objects."),
+		FConsoleCommandWithArgsAndOutputDeviceDelegate::CreateLambda([](const TArray<FString>& Args, FOutputDevice& Output)
+			{
+				using namespace UE::Editor::DataStorage::Queries;
+
+				TRACE_CPUPROFILER_EVENT_SCOPE(TEDS.Debug.AddOverlayColorToSelectionCommand);
+
+				if (ICoreProvider* DataStorage = GetMutableDataStorageFeature<ICoreProvider>(StorageFeatureName))
+				{
+					static QueryHandle OverlayQuery = InvalidQueryHandle;
+					if (OverlayQuery == InvalidQueryHandle)
+					{
+						OverlayQuery = DataStorage->RegisterQuery(
+							Select()
+							.Where()
+								.All<FTypedElementSelectionColumn>()
+							.Compile());
+					}
+				
+					if (OverlayQuery == InvalidQueryHandle)
+					{
+						return;
+					}
+
+					if (Args.IsEmpty())
+					{
+						Output.Log(TEXT("Provide a color index (0-7) to use as outline"));
+						return;
+					}
+				
+					// Parse the color
+					int32 ColorIndex;
+					LexFromString(ColorIndex, *Args[0]);
+
+					if (!(ColorIndex >= 0 && ColorIndex <= 7))
+					{
+						Output.Log(TEXT("Color index must be in range [0,7]"));
+						return;
+					}
+
+					TArray<RowHandle> RowHandles;
+				
+					DataStorage->RunQuery(OverlayQuery, CreateDirectQueryCallbackBinding(
+						[ColorIndex, &RowHandles](IDirectQueryContext& Context, const RowHandle*)
+					{
+						RowHandles.Append(Context.GetRowHandles());
+					}));
+
+					for (RowHandle Row : RowHandles)
+					{
+						DataStorage->AddColumn(Row, 
+							FTypedElementViewportOutlineColorColumn{ .SelectionOutlineColorIndex = static_cast<uint8>(ColorIndex) });
+						DataStorage->AddColumn<FTypedElementSyncBackToWorldTag>(Row);
+					}
+				}
+			}));
+
+	FAutoConsoleCommandWithArgsAndOutputDevice SetSelectionOverlayColorConsoleCommand(
+		TEXT("TEDS.Debug.SetOverlayColor"),
+		TEXT("Adds an overlay color to selected objects."),
+		FConsoleCommandWithArgsAndOutputDeviceDelegate::CreateLambda([](const TArray<FString>& Args, FOutputDevice& Output)
+			{
+				using namespace UE::Editor::DataStorage::Queries;
+
+				TRACE_CPUPROFILER_EVENT_SCOPE(TEDS.Debug.AddOverlayColorToSelectionCommand);
+
+				if (ICoreProvider* DataStorage = GetMutableDataStorageFeature<ICoreProvider>(StorageFeatureName))
+				{
+					static QueryHandle OverlayQuery = InvalidQueryHandle;
+					if (OverlayQuery == InvalidQueryHandle)
+					{
+						OverlayQuery = DataStorage->RegisterQuery(
+							Select()
+							.Where()
+								.All<FTypedElementSelectionColumn>()
+							.Compile());
+					}
+				
+					if (OverlayQuery == InvalidQueryHandle)
+					{
+						return;
+					}
+
+					if (Args.IsEmpty())
+					{
+						Output.Log(TEXT("Provide a color in hexadecimal format (#RRGGBBAA) to overlay."));
+						return;
+					}
+				
+					// Parse the color
+					FColor Color = FColor::FromHex(Args[0]);
+					Color.A = FMath::Clamp<uint8>(Color.A, 0, 128);
+
+					TArray<RowHandle> RowHandles;
+				
+					DataStorage->RunQuery(OverlayQuery, CreateDirectQueryCallbackBinding(
+						[Color, &RowHandles](IDirectQueryContext& Context, const RowHandle*)
+					{
+						RowHandles.Append(Context.GetRowHandles());
+					}));
+
+					for (RowHandle Row : RowHandles)
+					{
+						DataStorage->RemoveColumn<FTypedElementViewportOverlayColorColumn>(Row);
+						DataStorage->AddColumn(Row, FTypedElementViewportOverlayColorColumn{ .OverlayColor = Color });
+					}
+				}
+			}));
+
+	FAutoConsoleCommandWithArgsAndOutputDevice RemoveSelectionOverlayColorConsoleCommand(
+		TEXT("TEDS.Debug.RemoveOverlayColor"),
+		TEXT("Removes an overlay color to selected objects."),
+		FConsoleCommandWithArgsAndOutputDeviceDelegate::CreateLambda([](const TArray<FString>& Args, FOutputDevice& Output)
+			{
+				using namespace UE::Editor::DataStorage::Queries;
+
+				TRACE_CPUPROFILER_EVENT_SCOPE(TEDS.Debug.AddOverlayColorToSelectionCommand);
+
+				if (ICoreProvider* DataStorage = GetMutableDataStorageFeature<ICoreProvider>(StorageFeatureName))
+				{
+					static QueryHandle OverlayQuery = InvalidQueryHandle;
+					if (OverlayQuery == InvalidQueryHandle)
+					{
+						OverlayQuery = DataStorage->RegisterQuery(
+							Select()
+							.Where()
+								.All<FTypedElementSelectionColumn>()
+							.Compile());
+					}
+				
+					if (OverlayQuery == InvalidQueryHandle)
+					{
+						return;
+					}
+
+					TArray<RowHandle> RowHandles;
+				
+					DataStorage->RunQuery(OverlayQuery, CreateDirectQueryCallbackBinding(
+							[&RowHandles](IDirectQueryContext& Context, const RowHandle*)
+					{
+						RowHandles.Append(Context.GetRowHandles());
+					}));
+
+					for (RowHandle Row : RowHandles)
+					{
+						DataStorage->RemoveColumn<FTypedElementViewportOverlayColorColumn>(Row);
+					}
+				}
+			}));
+} // namespace UE::Editor::DataStorage::Private
+
+void UActorViewportDataStorageFactory::RegisterQueries(UE::Editor::DataStorage::ICoreProvider& DataStorage)
+{
+	RegisterOutlineColorColumnToActor(DataStorage);
+	RegisterOverlayColorColumnToActor(DataStorage);
+}
+
+void UActorViewportDataStorageFactory::RegisterOutlineColorColumnToActor(UE::Editor::DataStorage::ICoreProvider& DataStorage)
+{
+	using namespace UE::Editor::DataStorage::Queries;
+
+	DataStorage.RegisterQuery(
+		Select(
+			TEXT("Sync viewport outline color column to actor"),
+			FProcessor(EQueryTickPhase::DuringPhysics, DataStorage.GetQueryTickGroupName(EQueryTickGroups::SyncDataStorageToExternal))
+				.SetExecutionMode(EExecutionMode::GameThread),
+			[](FTypedElementUObjectColumn& Actor, const FTypedElementViewportOutlineColorColumn& ViewportColor)
+			{
+				if (AActor* ActorInstance = Cast<AActor>(Actor.Object); ActorInstance != nullptr)
+				{
+					bool bIncludeFromChildActors = false;
+					ActorInstance->ForEachComponent<UPrimitiveComponent>(bIncludeFromChildActors, [&ViewportColor](UPrimitiveComponent* PrimitiveComponent)
+					{
+						PrimitiveComponent->SetSelectionOutlineColorIndex(ViewportColor.SelectionOutlineColorIndex);
+					});
+				}
+			})
+		.Where()
+			.All<FTypedElementActorTag, FTypedElementSyncBackToWorldTag>()
+		.Compile());
+}
+
+void UActorViewportDataStorageFactory::RegisterOverlayColorColumnToActor(UE::Editor::DataStorage::ICoreProvider& DataStorage)
+{
+	using namespace UE::Editor::DataStorage::Queries;
+	
+	DataStorage.RegisterQuery(
+		Select(
+			TEXT("Sync viewport overlay color column to actor"),
+			FObserver::OnAdd<FTypedElementViewportOverlayColorColumn>()
+				.SetExecutionMode(EExecutionMode::GameThread),
+			[](FTypedElementUObjectColumn& Actor, const FTypedElementViewportOverlayColorColumn& ViewportColor)
+			{
+				if (AActor* ActorInstance = Cast<AActor>(Actor.Object); ActorInstance != nullptr)
+				{
+					bool bIncludeFromChildActors = true;
+					ActorInstance->ForEachComponent<UPrimitiveComponent>(bIncludeFromChildActors, [&ViewportColor](UPrimitiveComponent* PrimitiveComponent)
+					{
+						PrimitiveComponent->SetOverlayColor(ViewportColor.OverlayColor);
+						PrimitiveComponent->MarkRenderStateDirty();
+					});
+				}
+			})
+		.Where()
+			.All<FTypedElementActorTag>()
+		.Compile());
+
+	DataStorage.RegisterQuery(
+		Select(
+			TEXT("Remove viewport overlay color column from actor"),
+			FObserver::OnRemove<FTypedElementViewportOverlayColorColumn>(),
+			[](FTypedElementUObjectColumn& Actor)
+			{
+				if (AActor* ActorInstance = Cast<AActor>(Actor.Object); ActorInstance != nullptr)
+				{
+					bool bIncludeFromChildActors = true;
+					ActorInstance->ForEachComponent<UPrimitiveComponent>(bIncludeFromChildActors, [](UPrimitiveComponent* PrimitiveComponent)
+					{
+						PrimitiveComponent->RemoveOverlayColor();
+						PrimitiveComponent->MarkRenderStateDirty();
+					});
+				}
+			})
+		.Where()
+			.All<FTypedElementActorTag>()
+		.Compile());
+}
diff --git a/Engine/Plugins/Experimental/EditorDataStorageFeatures/Source/TedsActorCompatibility/Private/Compatibility/TedsActorViewportProcessors.h b/Engine/Plugins/Experimental/EditorDataStorageFeatures/Source/TedsActorCompatibility/Private/Compatibility/TedsActorViewportProcessors.h
new file mode 100644
--- /dev/null
+++ b/Engine/Plugins/Experimental/EditorDataStorageFeatures/Source/TedsActorCompatibility/Private/Compatibility/TedsActorViewportProcessors.h
@@ -0,0 +1,23 @@
+// Copyright Epic Games, Inc. All Rights Reserved.
+
+#pragma once
+
+#include "Elements/Interfaces/TypedElementDataStorageFactory.h"
+#include "UObject/ObjectMacros.h"
+
+#include "TedsActorViewportProcessors.generated.h"
+
+UCLASS()
+class UActorViewportDataStorageFactory : public UEditorDataStorageFactory
+{
+	GENERATED_BODY()
+
+public:
+	~UActorViewportDataStorageFactory() override = default;
+
+	void RegisterQueries(UE::Editor::DataStorage::ICoreProvider& DataStorage) override;
+
+private:
+	void RegisterOutlineColorColumnToActor(UE::Editor::DataStorage::ICoreProvider& DataStorage);
+	void RegisterOverlayColorColumnToActor(UE::Editor::DataStorage::ICoreProvider& DataStorage);
+};
diff --git a/Engine/Plugins/Experimental/EditorDataStorageFeatures/Source/TedsActorCompatibility/Private/Compatibility/TedsActorVisibilityQueries.cpp b/Engine/Plugins/Experimental/EditorDataStorageFeatures/Source/TedsActorCompatibility/Private/Compatibility/TedsActorVisibilityQueries.cpp
new file mode 100644
--- /dev/null
+++ b/Engine/Plugins/Experimental/EditorDataStorageFeatures/Source/TedsActorCompatibility/Private/Compatibility/TedsActorVisibilityQueries.cpp
@@ -0,0 +1,99 @@
+// Copyright Epic Games, Inc. All Rights Reserved.
+#include "Compatibility/TedsActorVisibilityQueries.h"
+
+#include "Elements/Columns/TypedElementCompatibilityColumns.h"
+#include "Elements/Columns/TypedElementMiscColumns.h"
+#include "Elements/Columns/TypedElementVisibilityColumns.h"
+#include "Elements/Common/TypedElementQueryTypes.h"
+#include "Elements/Framework/TypedElementQueryBuilder.h"
+#include "GameFramework/Actor.h"
+
+void UActorVisibilityDataStorageFactory::RegisterQueries(UE::Editor::DataStorage::ICoreProvider& DataStorage)
+{
+	RegisterActorAddVisibilityColumn(DataStorage);
+	RegisterActorVisibilityToColumnQuery(DataStorage);
+	RegisterVisibilityColumnToActorQuery(DataStorage);
+}
+
+void UActorVisibilityDataStorageFactory::RegisterActorAddVisibilityColumn(UE::Editor::DataStorage::ICoreProvider& DataStorage) const
+{
+	using namespace UE::Editor::DataStorage::Queries;
+
+	DataStorage.RegisterQuery(
+		Select(
+			TEXT("Actor Visibility Object to New Column"),
+			FProcessor(EQueryTickPhase::PrePhysics, DataStorage.GetQueryTickGroupName(EQueryTickGroups::SyncExternalToDataStorage)).SetExecutionMode(EExecutionMode::GameThread),
+			[](IQueryContext& Context, RowHandle Row, const FTypedElementUObjectColumn& Object)
+			{
+				if (const AActor* ActorInstance = Cast<AActor>(Object.Object); ActorInstance != nullptr)
+				{
+					Context.AddColumn(Row, FVisibleInEditorColumn{ .bIsVisibleInEditor = !ActorInstance->IsTemporarilyHiddenInEditor() });
+				}
+			})
+		.Where()
+			.All<FTypedElementActorTag, FTypedElementSyncFromWorldTag>()
+			.None<FVisibleInEditorColumn>()
+		.Compile()
+	);
+}
+
+void UActorVisibilityDataStorageFactory::RegisterActorVisibilityToColumnQuery(UE::Editor::DataStorage::ICoreProvider& DataStorage) const
+{
+	using namespace UE::Editor::DataStorage::Queries;
+
+	DataStorage.RegisterQuery(
+		Select(
+			TEXT("Actor Visibility Object to Column"),
+			FProcessor(EQueryTickPhase::PrePhysics, DataStorage.GetQueryTickGroupName(EQueryTickGroups::SyncExternalToDataStorage)).SetExecutionMode(EExecutionMode::GameThread),
+			[](IQueryContext& Context, RowHandle Row, const FTypedElementUObjectColumn& Object, FVisibleInEditorColumn& VisibilityColumn)
+			{
+				if (const AActor* ActorInstance = Cast<AActor>(Object.Object); ActorInstance != nullptr)
+				{
+					VisibilityColumn.bIsVisibleInEditor = !ActorInstance->IsTemporarilyHiddenInEditor();
+				}
+			})
+		.Where()
+			.All<FTypedElementActorTag, FTypedElementSyncFromWorldTag>()
+		.Compile()
+	);
+}
+
+void UActorVisibilityDataStorageFactory::RegisterVisibilityColumnToActorQuery(UE::Editor::DataStorage::ICoreProvider& DataStorage) const
+{
+	using namespace UE::Editor::DataStorage::Queries;
+
+	struct FSetActorVisibilityCommand
+	{
+		void operator()()
+		{
+			TStrongObjectPtr<AActor> PinnedActor = Actor.Pin();
+			if (PinnedActor)
+			{
+				PinnedActor->SetIsTemporarilyHiddenInEditor(!bIsVisible);
+			}
+			
+		}
+		TWeakObjectPtr<AActor> Actor;
+		bool bIsVisible;
+	};
+	
+	DataStorage.RegisterQuery(
+		Select(
+			TEXT("Actor Visibility Column to Object"),
+			FProcessor(EQueryTickPhase::PrePhysics, DataStorage.GetQueryTickGroupName(EQueryTickGroups::SyncDataStorageToExternal)).SetExecutionMode(EExecutionMode::GameThread),
+			[&](IQueryContext& Context, RowHandle Row, FTypedElementUObjectColumn& Object, const FVisibleInEditorColumn& VisibilityColumn)
+			{
+				if (AActor* ActorInstance = Cast<AActor>(Object.Object); ActorInstance != nullptr)
+				{
+					Context.PushCommand(FSetActorVisibilityCommand
+							{
+								.Actor = ActorInstance,
+								.bIsVisible = VisibilityColumn.bIsVisibleInEditor
+							});
+				}
+			})
+		.Where()
+			.All<FTypedElementActorTag, FTypedElementSyncBackToWorldTag>()
+		.Compile()
+	);
+}
diff --git a/Engine/Plugins/Experimental/EditorDataStorageFeatures/Source/TedsActorCompatibility/Private/Compatibility/TedsActorVisibilityQueries.h b/Engine/Plugins/Experimental/EditorDataStorageFeatures/Source/TedsActorCompatibility/Private/Compatibility/TedsActorVisibilityQueries.h
new file mode 100644
--- /dev/null
+++ b/Engine/Plugins/Experimental/EditorDataStorageFeatures/Source/TedsActorCompatibility/Private/Compatibility/TedsActorVisibilityQueries.h
@@ -0,0 +1,38 @@
+// Copyright Epic Games, Inc. All Rights Reserved.
+
+#pragma once
+
+#include "Elements/Interfaces/TypedElementDataStorageFactory.h"
+#include "UObject/ObjectMacros.h"
+
+#include "TedsActorVisibilityQueries.generated.h"
+
+struct ISceneOutlinerTreeItem;
+
+UCLASS()
+class UActorVisibilityDataStorageFactory : public UEditorDataStorageFactory
+{
+	GENERATED_BODY()
+
+public:
+	~UActorVisibilityDataStorageFactory() override = default;
+
+	void RegisterQueries(UE::Editor::DataStorage::ICoreProvider& DataStorage) override;
+
+private:
+
+	/**
+	 * Adds the Visibility columns to new actors that do not have one already.
+	 */
+	void RegisterActorAddVisibilityColumn(UE::Editor::DataStorage::ICoreProvider& DataStorage) const;
+
+	/**
+	 * Takes the visibility set on an actor and copies it to the Data Storage if they differ.
+	 */
+	void RegisterActorVisibilityToColumnQuery(UE::Editor::DataStorage::ICoreProvider& DataStorage) const;
+	/**
+	 * Takes the visibility stored in the Data Storage and copies it to the actor's visibility if the FTypedElementSyncBackToWorldTag
+	 * has been set and the visibility differ.
+	 */
+	void RegisterVisibilityColumnToActorQuery(UE::Editor::DataStorage::ICoreProvider& DataStorage) const;
+};
diff --git a/Engine/Plugins/Experimental/EditorDataStorageFeatures/Source/TedsActorCompatibility/Private/TedsActorCompatibilityFactory.cpp b/Engine/Plugins/Experimental/EditorDataStorageFeatures/Source/TedsActorCompatibility/Private/TedsActorCompatibilityFactory.cpp
new file mode 100644
--- /dev/null
+++ b/Engine/Plugins/Experimental/EditorDataStorageFeatures/Source/TedsActorCompatibility/Private/TedsActorCompatibilityFactory.cpp
@@ -0,0 +1,112 @@
+// Copyright Epic Games, Inc. All Rights Reserved.
+
+#include "TedsActorCompatibilityFactory.h"
+
+#include "Editor.h"
+#include "Engine/World.h"
+#include "Elements/Columns/TypedElementCompatibilityColumns.h"
+#include "Elements/Columns/TypedElementLabelColumns.h"
+#include "Elements/Columns/TypedElementMiscColumns.h"
+#include "Elements/Columns/TypedElementPackageColumns.h"
+#include "Elements/Columns/TypedElementTransformColumns.h"
+#include "Elements/Columns/TypedElementTypeInfoColumns.h"
+#include "Elements/Common/EditorDataStorageFeatures.h"
+#include "GameFramework/Actor.h"
+
+#define LOCTEXT_NAMESPACE "TedsActorFactory"
+
+void UTedsActorCompatibilityFactory::PreRegister(UE::Editor::DataStorage::ICoreProvider& DataStorage)
+{
+	using namespace UE::Editor::DataStorage;
+
+	DataStorageCompatibility = GetMutableDataStorageFeature<ICompatibilityProvider>(CompatibilityFeatureName);
+	check(DataStorageCompatibility);
+
+	// Register the Actor tables in this stage since it is a standard archetype and may be depended upon 
+	// before this factory has had the chance to register its' table.
+	RegisterTable(DataStorage);
+
+	PostWorldInitializationDelegateHandle = FWorldDelegates::OnPostWorldInitialization.AddUObject(this, &UTedsActorCompatibilityFactory::OnPostWorldInitialization);
+	PreWorldFinishDestroyDelegateHandle = FWorldDelegates::OnPreWorldFinishDestroy.AddUObject(this, &UTedsActorCompatibilityFactory::OnPreWorldFinishDestroy);
+
+	ActorOuterChangedDelegateHandle = GEngine->OnLevelActorOuterChanged().AddUObject(this, &UTedsActorCompatibilityFactory::OnActorOuterChanged);
+}
+
+void UTedsActorCompatibilityFactory::RegisterTable(UE::Editor::DataStorage::ICoreProvider& DataStorage)
+{
+	using namespace UE::Editor::DataStorage;
+
+	StandardActorTable = DataStorage.RegisterTable(TTypedElementColumnTypeList<
+		FTypedElementUObjectColumn, FTypedElementUObjectIdColumn, FTypedElementClassTypeInfoColumn,
+		FTypedElementLabelColumn, FTypedElementLabelHashColumn, FNameColumn, FTypedElementActorTag,
+		FTypedElementSyncFromWorldTag>(),
+		FName("Editor_StandardActorTable"));
+
+	StandardActorWithTransformTable = DataStorage.RegisterTable(StandardActorTable,
+		TTypedElementColumnTypeList<FTypedElementLocalTransformColumn>(),
+		FName("Editor_StandardActorWithTransformTable"));
+
+	DataStorageCompatibility->RegisterTypeTableAssociation(AActor::StaticClass(), StandardActorTable);
+}
+
+void UTedsActorCompatibilityFactory::PreShutdown(UE::Editor::DataStorage::ICoreProvider& DataStorage)
+{
+	using namespace UE::Editor::DataStorage;
+
+	for (TPair<UWorld*, FDelegateHandle>& It : ActorDestroyedDelegateHandles)
+	{
+		It.Key->RemoveOnActorDestroyedHandler(It.Value);
+	}
+
+	GEngine->OnLevelActorOuterChanged().RemoveAll(this);
+
+	FWorldDelegates::OnPreWorldFinishDestroy.Remove(PreWorldFinishDestroyDelegateHandle);
+	FWorldDelegates::OnPostWorldInitialization.Remove(PostWorldInitializationDelegateHandle);
+}
+
+void UTedsActorCompatibilityFactory::OnPostWorldInitialization(UWorld* World, const UWorld::InitializationValues InitializationValues)
+{
+	using namespace UE::Editor::DataStorage;
+
+	FDelegateHandle Handle = World->AddOnActorDestroyedHandler(
+		FOnActorDestroyed::FDelegate::CreateUObject(this, &UTedsActorCompatibilityFactory::OnActorDestroyed));
+	ActorDestroyedDelegateHandles.Add(World, Handle);
+}
+
+void UTedsActorCompatibilityFactory::OnPreWorldFinishDestroy(UWorld* World)
+{
+	using namespace UE::Editor::DataStorage;
+
+	FDelegateHandle Handle;
+	if (ActorDestroyedDelegateHandles.RemoveAndCopyValue(World, Handle))
+	{
+		World->RemoveOnActorDestroyedHandler(Handle);
+	}
+}
+
+void UTedsActorCompatibilityFactory::OnActorDestroyed(AActor* Actor)
+{
+	// The only function called is already thread safe.
+	if (DataStorageCompatibility)
+	{
+		DataStorageCompatibility->RemoveCompatibleObjectExplicit(Actor);
+	}
+}
+
+void UTedsActorCompatibilityFactory::OnActorOuterChanged(AActor* Actor, UObject* Outer)
+{
+	// We only want to register actors outer'd to a level in TEDS - so if the outer changes we add/remove the object based on that status
+	if (DataStorageCompatibility)
+	{
+		if (Actor->GetLevel())
+		{
+			DataStorageCompatibility->AddCompatibleObjectExplicit(Actor);
+		}
+		else
+		{
+			DataStorageCompatibility->RemoveCompatibleObjectExplicit(Actor);
+		}
+	}
+}
+
+#undef LOCTEXT_NAMESPACE
\ No newline at end of file
diff --git a/Engine/Plugins/Experimental/EditorDataStorageFeatures/Source/TedsActorCompatibility/Private/TedsActorCompatibilityFactory.h b/Engine/Plugins/Experimental/EditorDataStorageFeatures/Source/TedsActorCompatibility/Private/TedsActorCompatibilityFactory.h
new file mode 100644
--- /dev/null
+++ b/Engine/Plugins/Experimental/EditorDataStorageFeatures/Source/TedsActorCompatibility/Private/TedsActorCompatibilityFactory.h
@@ -0,0 +1,38 @@
+// Copyright Epic Games, Inc. All Rights Reserved.
+
+#pragma once
+
+#include "Elements/Interfaces/TypedElementDataStorageCompatibilityInterface.h"
+#include "Elements/Interfaces/TypedElementDataStorageInterface.h"
+#include "Elements/Interfaces/TypedElementDataStorageFactory.h"
+#include "Engine/World.h"
+
+#include "TedsActorCompatibilityFactory.generated.h"
+
+UCLASS()
+class UTedsActorCompatibilityFactory : public UEditorDataStorageFactory
+{
+	GENERATED_BODY()
+
+public:
+	virtual void PreRegister(UE::Editor::DataStorage::ICoreProvider& DataStorage) override;
+	virtual void PreShutdown(UE::Editor::DataStorage::ICoreProvider& DataStorage) override;
+
+private:
+	void RegisterTable(UE::Editor::DataStorage::ICoreProvider& DataStorage);
+
+	void OnPostWorldInitialization(UWorld* World, const UWorld::InitializationValues InitializationValues);
+	void OnPreWorldFinishDestroy(UWorld* World);
+	void OnActorDestroyed(AActor* Actor);
+	void OnActorOuterChanged(AActor* Actor, UObject* Outer);
+
+	UE::Editor::DataStorage::ICompatibilityProvider* DataStorageCompatibility{nullptr};
+
+	UE::Editor::DataStorage::TableHandle StandardActorTable{ UE::Editor::DataStorage::InvalidTableHandle };
+	UE::Editor::DataStorage::TableHandle StandardActorWithTransformTable{ UE::Editor::DataStorage::InvalidTableHandle };
+
+	TMap<UWorld*, FDelegateHandle> ActorDestroyedDelegateHandles;
+	FDelegateHandle PostWorldInitializationDelegateHandle;
+	FDelegateHandle PreWorldFinishDestroyDelegateHandle;
+	FDelegateHandle ActorOuterChangedDelegateHandle;
+};
\ No newline at end of file
diff --git a/Engine/Plugins/Experimental/EditorDataStorageFeatures/Source/TedsActorCompatibility/Private/TedsActorCompatibilityModule.cpp b/Engine/Plugins/Experimental/EditorDataStorageFeatures/Source/TedsActorCompatibility/Private/TedsActorCompatibilityModule.cpp
new file mode 100644
--- /dev/null
+++ b/Engine/Plugins/Experimental/EditorDataStorageFeatures/Source/TedsActorCompatibility/Private/TedsActorCompatibilityModule.cpp
@@ -0,0 +1,22 @@
+// Copyright Epic Games, Inc. All Rights Reserved.
+
+#include "TedsActorCompatibilityModule.h"
+
+#include "Modules/ModuleManager.h"
+
+namespace UE::Editor::DataStorage
+{
+	void FTedsActorCompatibilityModule::StartupModule()
+	{
+		IModuleInterface::StartupModule();
+
+		FModuleManager::Get().LoadModule(TEXT("TypedElementFramework"));
+	}
+
+	void FTedsActorCompatibilityModule::ShutdownModule()
+	{
+		IModuleInterface::ShutdownModule();
+	}
+} // namespace UE::Editor::DataStorage
+
+IMPLEMENT_MODULE(UE::Editor::DataStorage::FTedsActorCompatibilityModule, TedsActorCompatibility);
\ No newline at end of file
diff --git a/Engine/Plugins/Experimental/EditorDataStorageFeatures/Source/TedsActorCompatibility/Private/TedsActorCompatibilityModule.h b/Engine/Plugins/Experimental/EditorDataStorageFeatures/Source/TedsActorCompatibility/Private/TedsActorCompatibilityModule.h
new file mode 100644
--- /dev/null
+++ b/Engine/Plugins/Experimental/EditorDataStorageFeatures/Source/TedsActorCompatibility/Private/TedsActorCompatibilityModule.h
@@ -0,0 +1,20 @@
+// Copyright Epic Games, Inc. All Rights Reserved.
+
+#pragma once
+
+#include "Modules/ModuleInterface.h"
+
+namespace UE::Editor::DataStorage
+{
+	class FTedsActorCompatibilityModule
+		: public IModuleInterface
+	{
+	public:
+
+		FTedsActorCompatibilityModule() = default;
+
+		// IModuleInterface interface
+		virtual void StartupModule() override;
+		virtual void ShutdownModule() override;
+	};
+} // namespace UE::Editor::DataStorage
diff --git a/Engine/Plugins/Experimental/EditorDataStorageFeatures/Source/TedsActorCompatibility/TedsActorCompatibility.Build.cs b/Engine/Plugins/Experimental/EditorDataStorageFeatures/Source/TedsActorCompatibility/TedsActorCompatibility.Build.cs
new file mode 100644
--- /dev/null
+++ b/Engine/Plugins/Experimental/EditorDataStorageFeatures/Source/TedsActorCompatibility/TedsActorCompatibility.Build.cs
@@ -0,0 +1,33 @@
+// Copyright Epic Games, Inc. All Rights Reserved.
+
+using UnrealBuildTool;
+
+public class TedsActorCompatibility : ModuleRules
+{
+	public TedsActorCompatibility(ReadOnlyTargetRules Target) : base(Target)
+	{
+		// Enable truncation warnings in this module
+		CppCompileWarningSettings.UnsafeTypeCastWarningLevel = WarningLevel.Warning;
+
+		PCHUsage = ModuleRules.PCHUsageMode.UseExplicitOrSharedPCHs;
+
+		if (Target.bBuildEditor)
+		{
+			PublicIncludePaths.AddRange(new string[] { });
+			PrivateIncludePaths.AddRange(new string[] { });
+
+			PublicDependencyModuleNames.AddRange(new string[] { });
+			PrivateDependencyModuleNames.AddRange(
+				new string[] 
+				{
+					"Core",
+					"CoreUObject",
+					"Engine",
+					"TypedElementFramework",
+					"UnrealEd"
+				});
+
+			DynamicallyLoadedModuleNames.AddRange(new string[] { });
+		}
+	}
+}
\ No newline at end of file
diff --git a/Engine/Plugins/Experimental/EditorDataStorageFeatures/Source/TedsOutliner/TedsOutliner.Build.cs b/Engine/Plugins/Experimental/EditorDataStorageFeatures/Source/TedsOutliner/TedsOutliner.Build.cs
--- a/Engine/Plugins/Experimental/EditorDataStorageFeatures/Source/TedsOutliner/TedsOutliner.Build.cs
+++ b/Engine/Plugins/Experimental/EditorDataStorageFeatures/Source/TedsOutliner/TedsOutliner.Build.cs
@@ -40,7 +40,8 @@ public class TedsOutliner : ModuleRules
 					"ToolWidgets",
 					"InputCore",
 					"LevelEditor",
-					"TedsQueryStack"
+					"TedsQueryStack",
+					"TedsActorCompatibility"
 				});
 
 			DynamicallyLoadedModuleNames.AddRange(new string[] {});
