[Mobility] Add the Mobility Column to TEDS so it can be used in the outliner

diff --git a/Engine/Plugins/Experimental/EditorDataStorageFeatures/Source/TedsActorCompatibility/Private/Compatibility/TedsActorMobilityQueries.cpp b/Engine/Plugins/Experimental/EditorDataStorageFeatures/Source/TedsActorCompatibility/Private/Compatibility/TedsActorMobilityQueries.cpp
new file mode 100644
--- /dev/null
+++ b/Engine/Plugins/Experimental/EditorDataStorageFeatures/Source/TedsActorCompatibility/Private/Compatibility/TedsActorMobilityQueries.cpp
@@ -0,0 +1,91 @@
+// Copyright Epic Games, Inc. All Rights Reserved.
+
+#include "Compatibility/TedsActorMobilityQueries.h"
+
+#include "Elements/Columns/TypedElementCompatibilityColumns.h"
+#include "Elements/Columns/TypedElementMiscColumns.h"
+#include "Columns/TedsActorMobilityColumns.h"
+#include "Elements/Framework/TypedElementQueryBuilder.h"
+#include "GameFramework/Actor.h"
+
+void UActorMobilityDataStorageFactory::RegisterQueries(UE::Editor::DataStorage::ICoreProvider& DataStorage)
+{
+	RegisterActorAddMobilityColumn(DataStorage);
+	RegisterActorMobilityToColumnQuery(DataStorage);
+	RegisterMobilityColumnToActorQuery(DataStorage);
+}
+
+void UActorMobilityDataStorageFactory::RegisterActorAddMobilityColumn(UE::Editor::DataStorage::ICoreProvider& DataStorage) const
+{
+	using namespace UE::Editor::DataStorage::Queries;
+
+	DataStorage.RegisterQuery(
+		Select(
+			TEXT("Add Actor Mobility to a New Column"),
+			FProcessor(EQueryTickPhase::PrePhysics, DataStorage.GetQueryTickGroupName(EQueryTickGroups::SyncExternalToDataStorage))
+			.SetExecutionMode(EExecutionMode::GameThread),
+			[](IQueryContext& Context, RowHandle Row, const FTypedElementUObjectColumn& Object)
+			{
+				if (const AActor* ActorInstance = Cast<AActor>(Object.Object); ActorInstance != nullptr)
+				{
+					if (const USceneComponent* ActorInstanceRootComponent = ActorInstance->GetRootComponent())
+					{
+						Context.AddColumn(Row, FTedsActorMobilityColumn{ .Mobility = ActorInstanceRootComponent->GetMobility() });
+					}
+				}
+			})
+		.Where()
+			.All<FTypedElementActorTag, FTypedElementSyncFromWorldTag>()
+			.None<FTedsActorMobilityColumn>()
+		.Compile()
+	);
+}
+
+void UActorMobilityDataStorageFactory::RegisterActorMobilityToColumnQuery(UE::Editor::DataStorage::ICoreProvider& DataStorage) const
+{
+	using namespace UE::Editor::DataStorage::Queries;
+
+	DataStorage.RegisterQuery(
+		Select(
+			TEXT("Sync Actor's Mobility to Column"),
+			FProcessor(EQueryTickPhase::PrePhysics, DataStorage.GetQueryTickGroupName(EQueryTickGroups::SyncExternalToDataStorage))
+			.SetExecutionMode(EExecutionMode::GameThread),
+			[](IQueryContext& Context, RowHandle Row, const FTypedElementUObjectColumn& Object, FTedsActorMobilityColumn& MobilityColumn)
+			{
+				if (const AActor* ActorInstance = Cast<AActor>(Object.Object); ActorInstance != nullptr)
+				{
+					if (const USceneComponent* ActorInstanceRootComponent = ActorInstance->GetRootComponent())
+					{
+						MobilityColumn.Mobility = ActorInstanceRootComponent->GetMobility();
+					}
+				}
+			})
+		.Where()
+			.All<FTypedElementActorTag, FTypedElementSyncFromWorldTag>()
+		.Compile()
+	);
+}
+
+void UActorMobilityDataStorageFactory::RegisterMobilityColumnToActorQuery(UE::Editor::DataStorage::ICoreProvider& DataStorage) const
+{
+	using namespace UE::Editor::DataStorage::Queries;
+
+	DataStorage.RegisterQuery(
+		Select(
+			TEXT("Sync Column to Actor's Mobility"),
+			FProcessor(EQueryTickPhase::PrePhysics, DataStorage.GetQueryTickGroupName(EQueryTickGroups::SyncDataStorageToExternal)).SetExecutionMode(EExecutionMode::GameThread),
+			[&](IQueryContext& Context, RowHandle Row, FTypedElementUObjectColumn& Object, const FTedsActorMobilityColumn& MobilityColumn)
+			{
+				if (AActor* ActorInstance = Cast<AActor>(Object.Object); ActorInstance != nullptr)
+				{
+						if (USceneComponent* ActorInstanceRootComponent = ActorInstance->GetRootComponent())
+						{
+							ActorInstanceRootComponent->SetMobility(MobilityColumn.Mobility);
+						}
+				}
+			})
+		.Where()
+			.All<FTypedElementActorTag, FTypedElementSyncBackToWorldTag>()
+		.Compile()
+	);
+}
diff --git a/Engine/Plugins/Experimental/EditorDataStorageFeatures/Source/TedsActorCompatibility/Private/Compatibility/TedsActorMobilityQueries.h b/Engine/Plugins/Experimental/EditorDataStorageFeatures/Source/TedsActorCompatibility/Private/Compatibility/TedsActorMobilityQueries.h
new file mode 100644
--- /dev/null
+++ b/Engine/Plugins/Experimental/EditorDataStorageFeatures/Source/TedsActorCompatibility/Private/Compatibility/TedsActorMobilityQueries.h
@@ -0,0 +1,34 @@
+// Copyright Epic Games, Inc.All Rights Reserved.
+
+#pragma once
+
+#include "Elements/Interfaces/TypedElementDataStorageFactory.h"
+#include "UObject/ObjectMacros.h"
+
+#include "TedsActorMobilityQueries.generated.h"
+
+UCLASS()
+class UActorMobilityDataStorageFactory : public UEditorDataStorageFactory
+{
+	GENERATED_BODY()
+
+public:
+	~UActorMobilityDataStorageFactory() override = default;
+
+	void RegisterQueries(UE::Editor::DataStorage::ICoreProvider& DataStorage) override;
+private:
+
+	/**
+	 * Adds the mobility column to new actors that do not have one already.
+	 */
+	void RegisterActorAddMobilityColumn(UE::Editor::DataStorage::ICoreProvider& DataStorage) const;
+	/**
+	 * Takes the mobility set on an actor and copies it to the Data Storage if they differ.
+	 */
+	void RegisterActorMobilityToColumnQuery(UE::Editor::DataStorage::ICoreProvider& DataStorage) const;
+	/**
+	 * Takes the mobility stored in the Data Storage and copies it to the actor's mobility if the FTypedElementSyncBackToWorldTag
+	 * has been set and the visibility differ.
+	 */
+	void RegisterMobilityColumnToActorQuery(UE::Editor::DataStorage::ICoreProvider& DataStorage) const;
+};
diff --git a/Engine/Plugins/Experimental/EditorDataStorageFeatures/Source/TedsActorCompatibility/Private/TedsActorCompatibilityFactory.cpp b/Engine/Plugins/Experimental/EditorDataStorageFeatures/Source/TedsActorCompatibility/Private/TedsActorCompatibilityFactory.cpp
--- a/Engine/Plugins/Experimental/EditorDataStorageFeatures/Source/TedsActorCompatibility/Private/TedsActorCompatibilityFactory.cpp
+++ b/Engine/Plugins/Experimental/EditorDataStorageFeatures/Source/TedsActorCompatibility/Private/TedsActorCompatibilityFactory.cpp
@@ -10,6 +10,7 @@
 #include "Elements/Columns/TypedElementPackageColumns.h"
 #include "Elements/Columns/TypedElementTransformColumns.h"
 #include "Elements/Columns/TypedElementTypeInfoColumns.h"
+#include "Columns/TedsActorMobilityColumns.h"
 #include "Elements/Common/EditorDataStorageFeatures.h"
 #include "GameFramework/Actor.h"
 
@@ -38,8 +39,8 @@ void UTedsActorCompatibilityFactory::RegisterTable(UE::Editor::DataStorage::ICor
 
 	StandardActorTable = DataStorage.RegisterTable(TTypedElementColumnTypeList<
 		FTypedElementUObjectColumn, FTypedElementUObjectIdColumn, FTypedElementClassTypeInfoColumn,
-		FTypedElementLabelColumn, FTypedElementLabelHashColumn, FNameColumn, FTypedElementActorTag,
-		FTypedElementSyncFromWorldTag>(),
+		FTypedElementLabelColumn, FTypedElementLabelHashColumn, FNameColumn, FTedsActorMobilityColumn,
+		FTypedElementActorTag, FTypedElementSyncFromWorldTag>(),
 		FName("Editor_StandardActorTable"));
 
 	StandardActorWithTransformTable = DataStorage.RegisterTable(StandardActorTable,
diff --git a/Engine/Plugins/Experimental/EditorDataStorageFeatures/Source/TedsActorCompatibility/Public/Columns/TedsActorMobilityColumns.h b/Engine/Plugins/Experimental/EditorDataStorageFeatures/Source/TedsActorCompatibility/Public/Columns/TedsActorMobilityColumns.h
new file mode 100644
--- /dev/null
+++ b/Engine/Plugins/Experimental/EditorDataStorageFeatures/Source/TedsActorCompatibility/Public/Columns/TedsActorMobilityColumns.h
@@ -0,0 +1,20 @@
+// Copyright Epic Games, Inc. All Rights Reserved.
+
+#pragma once
+
+#include "Elements/Interfaces/TypedElementDataStorageInterface.h"
+#include "Engine/EngineTypes.h"
+
+#include "TedsActorMobilityColumns.generated.h"
+
+/**
+ * Column that stores the mobility of an actor's scene component.
+ */
+USTRUCT(meta = (DisplayName = "Mobility"))
+struct FTedsActorMobilityColumn final : public FEditorDataStorageColumn
+{
+	GENERATED_BODY()
+
+	UPROPERTY(meta = (Sortable))
+	TEnumAsByte<EComponentMobility::Type> Mobility = EComponentMobility::Static;
+};
\ No newline at end of file
diff --git a/Engine/Plugins/Experimental/EditorDataStorageFeatures/Source/TedsOutliner/Private/Compatibility/SceneOutlinerTedsBridge.cpp b/Engine/Plugins/Experimental/EditorDataStorageFeatures/Source/TedsOutliner/Private/Compatibility/SceneOutlinerTedsBridge.cpp
--- a/Engine/Plugins/Experimental/EditorDataStorageFeatures/Source/TedsOutliner/Private/Compatibility/SceneOutlinerTedsBridge.cpp
+++ b/Engine/Plugins/Experimental/EditorDataStorageFeatures/Source/TedsOutliner/Private/Compatibility/SceneOutlinerTedsBridge.cpp
@@ -7,6 +7,7 @@
 #include "Elements/Columns/TypedElementPackageColumns.h"
 #include "Elements/Columns/TypedElementTypeInfoColumns.h"
 #include "Elements/Columns/TypedElementVisibilityColumns.h"
+#include "Columns/TedsActorMobilityColumns.h"
 #include "Elements/Interfaces/TypedElementDataStorageInterface.h"
 #include "Elements/Interfaces/TypedElementDataStorageUiInterface.h"
 #include "SceneOutlinerPublicTypes.h"
@@ -26,6 +27,7 @@ void USceneOutlinerTedsBridgeFactory::PreRegister(UE::Editor::DataStorage::ICore
 	TEDSToOutlinerDefaultColumnMapping.Add(FTypedElementClassTypeInfoColumn::StaticStruct(), FSceneOutlinerBuiltInColumnTypes::ActorInfo());
 	TEDSToOutlinerDefaultColumnMapping.Add(FVisibleInEditorColumn::StaticStruct(), FSceneOutlinerBuiltInColumnTypes::Gutter());
 	TEDSToOutlinerDefaultColumnMapping.Add(FNameColumn::StaticStruct(), FSceneOutlinerBuiltInColumnTypes::IDName());
+	TEDSToOutlinerDefaultColumnMapping.Add(FTedsActorMobilityColumn::StaticStruct(), FSceneOutlinerBuiltInColumnTypes::Mobility());
 }
 
 void USceneOutlinerTedsBridgeFactory::RegisterTables(UE::Editor::DataStorage::ICoreProvider& DataStorage)
diff --git a/Engine/Plugins/Experimental/EditorDataStorageFeatures/Source/TedsOutliner/Private/TedsOutlinerModule.cpp b/Engine/Plugins/Experimental/EditorDataStorageFeatures/Source/TedsOutliner/Private/TedsOutlinerModule.cpp
--- a/Engine/Plugins/Experimental/EditorDataStorageFeatures/Source/TedsOutliner/Private/TedsOutlinerModule.cpp
+++ b/Engine/Plugins/Experimental/EditorDataStorageFeatures/Source/TedsOutliner/Private/TedsOutlinerModule.cpp
@@ -13,6 +13,7 @@
 #include "Elements/Columns/TypedElementPackageColumns.h"
 #include "Elements/Columns/TypedElementTypeInfoColumns.h"
 #include "Elements/Columns/TypedElementVisibilityColumns.h"
+#include "Columns/TedsActorMobilityColumns.h"
 #include "Elements/Common/EditorDataStorageFeatures.h"
 #include "Elements/Interfaces/TypedElementDataStorageInterface.h"
 #include "Modules/ModuleManager.h"
@@ -126,6 +127,7 @@ UE::Editor::DataStorage::FQueryDescription FTedsOutlinerModule::GetLevelEditorTe
 					  FAlertColumn, 
 					  FChildAlertColumn,
 					  FNameColumn,
+					  FTedsActorMobilityColumn,
 					  FVisibleInEditorColumn>()
 		.Compile();
 
@@ -136,6 +138,7 @@ UE::Editor::DataStorage::FQueryDescription FTedsOutlinerModule::GetLevelEditorTe
 					  FTypedElementPackageReference, 
 					  FAlertColumn, 
 					  FNameColumn,
+					  FTedsActorMobilityColumn,
 					  FVisibleInEditorColumn>()
 		.Compile();
 
diff --git a/Engine/Plugins/Experimental/EditorDataStorageFeatures/Source/TedsOutliner/Private/Widgets/OutlinerMobilityWidget.cpp b/Engine/Plugins/Experimental/EditorDataStorageFeatures/Source/TedsOutliner/Private/Widgets/OutlinerMobilityWidget.cpp
new file mode 100644
--- /dev/null
+++ b/Engine/Plugins/Experimental/EditorDataStorageFeatures/Source/TedsOutliner/Private/Widgets/OutlinerMobilityWidget.cpp
@@ -0,0 +1,70 @@
+// Copyright Epic Games, Inc. All Rights Reserved.
+
+#include "Widgets/OutlinerMobilityWidget.h"
+
+#include "SceneOutlinerHelpers.h"
+#include "Elements/Columns/TypedElementCompatibilityColumns.h"
+#include "Elements/Columns/TypedElementMiscColumns.h"
+#include "Columns/TedsActorMobilityColumns.h"
+#include "Elements/Framework/TypedElementAttributeBinding.h"
+#include "Widgets/Text/STextBlock.h"
+
+#define LOCTEXT_NAMESPACE "OutlinerMobilityWidget"
+
+void UOutlinerMobilityWidgetFactory::RegisterWidgetConstructors(UE::Editor::DataStorage::ICoreProvider& DataStorage,
+	UE::Editor::DataStorage::IUiProvider& DataStorageUi) const
+{
+	using namespace UE::Editor::DataStorage::Queries;
+
+	DataStorageUi.RegisterWidgetFactory<FOutlinerMobilityWidgetConstructor>(
+	DataStorageUi.FindPurpose(UE::Editor::DataStorage::IUiProvider::FPurposeInfo("SceneOutliner", "Cell", NAME_None).GeneratePurposeID()),
+		TColumn<FTedsActorMobilityColumn>());
+}
+
+FOutlinerMobilityWidgetConstructor::FOutlinerMobilityWidgetConstructor()
+	: FSimpleWidgetConstructor(StaticStruct())
+{
+}
+
+TSharedPtr<SWidget> FOutlinerMobilityWidgetConstructor::CreateWidget(
+	UE::Editor::DataStorage::ICoreProvider* DataStorage,
+	UE::Editor::DataStorage::IUiProvider* DataStorageUi,
+	UE::Editor::DataStorage::RowHandle TargetRow,
+	UE::Editor::DataStorage::RowHandle WidgetRow,
+	const UE::Editor::DataStorage::FMetaDataView& Arguments)
+{
+	using namespace UE::Editor::DataStorage;
+	if (!DataStorage->IsRowAvailable(TargetRow))
+	{
+		return SNew(STextBlock)
+			.Text(LOCTEXT("MissingRowReferenceColumn", "Unable to retrieve row reference."));
+	}
+
+	FAttributeBinder Binder(TargetRow, DataStorage);
+
+	return SNew(SHorizontalBox)
+			+ SHorizontalBox::Slot()
+			.AutoWidth()
+			.VAlign(VAlign_Center)
+			.Padding(8, 0, 0, 0)
+			[
+				SNew(STextBlock)
+					.ColorAndOpacity(FSlateColor::UseSubduedForeground())
+					.Text(Binder.BindData(&FTedsActorMobilityColumn::Mobility, [](EComponentMobility::Type InMobility)
+						{
+							switch(InMobility)
+							{
+								case EComponentMobility::Movable:
+									return LOCTEXT("OutlinerMobilityTypeMovable", "Movable");
+								case EComponentMobility::Static:
+									return LOCTEXT("OutlinerMobilityTypeStatic", "Static");
+								case EComponentMobility::Stationary:
+									return LOCTEXT("OutlinerMobilityTypeStationary", "Stationary");
+								default:
+									return LOCTEXT("OutlinerMobilityTypeUnknown", "<unknown>");
+							}
+						}))
+			];
+}
+
+#undef LOCTEXT_NAMESPACE
diff --git a/Engine/Plugins/Experimental/EditorDataStorageFeatures/Source/TedsOutliner/Public/Widgets/OutlinerMobilityWidget.h b/Engine/Plugins/Experimental/EditorDataStorageFeatures/Source/TedsOutliner/Public/Widgets/OutlinerMobilityWidget.h
new file mode 100644
--- /dev/null
+++ b/Engine/Plugins/Experimental/EditorDataStorageFeatures/Source/TedsOutliner/Public/Widgets/OutlinerMobilityWidget.h
@@ -0,0 +1,39 @@
+// Copyright Epic Games, Inc. All Rights Reserved.
+
+#pragma once
+
+#include "Elements/Interfaces/TypedElementDataStorageFactory.h"
+#include "Elements/Interfaces/TypedElementDataStorageInterface.h"
+#include "Elements/Interfaces/TypedElementDataStorageUiInterface.h"
+
+#include "OutlinerMobilityWidget.generated.h"
+
+UCLASS()
+class UOutlinerMobilityWidgetFactory : public UEditorDataStorageFactory
+{
+	GENERATED_BODY()
+
+public:
+	~UOutlinerMobilityWidgetFactory() override = default;
+
+	void RegisterWidgetConstructors(UE::Editor::DataStorage::ICoreProvider& DataStorage,
+		UE::Editor::DataStorage::IUiProvider& DataStorageUi) const override;
+};
+
+USTRUCT()
+struct TEDSOUTLINER_API FOutlinerMobilityWidgetConstructor : public FSimpleWidgetConstructor
+{
+	GENERATED_BODY()
+
+public:
+	FOutlinerMobilityWidgetConstructor();
+	~FOutlinerMobilityWidgetConstructor() override = default;
+
+	virtual TSharedPtr<SWidget> CreateWidget(
+		UE::Editor::DataStorage::ICoreProvider* DataStorage,
+		UE::Editor::DataStorage::IUiProvider* DataStorageUi,
+		UE::Editor::DataStorage::RowHandle TargetRow,
+		UE::Editor::DataStorage::RowHandle WidgetRow,
+		const UE::Editor::DataStorage::FMetaDataView& Arguments) override;
+
+};
\ No newline at end of file
