EverythingPicker V.1
- Add SEverythingPicker slate widget
- Add FPickerContext for context customization
- Add PickerContextUtil to guide users towards recommended use cases
-- Will be expanded with more contexts such as instancing as we go

diff --git a/Engine/Plugins/Experimental/EditorDataStorageFeatures/EditorDataStorageFeatures.uplugin b/Engine/Plugins/Experimental/EditorDataStorageFeatures/EditorDataStorageFeatures.uplugin
--- a/Engine/Plugins/Experimental/EditorDataStorageFeatures/EditorDataStorageFeatures.uplugin
+++ b/Engine/Plugins/Experimental/EditorDataStorageFeatures/EditorDataStorageFeatures.uplugin
@@ -58,6 +58,10 @@
 		{
 			"Name": "TedsTableViewer",
 			"Type": "Editor"
+		},
+		{
+			"Name": "TedsEverythingPicker",
+			"Type": "Editor"
 		}
 	],
 	"Plugins": [
diff --git a/Engine/Plugins/Experimental/EditorDataStorageFeatures/Source/TedsEverythingPicker/Private/Context/TedsPickerContextUtil.cpp b/Engine/Plugins/Experimental/EditorDataStorageFeatures/Source/TedsEverythingPicker/Private/Context/TedsPickerContextUtil.cpp
new file mode 100644
--- /dev/null
+++ b/Engine/Plugins/Experimental/EditorDataStorageFeatures/Source/TedsEverythingPicker/Private/Context/TedsPickerContextUtil.cpp
@@ -0,0 +1,47 @@
+// Copyright Epic Games, Inc. All Rights Reserved.
+
+#include "Context/TedsPickerContextUtil.h"
+
+#include "Elements/Common/EditorDataStorageFeatures.h"
+#include "Elements/Framework/TypedElementQueryBuilder.h"
+
+#include "Elements/Columns/TypedElementLabelColumns.h"
+#include "Elements/Columns/TypedElementTypeInfoColumns.h"
+
+#include "Widgets/STedsHierarchyViewer.h"
+#include "Widgets/SEverythingPicker.h"
+
+#include "TedsQueryNode.h"
+#include "TedsRowQueryResultsNode.h"
+
+namespace UE::Editor::DataStorage::Picker
+{
+	void SObjectReferenceContextView::Construct(const FArguments& InArgs)
+	{
+		ICoreProvider* Storage = GetMutableDataStorageFeature<ICoreProvider>(UE::Editor::DataStorage::StorageFeatureName);
+
+		ensure(Storage);
+
+		TSharedRef<QueryStack::FQueryNode> ReferenceQueryNode =
+			MakeShared<QueryStack::FQueryNode>(*Storage, InArgs._Query);
+
+		TSharedRef<QueryStack::FRowQueryResultsNode> ReferenceQueryResultsNode =
+			MakeShared<QueryStack::FRowQueryResultsNode>(*Storage,
+				ReferenceQueryNode, QueryStack::FRowQueryResultsNode::ESyncFlags::RefreshOnUpdate);
+
+		// Using the Teds Outliner Purpose for matching customization for now
+		IUiProvider::FPurposeID PurposeId = IUiProvider::FPurposeInfo("SceneOutliner", "RowLabel", NAME_None).GeneratePurposeID();
+
+		TSharedPtr<SHierarchyViewer> ReferenceView = SNew(SHierarchyViewer)
+			.QueryStack(ReferenceQueryResultsNode)
+			.Columns({ FTypedElementLabelColumn::StaticStruct(), FTypedElementClassTypeInfoColumn::StaticStruct() })
+			.CellWidgetPurpose(PurposeId)
+			.OnSelectionChanged(InArgs._OnSelectionChanged)
+			.DefaultExpansionState(SHierarchyViewer::EExpansionState::Expanded);
+
+		ChildSlot
+			[
+				ReferenceView.ToSharedRef()
+			];
+	}
+}
\ No newline at end of file
diff --git a/Engine/Plugins/Experimental/EditorDataStorageFeatures/Source/TedsEverythingPicker/Private/TedsEverythingPickerModule.cpp b/Engine/Plugins/Experimental/EditorDataStorageFeatures/Source/TedsEverythingPicker/Private/TedsEverythingPickerModule.cpp
new file mode 100644
--- /dev/null
+++ b/Engine/Plugins/Experimental/EditorDataStorageFeatures/Source/TedsEverythingPicker/Private/TedsEverythingPickerModule.cpp
@@ -0,0 +1,20 @@
+// Copyright Epic Games, Inc. All Rights Reserved.
+
+#include "Modules/ModuleManager.h"
+#include "Modules/ModuleInterface.h"
+
+class FTedsEverythingPickerModule : public IModuleInterface
+{
+public:
+
+	virtual void StartupModule() override
+	{
+	}
+
+	virtual void ShutdownModule() override
+	{
+	}
+};
+
+IMPLEMENT_MODULE(FTedsEverythingPickerModule, TedsEverythingPicker);
+
diff --git a/Engine/Plugins/Experimental/EditorDataStorageFeatures/Source/TedsEverythingPicker/Private/Widgets/SEverythingPicker.cpp b/Engine/Plugins/Experimental/EditorDataStorageFeatures/Source/TedsEverythingPicker/Private/Widgets/SEverythingPicker.cpp
new file mode 100644
--- /dev/null
+++ b/Engine/Plugins/Experimental/EditorDataStorageFeatures/Source/TedsEverythingPicker/Private/Widgets/SEverythingPicker.cpp
@@ -0,0 +1,130 @@
+// Copyright Epic Games, Inc. All Rights Reserved.
+
+#include "Widgets/SEverythingPicker.h"
+
+#include "Framework/MultiBox/MultiBoxBuilder.h"
+#include "Widgets/Input/SCheckBox.h"
+#include "Widgets/SBoxPanel.h"
+#include "Widgets/Layout/SWidgetSwitcher.h"
+#define LOCTEXT_NAMESPACE "SEverythingPicker"
+
+namespace UE::Editor::DataStorage::Picker
+{
+	TAutoConsoleVariable<bool> CVarPickerSystemEnabled(
+		TEXT("TEDS.Feature.PickerEnabled"),
+		false, // Defaulting to false until we hit feature parity (search bar, custom filtering, etc.)
+		TEXT("If true, enables uses of TED's EverythingPicker for testing."));
+
+	void SEverythingPicker::Construct(const FArguments& InArgs)
+	{
+		ContextTabs = SNew(SHorizontalBox);
+		ContextTabViews = SNew(SWidgetSwitcher);
+
+		static const FVector2D PickerWindowSize(350.0f, 300.0f);
+
+		FMenuBuilder MenuBuilder(true, nullptr);
+
+		MenuBuilder.BeginSection(NAME_None, LOCTEXT("PickerView", "View"));
+		{
+			TSharedPtr<SWidget> MenuContent;
+
+			for (const FPickerContext::FSlotArguments& ContextArgs : InArgs._Contexts)
+			{
+				AddContext(ContextArgs);
+			}
+
+			MenuContent =
+				SNew(SBox)
+				.WidthOverride(static_cast<float>(PickerWindowSize.X))
+				.HeightOverride(static_cast<float>(PickerWindowSize.Y))
+				[
+					SNew(SVerticalBox)
+					+ SVerticalBox::Slot()
+					.AutoHeight()
+					.Padding(8, 2)
+					[
+						ContextTabs.ToSharedRef()
+					]
+					+ SVerticalBox::Slot()
+					.VAlign(VAlign_Fill)
+					.FillHeight(1.0)
+					.Padding(3)
+					[
+						ContextTabViews.ToSharedRef()
+					]
+				];
+
+			MenuBuilder.AddWidget(MenuContent.ToSharedRef(), FText::GetEmpty(), true);
+		}
+		MenuBuilder.EndSection();
+
+		ChildSlot
+			[
+				MenuBuilder.MakeWidget()
+			];
+	}
+
+	void SEverythingPicker::AddContext(const FPickerContext::FSlotArguments& SlotArgs)
+	{
+		int32 NextContextId = ContextTabViews->GetNumWidgets();
+
+		ContextTabs->AddSlot()
+		.AutoWidth()
+		.Padding(4)
+		[
+			SNew(SCheckBox)
+			.Style(&FAppStyle::GetWidgetStyle<FCheckBoxStyle>("ToggleButtonCheckBox"))
+			.IsChecked(this, &SEverythingPicker::IsContextTabSelected, NextContextId)
+			.OnCheckStateChanged(this, &SEverythingPicker::OnActiveContextChanged, NextContextId)
+			[
+				SNew(STextBlock)
+				.Text(SlotArgs._Label)
+			]
+		];
+
+		TSharedPtr<SWidget> ContextWidget = SlotArgs.GetAttachedWidget();
+		if (ContextWidget)
+		{
+			ContextTabViews->AddSlot()
+			[
+				ContextWidget.ToSharedRef()
+			];
+		}
+		else
+		{
+			ContextTabViews->AddSlot()
+			[
+				SNew(SBox)
+				.HAlign(HAlign_Center)
+				.VAlign(VAlign_Center)
+				[
+					SNew(STextBlock)
+					.Text(LOCTEXT("EmptyContextViewText", "Empty View"))
+				]
+			];
+		}
+	}
+
+	ECheckBoxState SEverythingPicker::IsContextTabSelected(int32 ContextId) const
+	{
+		return (ContextId == ActiveContextId) ? ECheckBoxState::Checked : ECheckBoxState::Unchecked;
+	}
+
+	void SEverythingPicker::OnActiveContextChanged(ECheckBoxState State, int32 ContextId)
+	{
+		if (State != ECheckBoxState::Checked)
+		{
+			return;
+		}
+
+		if (ContextId < 0 || ContextId >= ContextTabViews->GetNumWidgets())
+		{
+			return;
+		}
+
+		ActiveContextId = ContextId;
+		ContextTabViews->SetActiveWidgetIndex(ActiveContextId);
+	}
+}
+
+#undef LOCTEXT_NAMESPACE // "SEverythingPicker"
\ No newline at end of file
diff --git a/Engine/Plugins/Experimental/EditorDataStorageFeatures/Source/TedsEverythingPicker/Public/Context/TedsPickerContext.h b/Engine/Plugins/Experimental/EditorDataStorageFeatures/Source/TedsEverythingPicker/Public/Context/TedsPickerContext.h
new file mode 100644
--- /dev/null
+++ b/Engine/Plugins/Experimental/EditorDataStorageFeatures/Source/TedsEverythingPicker/Public/Context/TedsPickerContext.h
@@ -0,0 +1,22 @@
+// Copyright Epic Games, Inc. All Rights Reserved.
+
+#pragma once
+
+#include "SlotBase.h"
+#include "Templates/SharedPointer.h"
+#include "Widgets/DeclarativeSyntaxSupport.h"
+#include "ITedsTableViewer.h"
+#include "TedsQueryStackInterfaces.h"
+
+namespace UE::Editor::DataStorage
+{
+	class TEDSEVERYTHINGPICKER_API FPickerContext : public TSlotBase<FPickerContext>
+	{
+	public:
+		using TSlotBase<FPickerContext>::TSlotBase;
+
+		SLATE_SLOT_BEGIN_ARGS(FPickerContext, TSlotBase<FPickerContext>)
+			SLATE_ARGUMENT(FText, Label)
+		SLATE_SLOT_END_ARGS()
+	};
+}
\ No newline at end of file
diff --git a/Engine/Plugins/Experimental/EditorDataStorageFeatures/Source/TedsEverythingPicker/Public/Context/TedsPickerContextUtil.h b/Engine/Plugins/Experimental/EditorDataStorageFeatures/Source/TedsEverythingPicker/Public/Context/TedsPickerContextUtil.h
new file mode 100644
--- /dev/null
+++ b/Engine/Plugins/Experimental/EditorDataStorageFeatures/Source/TedsEverythingPicker/Public/Context/TedsPickerContextUtil.h
@@ -0,0 +1,25 @@
+// Copyright Epic Games, Inc. All Rights Reserved.
+
+#pragma once
+
+#include "Context/TedsPickerContext.h"
+#include "Elements/Common/TypedElementCommonTypes.h"
+#include "Elements/Framework/TypedElementQueryBuilder.h"
+
+#include "Widgets/SCompoundWidget.h"
+namespace UE::Editor::DataStorage::Picker
+{
+	DECLARE_DELEGATE_OneParam(FOnSelectionChanged, RowHandle);
+
+	class TEDSEVERYTHINGPICKER_API SObjectReferenceContextView : public SCompoundWidget
+	{
+	public:
+		SLATE_BEGIN_ARGS(SObjectReferenceContextView)
+			{}
+			SLATE_ARGUMENT(FQueryDescription, Query)
+			SLATE_ARGUMENT(FOnSelectionChanged, OnSelectionChanged)
+		SLATE_END_ARGS()
+
+		void Construct(const FArguments& InArgs);
+	};
+}
\ No newline at end of file
diff --git a/Engine/Plugins/Experimental/EditorDataStorageFeatures/Source/TedsEverythingPicker/Public/Widgets/SEverythingPicker.h b/Engine/Plugins/Experimental/EditorDataStorageFeatures/Source/TedsEverythingPicker/Public/Widgets/SEverythingPicker.h
new file mode 100644
--- /dev/null
+++ b/Engine/Plugins/Experimental/EditorDataStorageFeatures/Source/TedsEverythingPicker/Public/Widgets/SEverythingPicker.h
@@ -0,0 +1,46 @@
+// Copyright Epic Games, Inc. All Rights Reserved.
+
+#pragma once
+
+#include "Widgets/SCompoundWidget.h"
+
+#include "Styling/SlateTypes.h"
+
+#include "Context/TedsPickerContext.h"
+
+class SWidgetSwitcher;
+class SHorizontalBox;
+
+namespace UE::Editor::DataStorage::Picker
+{
+	class TEDSEVERYTHINGPICKER_API SEverythingPicker : public SCompoundWidget
+	{
+	public:
+		SLATE_BEGIN_ARGS(SEverythingPicker)
+			{}
+			SLATE_SLOT_ARGUMENT(FPickerContext, Contexts)
+		SLATE_END_ARGS()
+
+		void Construct(const FArguments& InArgs);
+
+		void AddContext(const FPickerContext::FSlotArguments& SlotArgs);
+
+		static FPickerContext::FSlotArguments Context()
+		{
+			return FPickerContext::FSlotArguments(MakeUnique<FPickerContext>());
+		}
+
+		static FPickerContext::FSlotArguments Context(FPickerContext::FSlotArguments&& InContext)
+		{
+			return MoveTemp(InContext);
+		}
+	private:
+
+		ECheckBoxState IsContextTabSelected(int32 ContextId) const;
+		void OnActiveContextChanged(ECheckBoxState State, int32 ContextId);
+
+		int32 ActiveContextId;
+		TSharedPtr<SHorizontalBox> ContextTabs;
+		TSharedPtr<SWidgetSwitcher> ContextTabViews;
+	};
+}
\ No newline at end of file
diff --git a/Engine/Plugins/Experimental/EditorDataStorageFeatures/Source/TedsEverythingPicker/TedsEverythingPicker.Build.cs b/Engine/Plugins/Experimental/EditorDataStorageFeatures/Source/TedsEverythingPicker/TedsEverythingPicker.Build.cs
new file mode 100644
--- /dev/null
+++ b/Engine/Plugins/Experimental/EditorDataStorageFeatures/Source/TedsEverythingPicker/TedsEverythingPicker.Build.cs
@@ -0,0 +1,36 @@
+// Copyright Epic Games, Inc. All Rights Reserved.
+
+using UnrealBuildTool;
+
+public class TedsEverythingPicker : ModuleRules
+{
+	public TedsEverythingPicker(ReadOnlyTargetRules Target) : base(Target)
+	{
+		// Enable truncation warnings in this module
+		CppCompileWarningSettings.UnsafeTypeCastWarningLevel = WarningLevel.Warning;
+
+		if (Target.bBuildEditor)
+		{
+			PublicIncludePaths.AddRange(new string[] {});
+			PrivateIncludePaths.AddRange(new string[] {});
+
+			PublicDependencyModuleNames.AddRange(
+				new string[]
+				{
+					"TypedElementFramework",
+				});
+
+			PrivateDependencyModuleNames.AddRange(
+				new string[]
+				{
+					"Core",
+					"CoreUObject",
+					"Engine",
+					"Slate",
+					"SlateCore",
+					"TedsQueryStack",
+					"TedsTableViewer",
+				});
+		}
+	}
+}
\ No newline at end of file
