Scene Outliner: Fix a bunch of places holding hard references to the Scene Outliner, causing it to not destruct when expected

diff --git a/Engine/Plugins/Experimental/EditorDataStorageFeatures/Source/TedsOutliner/Private/TedsOutlinerFilter.cpp b/Engine/Plugins/Experimental/EditorDataStorageFeatures/Source/TedsOutliner/Private/TedsOutlinerFilter.cpp
--- a/Engine/Plugins/Experimental/EditorDataStorageFeatures/Source/TedsOutliner/Private/TedsOutlinerFilter.cpp
+++ b/Engine/Plugins/Experimental/EditorDataStorageFeatures/Source/TedsOutliner/Private/TedsOutlinerFilter.cpp
@@ -50,13 +50,16 @@ bool FTedsOutlinerFilter::IsInverseFilter() const
 
 void FTedsOutlinerFilter::ActiveStateChanged(bool bActive)
 {
-	if(bActive)
+	if (TSharedPtr<FTedsOutlinerImpl> TedsOutlinerImplPin = TedsOutlinerImpl.Pin())
 	{
-		TedsOutlinerImpl->AddExternalQuery(FilterName, FilterQuery);
-	}
-	else
-	{
-		TedsOutlinerImpl->RemoveExternalQuery(FilterName);
+		if(bActive)
+		{
+			TedsOutlinerImplPin->AddExternalQuery(FilterName, FilterQuery);
+		}
+		else
+		{
+			TedsOutlinerImplPin->RemoveExternalQuery(FilterName);
+		}
 	}
 }
 
@@ -77,11 +80,14 @@ void FTedsOutlinerFilter::LoadSettings(const FString& IniFilename, const FString
 
 bool FTedsOutlinerFilter::PassesFilter(SceneOutliner::FilterBarType InItem) const
 {
-	// If this item is not compatible with the owning Table Viewer - it does not pass any filter queries
-	// If it is compatible, this is simply a dummy filter for the UI while the actual filter is applied through the TEDS query
-	if(TedsOutlinerImpl->IsItemCompatible().IsBound())
+	if (TSharedPtr<FTedsOutlinerImpl> TedsOutlinerImplPin = TedsOutlinerImpl.Pin())
 	{
-		return TedsOutlinerImpl->IsItemCompatible().Execute(InItem);
+		// If this item is not compatible with the owning Table Viewer - it does not pass any filter queries
+		// If it is compatible, this is simply a dummy filter for the UI while the actual filter is applied through the TEDS query
+		if(TedsOutlinerImplPin->IsItemCompatible().IsBound())
+		{
+			return TedsOutlinerImplPin->IsItemCompatible().Execute(InItem);
+		}
 	}
 
 	// The filter is applied through a TEDS query and this is just a dummy to activate it, so we can simply return true otherwise
diff --git a/Engine/Plugins/Experimental/EditorDataStorageFeatures/Source/TedsOutliner/Public/TedsOutlinerFilter.h b/Engine/Plugins/Experimental/EditorDataStorageFeatures/Source/TedsOutliner/Public/TedsOutlinerFilter.h
--- a/Engine/Plugins/Experimental/EditorDataStorageFeatures/Source/TedsOutliner/Public/TedsOutlinerFilter.h
+++ b/Engine/Plugins/Experimental/EditorDataStorageFeatures/Source/TedsOutliner/Public/TedsOutlinerFilter.h
@@ -57,7 +57,7 @@ protected:
 	FName FilterName;
 	FText FilterDisplayName;
 
-	TSharedRef<FTedsOutlinerImpl> TedsOutlinerImpl;
+	TWeakPtr<FTedsOutlinerImpl> TedsOutlinerImpl;
 	const DataStorage::FQueryDescription FilterQuery;
 };
 } // namespace UE::Editor::Outliner
