Teds Outliner: Add sorting support

diff --git a/Engine/Plugins/Experimental/EditorDataStorageFeatures/Source/TedsOutliner/Private/RowMapNode.cpp b/Engine/Plugins/Experimental/EditorDataStorageFeatures/Source/TedsOutliner/Private/RowMapNode.cpp
new file mode 100644
--- /dev/null
+++ b/Engine/Plugins/Experimental/EditorDataStorageFeatures/Source/TedsOutliner/Private/RowMapNode.cpp
@@ -0,0 +1,60 @@
+// Copyright Epic Games, Inc. All Rights Reserved.
+
+#include "RowMapNode.h"
+
+namespace UE::Editor::Outliner
+{
+	FRowMapNode::FRowMapNode(TSharedPtr<IRowNode> InParent)
+		: Parent(MoveTemp(InParent))
+		, CachedParentRevision(Parent->GetRevision())
+	{
+		RefreshMap();
+	}
+
+	DataStorage::FRowHandleArrayView FRowMapNode::GetRows() const
+	{
+		return Parent->GetRows();
+	}
+
+	DataStorage::FRowHandleArray& FRowMapNode::GetMutableRows()
+	{
+		return Parent->GetMutableRows();
+	}
+
+	DataStorage::QueryStack::INode::RevisionId FRowMapNode::GetRevision() const
+	{
+		return Revision;
+	}
+
+	void FRowMapNode::Update()
+	{
+		Parent->Update();
+		if (Parent->GetRevision() != CachedParentRevision)
+		{
+			RefreshMap();
+			CachedParentRevision = Parent->GetRevision();
+			Revision++;
+		}
+	}
+
+	int32 FRowMapNode::GetRowIndex(DataStorage::RowHandle Row) const
+	{
+		if (const int32* Index = RowMap.Find(Row))
+		{
+			return *Index;
+		}
+
+		return INDEX_NONE;
+	}
+
+	void FRowMapNode::RefreshMap()
+	{
+		RowMap.Empty();
+
+		DataStorage::FRowHandleArrayView ParentRows = Parent->GetRows();
+		for (int32 Index = 0; Index < ParentRows.Num(); ++Index)
+		{
+			RowMap.Emplace(ParentRows[Index], Index);
+		}
+	}
+}
diff --git a/Engine/Plugins/Experimental/EditorDataStorageFeatures/Source/TedsOutliner/Private/RowMapNode.h b/Engine/Plugins/Experimental/EditorDataStorageFeatures/Source/TedsOutliner/Private/RowMapNode.h
new file mode 100644
--- /dev/null
+++ b/Engine/Plugins/Experimental/EditorDataStorageFeatures/Source/TedsOutliner/Private/RowMapNode.h
@@ -0,0 +1,38 @@
+// Copyright Epic Games, Inc. All Rights Reserved.
+
+#pragma once
+
+#include "TedsQueryStackInterfaces.h"
+
+namespace UE::Editor::Outliner
+{
+	/*
+	 * A specialized node for the Teds Outliner that stores a map from row handle to index in the parent node.
+	 * This is currently required to integrate sorting with the Teds Outliner without a major refactor of the base Outliner code.
+	 */
+	class FRowMapNode : public DataStorage::QueryStack::IRowNode
+	{
+	public:
+		FRowMapNode(TSharedPtr<IRowNode> InParent);
+		
+		virtual ~FRowMapNode() override = default;
+		virtual DataStorage::FRowHandleArrayView GetRows() const override;
+		virtual DataStorage::FRowHandleArray& GetMutableRows() override;
+		virtual RevisionId GetRevision() const override;
+		virtual void Update() override;
+
+		// Look up the index of the given row in the parent node. INDEX_NONE if not found
+		int32 GetRowIndex(DataStorage::RowHandle Row) const;
+		
+	protected:
+		void RefreshMap();
+		
+	protected:
+
+		// Map of row handle -> index in the parent array
+		TMap<DataStorage::RowHandle, int32> RowMap;
+		TSharedPtr<IRowNode> Parent;
+		RevisionId CachedParentRevision;
+		RevisionId Revision = 0;
+	};
+}
\ No newline at end of file
diff --git a/Engine/Plugins/Experimental/EditorDataStorageFeatures/Source/TedsOutliner/Private/TedsOutlinerColumn.cpp b/Engine/Plugins/Experimental/EditorDataStorageFeatures/Source/TedsOutliner/Private/TedsOutlinerColumn.cpp
--- a/Engine/Plugins/Experimental/EditorDataStorageFeatures/Source/TedsOutliner/Private/TedsOutlinerColumn.cpp
+++ b/Engine/Plugins/Experimental/EditorDataStorageFeatures/Source/TedsOutliner/Private/TedsOutlinerColumn.cpp
@@ -11,40 +11,43 @@
 
 namespace UE::Editor::Outliner
 {
-	FTedsOutlinerUiColumn::FTedsOutlinerUiColumn(
-		DataStorage::FQueryDescription& InQuery,
-		DataStorage::ICoreProvider& InStorage, 
-		DataStorage::IUiProvider& InStorageUi, 
-		DataStorage::ICompatibilityProvider& InStorageCompatibility,
-		FName InNameId, 
-		TArray<TWeakObjectPtr<const UScriptStruct>> InColumnTypes,
-		TSharedPtr<FTypedElementWidgetConstructor> InHeaderWidgetConstructor,
-		TSharedPtr<FTypedElementWidgetConstructor> InCellWidgetConstructor,
-		FName InFallbackColumnName,
-		TWeakPtr<ISceneOutliner> InOwningOutliner,
-		const FTreeItemIDDealiaser& InDealiaser)
-		: Storage(InStorage)
+	FTedsOutlinerUiColumnInitParams::FTedsOutlinerUiColumnInitParams(DataStorage::FQueryDescription& InQuery, DataStorage::ICoreProvider& InStorage,
+		DataStorage::IUiProvider& InStorageUi, DataStorage::ICompatibilityProvider& InStorageCompatibility)
+		: Query(InQuery)
+		, Storage(InStorage)
 		, StorageUi(InStorageUi)
 		, StorageCompatibility(InStorageCompatibility)
-		, QueryDescription(InQuery)
-		, NameId(InNameId)
-		, OwningOutliner(InOwningOutliner)
-		, Dealiaser(InDealiaser)
+	{
+	}
+
+	FTedsOutlinerUiColumn::FTedsOutlinerUiColumn(FTedsOutlinerUiColumnInitParams& InitParams)
+		: Storage(InitParams.Storage)
+		, StorageUi(InitParams.StorageUi)
+		, StorageCompatibility(InitParams.StorageCompatibility)
+		, QueryDescription(InitParams.Query)
+		, NameId(InitParams.NameId)
+		, FallbackColumn(InitParams.FallbackColumn)
+		, OwningOutliner(InitParams.OwningOutliner)
+		, TedsOutlinerImpl(InitParams.TedsOutlinerImpl)
+		, Dealiaser(InitParams.Dealiaser)
+		, bHybridMode(InitParams.bHybridMode)
 	{
 		MetaData.AddOrSetMutableData(TEXT("Name"), NameId.ToString());
 
 		using namespace UE::Editor::DataStorage;
 
-		TableViewerColumnImpl = MakeUnique<FTedsTableViewerColumn>(NameId, InCellWidgetConstructor, InColumnTypes,
-			InHeaderWidgetConstructor, FComboMetaDataView(FGenericMetaDataView(MetaData)).Next(FQueryMetaDataView(QueryDescription)));
+		TableViewerColumnImpl = MakeUnique<FTedsTableViewerColumn>(NameId, InitParams.CellWidgetConstructor, InitParams.ColumnTypes,
+			InitParams.HeaderWidgetConstructor, FComboMetaDataView(FGenericMetaDataView(MetaData)).Next(FQueryMetaDataView(QueryDescription)));
 
 		TableViewerColumnImpl->SetIsRowVisibleDelegate(
 			FTedsTableViewerColumn::FIsRowVisible::CreateRaw(this, &FTedsOutlinerUiColumn::IsRowVisible)
 		);
-		
-		// Try to find a fallback column from the regular item, for handling cases like folders which are not in TEDS but want to use TEDS columns
-		FSceneOutlinerModule& SceneOutlinerModule = FModuleManager::LoadModuleChecked<FSceneOutlinerModule>("SceneOutliner");
-		FallbackColumn = SceneOutlinerModule.FactoryColumn(InFallbackColumnName, *OwningOutliner.Pin());
+
+		// TEDS sorting behaves unexpectedly in hybrid mode, so we only allow fallback sorting there
+		if (!bHybridMode)
+		{
+			TableViewerColumnImpl->SetSortDelegates(InitParams.SortDelegates);
+		}
 	};
 		
 	FName FTedsOutlinerUiColumn::GetColumnID()
@@ -85,7 +88,11 @@ namespace UE::Editor::Outliner
 	// TODO: Sorting is currently handled through the fallback column if it exists because we have no way to sort columns through TEDS
 	void FTedsOutlinerUiColumn::SortItems(TArray<FSceneOutlinerTreeItemPtr>& RootItems, const EColumnSortMode::Type SortMode) const
 	{
-		if (FallbackColumn)
+		if (TSharedPtr<FTedsOutlinerImpl> TedsOutlinerImplPin = TedsOutlinerImpl.Pin(); TedsOutlinerImplPin && !bHybridMode)
+		{
+			TedsOutlinerImplPin->SortItems(RootItems, SortMode);
+		}
+		else if (FallbackColumn)
 		{
 			FallbackColumn->SortItems(RootItems, SortMode);
 		}
@@ -93,9 +100,27 @@ namespace UE::Editor::Outliner
 
 	bool FTedsOutlinerUiColumn::SupportsSorting() const 
 	{
+		if (!bHybridMode && TableViewerColumnImpl->HasSortInfo())
+		{
+			return true;
+		}
+		
 		return FallbackColumn ? FallbackColumn->SupportsSorting() : false;
 	}
 
+	void FTedsOutlinerUiColumn::OnSortRequested(const EColumnSortPriority::Type SortPriority, const EColumnSortMode::Type InSortMode)
+	{
+		if (!bHybridMode)
+		{
+			TableViewerColumnImpl->OnSortCallback(SortPriority, TableViewerColumnImpl->GetColumnName(), InSortMode);
+		}
+	}
+
+	bool FTedsOutlinerUiColumn::IsSortReady()
+	{
+		return !bHybridMode && TedsOutlinerImpl.IsValid() ? !TedsOutlinerImpl.Pin()->IsSorting() : true;
+	}
+
 	void FTedsOutlinerUiColumn::SetHighlightText(SWidget& Widget)
 	{
 		TSharedPtr<ISceneOutliner> OutlinerPinned = OwningOutliner.Pin();
diff --git a/Engine/Plugins/Experimental/EditorDataStorageFeatures/Source/TedsOutliner/Private/TedsOutlinerColumn.h b/Engine/Plugins/Experimental/EditorDataStorageFeatures/Source/TedsOutliner/Private/TedsOutlinerColumn.h
--- a/Engine/Plugins/Experimental/EditorDataStorageFeatures/Source/TedsOutliner/Private/TedsOutlinerColumn.h
+++ b/Engine/Plugins/Experimental/EditorDataStorageFeatures/Source/TedsOutliner/Private/TedsOutlinerColumn.h
@@ -10,22 +10,47 @@
 
 namespace UE::Editor::Outliner
 {
+	class FTedsOutlinerImpl;
+
+	struct FTedsOutlinerUiColumnInitParams
+	{
+		FTedsOutlinerUiColumnInitParams(DataStorage::FQueryDescription& InQuery, DataStorage::ICoreProvider& InStorage,
+			DataStorage::IUiProvider& InStorageUi, DataStorage::ICompatibilityProvider& InStorageCompatibility);
+		
+		// The query used for this column
+		DataStorage::FQueryDescription& Query;
+
+		// Data storage interfaces
+		DataStorage::ICoreProvider& Storage;
+		DataStorage::IUiProvider& StorageUi;
+		DataStorage::ICompatibilityProvider& StorageCompatibility;
+
+		// Unique ID
+		FName NameId;
+
+		// TEDS UI info
+		TArray<TWeakObjectPtr<const UScriptStruct>> ColumnTypes;
+		TSharedPtr<FTypedElementWidgetConstructor> HeaderWidgetConstructor;
+		TSharedPtr<FTypedElementWidgetConstructor> CellWidgetConstructor;
+
+		// Fallbacks to work with legacy columns
+		TSharedPtr<ISceneOutlinerColumn> FallbackColumn;
+		FTreeItemIDDealiaser Dealiaser;
+		bool bHybridMode = false;
+
+		// Sorting info
+		DataStorage::FTedsTableViewerColumn::FSortHandler SortDelegates;
+
+		// Outliner refernces
+		TWeakPtr<ISceneOutliner> OwningOutliner;
+		TWeakPtr<FTedsOutlinerImpl> TedsOutlinerImpl;
+	};
+
 	class FTedsOutlinerUiColumn : public ISceneOutlinerColumn
 	{
 	public:
 		
-		FTedsOutlinerUiColumn(
-			DataStorage::FQueryDescription& InQuery,
-			DataStorage::ICoreProvider& InStorage, 
-			DataStorage::IUiProvider& InStorageUi, 
-			DataStorage::ICompatibilityProvider& InStorageCompatibility,
-			FName InNameId, 
-			TArray<TWeakObjectPtr<const UScriptStruct>> InColumnTypes,
-			TSharedPtr<FTypedElementWidgetConstructor> InHeaderWidgetConstructor,
-			TSharedPtr<FTypedElementWidgetConstructor> InCellWidgetConstructor,
-			FName InFallbackColumnName,
-			TWeakPtr<ISceneOutliner> InOwningOutliner,
-			const FTreeItemIDDealiaser& InDealiaser);
+		FTedsOutlinerUiColumn(FTedsOutlinerUiColumnInitParams& InitParams);
 		
 		FName GetColumnID() override;
 
@@ -33,6 +58,8 @@ namespace UE::Editor::Outliner
 		virtual SHeaderRow::FColumn::FArguments ConstructHeaderRowColumn() override;
 		virtual void SortItems(TArray<FSceneOutlinerTreeItemPtr>& RootItems, const EColumnSortMode::Type SortMode) const override;
 		virtual bool SupportsSorting() const override;
+		virtual void OnSortRequested(const EColumnSortPriority::Type SortPriority, const EColumnSortMode::Type InSortMode) override;
+		virtual bool IsSortReady() override;
 		const TSharedRef<SWidget> ConstructRowWidget(FSceneOutlinerTreeItemRef TreeItem, const STableRow<FSceneOutlinerTreeItemPtr>& Row) override;
 		virtual void PopulateSearchStrings( const ISceneOutlinerTreeItem& Item, TArray< FString >& OutSearchStrings ) const override;
 		
@@ -52,7 +79,9 @@ namespace UE::Editor::Outliner
 		FName NameId;
 		TSharedPtr<ISceneOutlinerColumn> FallbackColumn;
 		TWeakPtr<ISceneOutliner> OwningOutliner;
+		TWeakPtr<FTedsOutlinerImpl> TedsOutlinerImpl;
 		FTreeItemIDDealiaser Dealiaser;
+		bool bHybridMode;
 	};
 
 }
diff --git a/Engine/Plugins/Experimental/EditorDataStorageFeatures/Source/TedsOutliner/Private/TedsOutlinerImpl.cpp b/Engine/Plugins/Experimental/EditorDataStorageFeatures/Source/TedsOutliner/Private/TedsOutlinerImpl.cpp
--- a/Engine/Plugins/Experimental/EditorDataStorageFeatures/Source/TedsOutliner/Private/TedsOutlinerImpl.cpp
+++ b/Engine/Plugins/Experimental/EditorDataStorageFeatures/Source/TedsOutliner/Private/TedsOutlinerImpl.cpp
@@ -15,6 +15,7 @@
 #include "Elements/Framework/TypedElementQueryBuilder.h"
 #include "Columns/SlateDelegateColumns.h"
 #include "Compatibility/SceneOutlinerRowHandleColumn.h"
+#include "RowMapNode.h"
 #include "TedsOutlinerFilter.h"
 #include "TedsOutlinerHelpers.h"
 #include "TedsOutlinerItem.h"
@@ -24,7 +25,9 @@
 #include "TedsRowHandleSortNode.h"
 #include "TedsRowMergeNode.h"
 #include "TedsRowMonitorNode.h"
+#include "TedsRowOrderInversionNode.h"
 #include "TedsRowQueryResultsNode.h"
+#include "TedsRowSortNode.h"
 #include "Columns/TedsOutlinerColumns.h"
 #include "Compatibility/SceneOutlinerTedsBridge.h"
 #include "Elements/Columns/TypedElementCompatibilityColumns.h"
@@ -192,7 +195,7 @@ FTedsOutlinerParams::FTedsOutlinerParams(SSceneOutliner* InSceneOutliner)
 			"SceneOutliner", "RowLabel", NAME_None).GeneratePurposeID());
 }
 	
-FTedsOutlinerImpl::FTedsOutlinerImpl(const FTedsOutlinerParams& InParams, ISceneOutlinerMode* InMode)
+FTedsOutlinerImpl::FTedsOutlinerImpl(const FTedsOutlinerParams& InParams, ISceneOutlinerMode* InMode, bool bInHybridMode)
 	: CreationParams(InParams)
 	, CellWidgetPurpose(InParams.CellWidgetPurpose)
 	, HeaderWidgetPurpose(InParams.HeaderWidgetPurpose)
@@ -203,6 +206,7 @@ FTedsOutlinerImpl::FTedsOutlinerImpl(const FTedsOutlinerParams& InParams, IScene
 	, bForceShowParents(InParams.bForceShowParents)
 	, SceneOutlinerMode(InMode)
 	, SceneOutliner(InParams.SceneOutliner)
+	, bHybridMode(bInHybridMode)
 {
 	// Initialize the TEDS constructs
 	using namespace UE::Editor::DataStorage;
@@ -506,15 +510,18 @@ void FTedsOutlinerImpl::RegenerateColumns()
 			 * replace it with the TEDS column. But also allow the TEDS-Outliner column to fallback to it for
 			 * data not in TEDS yet.
 			 */
-			FName FallbackColumn = Helpers::FindOutlinerColumnFromTedsColumns(Storage, ColumnTypes);
-			const FSceneOutlinerColumnInfo* FallbackColumnInfo = SceneOutliner->GetSharedData().ColumnMap.Find(FallbackColumn);
+			FName FallbackColumnName = Helpers::FindOutlinerColumnFromTedsColumns(Storage, ColumnTypes);
+			const FSceneOutlinerColumnInfo* FallbackColumnInfo = SceneOutliner->GetSharedData().ColumnMap.Find(FallbackColumnName);
 			uint8 ColumnPriority = FallbackColumnInfo
 				? FallbackColumnInfo->PriorityIndex
 				: static_cast<uint8>(
 					FMath::Clamp(DefaultPriorityIndex + static_cast<uint8>(
 						FMath::Clamp(IndexOffset, 0, 255)), 0, 255));
 
-			SceneOutliner->RemoveColumn(FallbackColumn);
+			// Grab a reference to the fallback column and remove it from the actual outliner
+			const TMap<FName, TSharedPtr<ISceneOutlinerColumn>>& Columns = SceneOutliner->GetColumns();
+			const TSharedPtr<ISceneOutlinerColumn> FallbackColumnPtr = Columns.FindRef(FallbackColumnName);
+			SceneOutliner->RemoveColumn(FallbackColumnName);
 
 			FName NameId = UE::Editor::DataStorage::TableViewerUtils::FindLongestMatchingName(ColumnTypes, IndexOffset);
 			FText DisplayName = CellConstructor.Get()->CreateWidgetDisplayNameText(Storage);
@@ -524,15 +531,30 @@ void FTedsOutlinerImpl::RegenerateColumns()
 					ESceneOutlinerColumnVisibility::Visible, 
 					ColumnPriority,
 					FCreateSceneOutlinerColumn::CreateLambda(
-						[this, ColumnQueryColumn, TedsOutlinerColumn, MetaDataView, NameId, &ColumnTypes, CellConstructor, FallbackColumn, HeaderPurposeRow, Dealiaser](ISceneOutliner&)
+						[this, ColumnQueryColumn, TedsOutlinerColumn, MetaDataView, NameId, &ColumnTypes, CellConstructor, FallbackColumnPtr, HeaderPurposeRow, Dealiaser](ISceneOutliner&)
 						{
 							TSharedPtr<FTypedElementWidgetConstructor> HeaderConstructor = 
 								UE::Editor::DataStorage::TableViewerUtils::CreateHeaderWidgetConstructor(*StorageUi, MetaDataView, ColumnTypes, HeaderPurposeRow);
+
+							DataStorage::FTedsTableViewerColumn::FSortHandler SortDelegates
+							{
+								.OnSortHandler = DataStorage::FTedsTableViewerColumn::FOnSort::CreateSP(this, &FTedsOutlinerImpl::OnSort),
+								.IsSortingHandler = DataStorage::FTedsTableViewerColumn::FIsSorting::CreateSP(this, &FTedsOutlinerImpl::IsSorting)
+							};
+
+							FTedsOutlinerUiColumnInitParams Params(ColumnQueryColumn->ColumnQueryDescription, *Storage, *StorageUi, *StorageCompatibility);
+							Params.NameId = NameId;
+							Params.ColumnTypes = TArray<TWeakObjectPtr<const UScriptStruct>>(ColumnTypes.GetData(), ColumnTypes.Num());
+							Params.HeaderWidgetConstructor = MoveTemp(HeaderConstructor);
+							Params.CellWidgetConstructor = CellConstructor;
+							Params.FallbackColumn = FallbackColumnPtr;
+							Params.OwningOutliner = TedsOutlinerColumn->Outliner;
+							Params.TedsOutlinerImpl = AsWeak();
+							Params.SortDelegates = MoveTemp(SortDelegates);
+							Params.Dealiaser = Dealiaser;
+							Params.bHybridMode = bHybridMode;
 							
-							return MakeShared<FTedsOutlinerUiColumn>(
-								ColumnQueryColumn->ColumnQueryDescription, *Storage, *StorageUi, *StorageCompatibility, NameId,
-								TArray<TWeakObjectPtr<const UScriptStruct>>(ColumnTypes.GetData(), ColumnTypes.Num()), 
-								MoveTemp(HeaderConstructor), CellConstructor, FallbackColumn, TedsOutlinerColumn->Outliner, Dealiaser);
+							return MakeShared<FTedsOutlinerUiColumn>(Params);
 
 						}),
 					true,
@@ -756,6 +778,13 @@ void FTedsOutlinerImpl::RecompileQueries()
 	TSharedRef<FRowCopyNode> CopyNode = MakeShared<FRowCopyNode>(ActualCollectorNode);
 	RowHandleSortNode = MakeShared<FRowHandleSortNode>(CopyNode);
 
+	// Nodes to actually sort the rows by the column the user requested
+	RowPrimarySortingNode = MakeShared<FRowSortNode>(*Storage, RowHandleSortNode);
+	RowPrimaryInversionNode = MakeShared<FRowOrderInversionNode>(RowPrimarySortingNode, false);
+
+	// A custom node to store row handle -> sort index mapping.
+	RowSortMapNode = MakeShared<FRowMapNode>(RowPrimaryInversionNode);
+
 	// Our final query to collect rows to populate the Outliner - currently the same as the initial query the user provided
 	FQueryDescription FinalQueryDescription(InitialQueryDescription);
 
@@ -1022,7 +1051,50 @@ void FTedsOutlinerImpl::Tick()
 	}
 
 	// The Sort node is the childmost node in all cases, so updating it causes all the nodes in the stack to update
-	RowHandleSortNode->Update();
+	RowSortMapNode->Update();
+}
+
+bool FTedsOutlinerImpl::IsSorting() const
+{
+	return RowPrimarySortingNode ? RowPrimarySortingNode->IsSorting() : false;
+}
+
+void FTedsOutlinerImpl::OnSort(FName ColumnName, TSharedPtr<const DataStorage::FColumnSorterInterface> ColumnSorter, EColumnSortMode::Type Direction,
+	EColumnSortPriority::Type Priority)
+{
+	if (RowPrimarySortingNode)
+	{
+		if (Direction == EColumnSortMode::Type::Ascending)
+		{
+			RowPrimarySortingNode->SetColumnSorter(ColumnSorter);
+			RowPrimaryInversionNode->Enable(false);
+		}
+		else
+		{
+			RowPrimaryInversionNode->Enable(true);
+		}
+	}
+}
+
+void FTedsOutlinerImpl::SortItems(TArray<FSceneOutlinerTreeItemPtr>& Items, EColumnSortMode::Type Direction) const
+{
+	Items.Sort([this](FSceneOutlinerTreeItemPtr A, FSceneOutlinerTreeItemPtr B)
+	{
+		FTedsOutlinerTreeItem* TedsA = A->CastTo<FTedsOutlinerTreeItem>();
+		FTedsOutlinerTreeItem* TedsB = B->CastTo<FTedsOutlinerTreeItem>();
+
+		if (!TedsA)
+		{
+			return false;
+		}
+
+		if (!TedsB)
+		{
+			return true;
+		}
+			
+		return RowSortMapNode->GetRowIndex(TedsA->GetRowHandle()) < RowSortMapNode->GetRowIndex(TedsB->GetRowHandle());
+	});
 }
 
 DataStorage::ICoreProvider* FTedsOutlinerImpl::GetStorage() const
diff --git a/Engine/Plugins/Experimental/EditorDataStorageFeatures/Source/TedsOutliner/Private/TedsOutlinerMode.cpp b/Engine/Plugins/Experimental/EditorDataStorageFeatures/Source/TedsOutliner/Private/TedsOutlinerMode.cpp
--- a/Engine/Plugins/Experimental/EditorDataStorageFeatures/Source/TedsOutliner/Private/TedsOutlinerMode.cpp
+++ b/Engine/Plugins/Experimental/EditorDataStorageFeatures/Source/TedsOutliner/Private/TedsOutlinerMode.cpp
@@ -30,7 +30,7 @@ namespace UE::Editor::Outliner
 FTedsOutlinerMode::FTedsOutlinerMode(const FTedsOutlinerParams& InParams)
 	: ISceneOutlinerMode(InParams.SceneOutliner)
 {	
-	TedsOutlinerImpl = MakeShared<FTedsOutlinerImpl>(InParams, this);
+	TedsOutlinerImpl = MakeShared<FTedsOutlinerImpl>(InParams, this, false /* bHybridMode */);
 	TedsOutlinerImpl->Init();
 	TedsOutlinerImpl->OnSelectionChanged().AddRaw(this, &FTedsOutlinerMode::OnSelectionChanged);
 	
diff --git a/Engine/Plugins/Experimental/EditorDataStorageFeatures/Source/TedsOutliner/Public/TedsOutlinerImpl.h b/Engine/Plugins/Experimental/EditorDataStorageFeatures/Source/TedsOutliner/Public/TedsOutlinerImpl.h
--- a/Engine/Plugins/Experimental/EditorDataStorageFeatures/Source/TedsOutliner/Public/TedsOutlinerImpl.h
+++ b/Engine/Plugins/Experimental/EditorDataStorageFeatures/Source/TedsOutliner/Public/TedsOutlinerImpl.h
@@ -15,6 +15,8 @@ namespace UE::Editor::DataStorage
 {
 	namespace QueryStack
 	{
+		class FRowOrderInversionNode;
+		class FRowSortNode;
 		class FRowHandleSortNode;
 		class FRowQueryResultsNode;
 		class IQueryNode;
@@ -26,6 +28,11 @@ namespace UE::Editor::DataStorage
 	class ICompatibilityProvider;
 } // namespace UE::Editor::DataStorage
 
+namespace UE::Editor::Outliner
+{
+	class FRowMapNode;
+}
+
 struct FTypedElementWidgetConstructor;
 class SWidget;
 
@@ -154,7 +161,7 @@ public:
 		const STableRow<FSceneOutlinerTreeItemPtr>& RowItem,
 		bool bInteractable = true);
 
-	UE_API FTedsOutlinerImpl(const FTedsOutlinerParams& InParams, ISceneOutlinerMode* InMode);
+	UE_API FTedsOutlinerImpl(const FTedsOutlinerParams& InParams, ISceneOutlinerMode* InMode, bool bInHybridMode);
 	UE_API virtual ~FTedsOutlinerImpl();
 
 	UE_API void Init();
@@ -212,6 +219,12 @@ public:
 	
 	UE_API void Tick();
 
+	// Sorting delegates
+	bool IsSorting() const;
+	void OnSort(FName ColumnName, TSharedPtr<const DataStorage::FColumnSorterInterface> ColumnSorter, EColumnSortMode::Type Direction, 
+		EColumnSortPriority::Type Priority);
+	void SortItems(TArray<FSceneOutlinerTreeItemPtr>& Items, EColumnSortMode::Type Direction) const;
+
 protected:
 	
 	UE_API void ClearSelection() const;
@@ -263,6 +276,15 @@ protected:
 	// A query stack node that keeps a copy of our rows sorted by row handle for faster search/lookup
 	TSharedPtr<DataStorage::QueryStack::FRowHandleSortNode> RowHandleSortNode;
 
+	// A query stack node that keeps a copy of our rows sorted by the actual column the user wants to sort by
+	TSharedPtr<DataStorage::QueryStack::FRowSortNode> RowPrimarySortingNode;
+
+	// Inverts the rows from the query stack when needed by the primary sort.
+	TSharedPtr<DataStorage::QueryStack::FRowOrderInversionNode> RowPrimaryInversionNode;
+
+	// A map of the row handle -> index in the sort node for lookup
+	TSharedPtr<FRowMapNode> RowSortMapNode;
+
 	// Query to get all child rows
 	DataStorage::QueryHandle ChildRowHandleQuery = DataStorage::InvalidQueryHandle;
 
@@ -310,6 +332,9 @@ protected:
 
 	// The columns currently being displayed by the Teds Outliner
 	TArray<FName> AddedColumns;
+
+	// Whether this is a hybrid outliner, i.e it is displaying TEDS rows and some non-TEDS items
+	bool bHybridMode;
 };
 } // namsepace UE::Editor::Outliner
 
diff --git a/Engine/Plugins/Experimental/EditorDataStorageFeatures/Source/TedsTableViewer/Private/TedsTableViewerColumn.cpp b/Engine/Plugins/Experimental/EditorDataStorageFeatures/Source/TedsTableViewer/Private/TedsTableViewerColumn.cpp
--- a/Engine/Plugins/Experimental/EditorDataStorageFeatures/Source/TedsTableViewer/Private/TedsTableViewerColumn.cpp
+++ b/Engine/Plugins/Experimental/EditorDataStorageFeatures/Source/TedsTableViewer/Private/TedsTableViewerColumn.cpp
@@ -197,6 +197,11 @@ namespace UE::Editor::DataStorage
 		IsRowVisibleDelegate = MoveTemp(InIsRowVisibleDelegate);
 	}
 
+	bool FTedsTableViewerColumn::HasSortInfo() const
+	{
+		return !ColumnSorters.IsEmpty();
+	}
+
 	void FTedsTableViewerColumn::SetSortDelegates(FSortHandler InSortDelegates)
 	{
 		SortDelegates = MoveTemp(InSortDelegates);
diff --git a/Engine/Plugins/Experimental/EditorDataStorageFeatures/Source/TedsTableViewer/Public/TedsTableViewerColumn.h b/Engine/Plugins/Experimental/EditorDataStorageFeatures/Source/TedsTableViewer/Public/TedsTableViewerColumn.h
--- a/Engine/Plugins/Experimental/EditorDataStorageFeatures/Source/TedsTableViewer/Public/TedsTableViewerColumn.h
+++ b/Engine/Plugins/Experimental/EditorDataStorageFeatures/Source/TedsTableViewer/Public/TedsTableViewerColumn.h
@@ -24,9 +24,6 @@ namespace UE::Editor::DataStorage
 	 */
 	class FTedsTableViewerColumn
 	{
-		// Done so only the view model has access to setting up the sorting callbacks.
-		friend class FTedsTableViewerModel;
-
 	public:
 		// Delegate to check if a row is currently visible in the owning table viewer's UI
 		DECLARE_DELEGATE_RetVal_OneParam(bool, FIsRowVisible, const RowHandle);
@@ -48,41 +45,45 @@ namespace UE::Editor::DataStorage
 		TEDSTABLEVIEWER_API void Tick();
 		
 		TEDSTABLEVIEWER_API void SetIsRowVisibleDelegate(FIsRowVisible InIsRowVisibleDelegate);
-		
-		TEDSTABLEVIEWER_API void ClearSort();
+
 
 		TEDSTABLEVIEWER_API FName GetColumnName() const;
 
 		TEDSTABLEVIEWER_API TConstArrayView<TWeakObjectPtr<const UScriptStruct>> GetMatchedColumns() const;
 
 		TEDSTABLEVIEWER_API ICoreProvider* GetStorage() const;
-	protected:
 
-		void RegisterQueries();
-		void UnRegisterQueries();
-		bool IsRowVisible(const RowHandle InRowHandle) const;
-		void UpdateWidgets();
-		
-	private:
 		// Delegate triggered when a column is sorted.
 		DECLARE_DELEGATE_FourParams(FOnSort, FName ColumnName, TSharedPtr<const FColumnSorterInterface> ColumnSorter,
 			EColumnSortMode::Type Direction, EColumnSortPriority::Type Priority);
 		// Delegate triggered used to determine is sorting is in progress.
 		DECLARE_DELEGATE_RetVal(bool, FIsSorting);
-
+		
 		struct FSortHandler
 		{
 			FOnSort OnSortHandler;
 			FIsSorting IsSortingHandler;
 		};
 
-		void SetSortDelegates(FSortHandler InSortDelegates);
+		TEDSTABLEVIEWER_API void SetSortDelegates(FSortHandler InSortDelegates);
+		TEDSTABLEVIEWER_API bool HasSortInfo() const;
+		TEDSTABLEVIEWER_API void ClearSort();
+		TEDSTABLEVIEWER_API void OnSortCallback(EColumnSortPriority::Type InSortPriority, const FName& InColumnId, EColumnSortMode::Type InSortDirection);
+
+	protected:
+
+		void RegisterQueries();
+		void UnRegisterQueries();
+		bool IsRowVisible(const RowHandle InRowHandle) const;
+		void UpdateWidgets();
+		
+	private:
+
 		EColumnSortMode::Type GetSortMode() const;
 		EColumnSortPriority::Type GetSortPriority() const;
 		bool IsSorting() const;
 		EVisibility GetSortTextVisibility() const;
 		FText GetSortText() const;
-		void OnSortCallback(EColumnSortPriority::Type InSortPriority, const FName& InColumnId, EColumnSortMode::Type InSortDirection);
 
 		// The ID of the column
 		FName ColumnName;
diff --git a/Engine/Source/Editor/SceneOutliner/Private/SSceneOutliner.cpp b/Engine/Source/Editor/SceneOutliner/Private/SSceneOutliner.cpp
--- a/Engine/Source/Editor/SceneOutliner/Private/SSceneOutliner.cpp
+++ b/Engine/Source/Editor/SceneOutliner/Private/SSceneOutliner.cpp
@@ -2393,6 +2393,12 @@ void SSceneOutliner::Tick(const FGeometry& AllottedGeometry, const double InCurr
 			}
 		}
 
+		// If the column is doing any work to prepare the sort, make sure it is ready before actually sorting the items
+		if (TSharedPtr<ISceneOutlinerColumn> Column = Columns.FindRef(SortByColumn))
+		{
+			bExecuteSort &= Column->IsSortReady();
+		}
+
 		if (bExecuteSort)
 		{
 			UE_LOG(LogSceneOutliner, VeryVerbose, TEXT("Sort Executed"));
@@ -2459,6 +2465,11 @@ void SSceneOutliner::OnColumnSortModeChanged( const EColumnSortPriority::Type So
 void SSceneOutliner::RequestSort()
 {
 	bSortDirty = true;
+
+	if (TSharedPtr<ISceneOutlinerColumn> Column = Columns.FindRef(SortByColumn))
+	{
+		Column->OnSortRequested(EColumnSortPriority::Primary, SortMode);
+	}
 }
 
 void SSceneOutliner::SortItems(TArray<FSceneOutlinerTreeItemPtr>& Items) const
diff --git a/Engine/Source/Editor/SceneOutliner/Public/ISceneOutlinerColumn.h b/Engine/Source/Editor/SceneOutliner/Public/ISceneOutlinerColumn.h
--- a/Engine/Source/Editor/SceneOutliner/Public/ISceneOutlinerColumn.h
+++ b/Engine/Source/Editor/SceneOutliner/Public/ISceneOutlinerColumn.h
@@ -32,5 +32,11 @@ public:
 	virtual bool SupportsSorting() const { return false; }
 
 	virtual void SortItems(TArray<FSceneOutlinerTreeItemPtr>& OutItems, const EColumnSortMode::Type SortMode) const {}
+
+	// Called when a sort is first requested in the Outliner
+	virtual void OnSortRequested(const EColumnSortPriority::Type SortPriority, const EColumnSortMode::Type InSortMode) {}
+
+	// Checks whether the column is still processing the sort or if it is ready to start sorting items
+	virtual bool IsSortReady() { return true; }
 };
 
