[FNameColumn] Add a new column to TedsOutliner to represent the name column and replace uses of the FNameColumn

diff --git a/Engine/Plugins/Experimental/EditorDataStorage/Source/TedsCore/Private/Compatibility/TypedElementUObjectNameQueries.cpp b/Engine/Plugins/Experimental/EditorDataStorage/Source/TedsCore/Private/Compatibility/TypedElementUObjectNameQueries.cpp
--- a/Engine/Plugins/Experimental/EditorDataStorage/Source/TedsCore/Private/Compatibility/TypedElementUObjectNameQueries.cpp
+++ b/Engine/Plugins/Experimental/EditorDataStorage/Source/TedsCore/Private/Compatibility/TypedElementUObjectNameQueries.cpp
@@ -21,11 +21,11 @@ void UObjectNameDataStorageFactory::RegisterUObjectAddNameColumn(UE::Editor::Dat
 			TEXT("Sync UObject Name to column"),
 			FProcessor(EQueryTickPhase::PrePhysics, DataStorage.GetQueryTickGroupName(EQueryTickGroups::SyncExternalToDataStorage))
 			.SetExecutionMode(EExecutionMode::GameThread),
-			[](const FTypedElementUObjectColumn& Object, FNameColumn& Name)
+			[](const FTypedElementUObjectColumn& Object, FUObjectIdNameColumn& IdName)
 			{
 				if (UObject* ObjectInstance = Object.Object.Get())
 				{
-					Name.Name = ObjectInstance->GetFName();
+					IdName.IdName = ObjectInstance->GetFName();
 				}
 			}
 		)
diff --git a/Engine/Plugins/Experimental/EditorDataStorage/Source/TedsCore/Private/TypedElementDatabaseCompatibility.cpp b/Engine/Plugins/Experimental/EditorDataStorage/Source/TedsCore/Private/TypedElementDatabaseCompatibility.cpp
--- a/Engine/Plugins/Experimental/EditorDataStorage/Source/TedsCore/Private/TypedElementDatabaseCompatibility.cpp
+++ b/Engine/Plugins/Experimental/EditorDataStorage/Source/TedsCore/Private/TypedElementDatabaseCompatibility.cpp
@@ -344,8 +344,8 @@ void UEditorDataStorageCompatibility::CreateStandardArchetypes()
 	// Thread-safe as this is only called from a function that has an exclusive lock.
 	using namespace UE::Editor::DataStorage;
 	StandardUObjectTable = Storage->RegisterTable(TTypedElementColumnTypeList<
-			FTypedElementUObjectColumn, FTypedElementUObjectIdColumn, FTypedElementClassTypeInfoColumn,
-			FNameColumn, FTypedElementSyncFromWorldTag>(), 
+			FTypedElementUObjectColumn, FTypedElementUObjectIdColumn, FUObjectIdNameColumn,
+			FTypedElementClassTypeInfoColumn, FTypedElementSyncFromWorldTag>(),
 		FName("Editor_StandardUObjectTable"));
 
 	StandardExternalObjectTable = Storage->RegisterTable(TTypedElementColumnTypeList<
diff --git a/Engine/Plugins/Experimental/EditorDataStorageFeatures/Source/TedsActorCompatibility/Private/TedsActorCompatibilityFactory.cpp b/Engine/Plugins/Experimental/EditorDataStorageFeatures/Source/TedsActorCompatibility/Private/TedsActorCompatibilityFactory.cpp
--- a/Engine/Plugins/Experimental/EditorDataStorageFeatures/Source/TedsActorCompatibility/Private/TedsActorCompatibilityFactory.cpp
+++ b/Engine/Plugins/Experimental/EditorDataStorageFeatures/Source/TedsActorCompatibility/Private/TedsActorCompatibilityFactory.cpp
@@ -38,9 +38,9 @@ void UTedsActorCompatibilityFactory::RegisterTable(UE::Editor::DataStorage::ICor
 	using namespace UE::Editor::DataStorage;
 
 	StandardActorTable = DataStorage.RegisterTable(TTypedElementColumnTypeList<
-		FTypedElementUObjectColumn, FTypedElementUObjectIdColumn, FTypedElementClassTypeInfoColumn,
-		FTypedElementLabelColumn, FTypedElementLabelHashColumn, FNameColumn, FTedsActorMobilityColumn,
-		FTypedElementActorTag, FTypedElementSyncFromWorldTag>(),
+		FTypedElementUObjectColumn, FTypedElementUObjectIdColumn, FUObjectIdNameColumn,
+		FTypedElementClassTypeInfoColumn, FTypedElementLabelColumn, FTypedElementLabelHashColumn, 
+		FTedsActorMobilityColumn, FTypedElementActorTag, FTypedElementSyncFromWorldTag>(),
 		FName("Editor_StandardActorTable"));
 
 	StandardActorWithTransformTable = DataStorage.RegisterTable(StandardActorTable,
diff --git a/Engine/Plugins/Experimental/EditorDataStorageFeatures/Source/TedsOutliner/Private/Compatibility/SceneOutlinerTedsBridge.cpp b/Engine/Plugins/Experimental/EditorDataStorageFeatures/Source/TedsOutliner/Private/Compatibility/SceneOutlinerTedsBridge.cpp
--- a/Engine/Plugins/Experimental/EditorDataStorageFeatures/Source/TedsOutliner/Private/Compatibility/SceneOutlinerTedsBridge.cpp
+++ b/Engine/Plugins/Experimental/EditorDataStorageFeatures/Source/TedsOutliner/Private/Compatibility/SceneOutlinerTedsBridge.cpp
@@ -3,6 +3,7 @@
 #include "Compatibility/SceneOutlinerTedsBridge.h"
 
 #include "Columns/TedsOutlinerColumns.h"
+#include "Elements/Columns/TypedElementCompatibilityColumns.h"
 #include "Elements/Columns/TypedElementMiscColumns.h"
 #include "Elements/Columns/TypedElementPackageColumns.h"
 #include "Elements/Columns/TypedElementTypeInfoColumns.h"
@@ -23,10 +24,12 @@
 
 void USceneOutlinerTedsBridgeFactory::PreRegister(UE::Editor::DataStorage::ICoreProvider& DataStorage)
 {
+	using namespace UE::Editor::DataStorage;
+
 	// Map the type column from the TEDS to the default Outliner type column, so we can show type info for objects not in TEDS
 	TEDSToOutlinerDefaultColumnMapping.Add(FTypedElementClassTypeInfoColumn::StaticStruct(), FSceneOutlinerBuiltInColumnTypes::ActorInfo());
 	TEDSToOutlinerDefaultColumnMapping.Add(FVisibleInEditorColumn::StaticStruct(), FSceneOutlinerBuiltInColumnTypes::Gutter());
-	TEDSToOutlinerDefaultColumnMapping.Add(FNameColumn::StaticStruct(), FSceneOutlinerBuiltInColumnTypes::IDName());
+	TEDSToOutlinerDefaultColumnMapping.Add(FUObjectIdNameColumn::StaticStruct(), FSceneOutlinerBuiltInColumnTypes::IDName());
 	TEDSToOutlinerDefaultColumnMapping.Add(FTedsActorMobilityColumn::StaticStruct(), FSceneOutlinerBuiltInColumnTypes::Mobility());
 }
 
diff --git a/Engine/Plugins/Experimental/EditorDataStorageFeatures/Source/TedsOutliner/Private/TedsOutlinerItem.cpp b/Engine/Plugins/Experimental/EditorDataStorageFeatures/Source/TedsOutliner/Private/TedsOutlinerItem.cpp
--- a/Engine/Plugins/Experimental/EditorDataStorageFeatures/Source/TedsOutliner/Private/TedsOutlinerItem.cpp
+++ b/Engine/Plugins/Experimental/EditorDataStorageFeatures/Source/TedsOutliner/Private/TedsOutlinerItem.cpp
@@ -44,15 +44,17 @@ FSceneOutlinerTreeItemID FTedsOutlinerTreeItem::GetID() const
 
 FString FTedsOutlinerTreeItem::GetDisplayString() const
 {
+	using namespace UE::Editor::DataStorage;
+
 	if (const TSharedPtr<const FTedsOutlinerImpl> TedsOutlinerImplPin = TedsOutlinerImpl.Pin())
 	{
 		if (const FTypedElementLabelColumn* LabelColumn = TedsOutlinerImplPin->GetStorage()->GetColumn<FTypedElementLabelColumn>(RowHandle))
 		{
 			return LabelColumn->Label;
 		}
-		if (const FNameColumn* NameColumn = TedsOutlinerImplPin->GetStorage()->GetColumn<FNameColumn>(RowHandle))
+		if (const FUObjectIdNameColumn* NameColumn = TedsOutlinerImplPin->GetStorage()->GetColumn<FUObjectIdNameColumn>(RowHandle))
 		{
-			return NameColumn->Name.ToString();
+			return NameColumn->IdName.ToString();
 		}
 	}
 	
diff --git a/Engine/Plugins/Experimental/EditorDataStorageFeatures/Source/TedsOutliner/Private/TedsOutlinerModule.cpp b/Engine/Plugins/Experimental/EditorDataStorageFeatures/Source/TedsOutliner/Private/TedsOutlinerModule.cpp
--- a/Engine/Plugins/Experimental/EditorDataStorageFeatures/Source/TedsOutliner/Private/TedsOutlinerModule.cpp
+++ b/Engine/Plugins/Experimental/EditorDataStorageFeatures/Source/TedsOutliner/Private/TedsOutlinerModule.cpp
@@ -126,7 +126,7 @@ UE::Editor::DataStorage::FQueryDescription FTedsOutlinerModule::GetLevelEditorTe
 			.ReadOnly<FTypedElementClassTypeInfoColumn,
 					  FAlertColumn, 
 					  FChildAlertColumn,
-					  FNameColumn,
+					  FUObjectIdNameColumn,
 					  FTedsActorMobilityColumn,
 					  FVisibleInEditorColumn>()
 		.Compile();
@@ -137,7 +137,7 @@ UE::Editor::DataStorage::FQueryDescription FTedsOutlinerModule::GetLevelEditorTe
 			.ReadOnly<FTypedElementClassTypeInfoColumn,  
 					  FTypedElementPackageReference, 
 					  FAlertColumn, 
-					  FNameColumn,
+					  FUObjectIdNameColumn,
 					  FTedsActorMobilityColumn,
 					  FVisibleInEditorColumn>()
 		.Compile();
diff --git a/Engine/Plugins/Experimental/EditorDataStorageFeatures/Source/TedsOutliner/Private/Widgets/OutlinerIdNameWidget.cpp b/Engine/Plugins/Experimental/EditorDataStorageFeatures/Source/TedsOutliner/Private/Widgets/OutlinerIdNameWidget.cpp
--- a/Engine/Plugins/Experimental/EditorDataStorageFeatures/Source/TedsOutliner/Private/Widgets/OutlinerIdNameWidget.cpp
+++ b/Engine/Plugins/Experimental/EditorDataStorageFeatures/Source/TedsOutliner/Private/Widgets/OutlinerIdNameWidget.cpp
@@ -24,7 +24,7 @@ void UOutlinerIdNameWidgetFactory::RegisterWidgetConstructors(UE::Editor::DataSt
 
 	DataStorageUi.RegisterWidgetFactory<FOutlinerIdNameWidgetConstructor>(
 	DataStorageUi.FindPurpose(UE::Editor::DataStorage::IUiProvider::FPurposeInfo("SceneOutliner", "Cell", NAME_None).GeneratePurposeID()),
-		TColumn<FNameColumn>());
+		TColumn<FUObjectIdNameColumn>());
 }
 
 FOutlinerIdNameWidgetConstructor::FOutlinerIdNameWidgetConstructor()
@@ -56,8 +56,8 @@ TSharedPtr<SWidget> FOutlinerIdNameWidgetConstructor::CreateWidget(
 		[
 			SNew(STextBlock)
 				.ColorAndOpacity(FSlateColor::UseSubduedForeground())
-				.Text(TargetRowBinder.BindText(&FNameColumn::Name))
-				.ToolTipText(TargetRowBinder.BindText(&FNameColumn::Name))
+				.Text(TargetRowBinder.BindText(&FUObjectIdNameColumn::IdName))
+				.ToolTipText(TargetRowBinder.BindText(&FUObjectIdNameColumn::IdName))
 		];
 
 	return Widget;
diff --git a/Engine/Source/Runtime/TypedElementFramework/Public/Elements/Columns/TypedElementCompatibilityColumns.h b/Engine/Source/Runtime/TypedElementFramework/Public/Elements/Columns/TypedElementCompatibilityColumns.h
--- a/Engine/Source/Runtime/TypedElementFramework/Public/Elements/Columns/TypedElementCompatibilityColumns.h
+++ b/Engine/Source/Runtime/TypedElementFramework/Public/Elements/Columns/TypedElementCompatibilityColumns.h
@@ -116,4 +116,21 @@ struct FTypedElementPropertyBagPlaceholderTypeInfoColumn final : public FEditorD
 	GENERATED_BODY()
 
 	TWeakObjectPtr<const UStruct> InferredBaseType;
-};
\ No newline at end of file
+};
+
+/**
+ * Column containing the ID Name of the UObject.
+ */
+USTRUCT(meta = (DisplayName = "Name"))
+struct FEditorDataStorageUObjectIdNameColumn final : public FEditorDataStorageColumn
+{
+	GENERATED_BODY()
+
+	UPROPERTY(meta = (Searchable, Sortable))
+	FName IdName;
+};
+
+namespace UE::Editor::DataStorage
+{
+	using FUObjectIdNameColumn = FEditorDataStorageUObjectIdNameColumn;
+}  // namespace UE::Editor::DataStorage
\ No newline at end of file
