[Mobility] Add the mobility column to entities with hardcoded 'Moveable', such that the outliner sorting looks correct until the hybrid outliner is removed

diff --git a/Engine/Plugins/Experimental/EditorDataStorageFeatures/Source/TedsActorCompatibility/Public/Columns/TedsActorMobilityColumns.h b/Engine/Plugins/Experimental/EditorDataStorageFeatures/Source/TedsActorCompatibility/Public/Columns/TedsActorMobilityColumns.h
--- a/Engine/Plugins/Experimental/EditorDataStorageFeatures/Source/TedsActorCompatibility/Public/Columns/TedsActorMobilityColumns.h
+++ b/Engine/Plugins/Experimental/EditorDataStorageFeatures/Source/TedsActorCompatibility/Public/Columns/TedsActorMobilityColumns.h
@@ -16,5 +16,5 @@ struct FTedsActorMobilityColumn final : public FEditorDataStorageColumn
 	GENERATED_BODY()
 
 	UPROPERTY(meta = (Sortable))
-	TEnumAsByte<EComponentMobility::Type> Mobility = EComponentMobility::Static;
+	TEnumAsByte<EComponentMobility::Type> Mobility = EComponentMobility::Movable;
 };
\ No newline at end of file
