Fix some potential link errors

#rnx

diff --git a/Engine/Plugins/AudioInsights/Source/AudioInsights/Private/AudioInsightsStyle.cpp b/Engine/Plugins/AudioInsights/Source/AudioInsights/Private/AudioInsightsStyle.cpp
new file mode 100644
--- /dev/null
+++ b/Engine/Plugins/AudioInsights/Source/AudioInsights/Private/AudioInsightsStyle.cpp
@@ -0,0 +1,12 @@
+// Copyright Epic Games, Inc. All Rights Reserved.
+
+#include "AudioInsightsStyle.h"
+
+namespace UE::Audio::Insights
+{
+	FSlateStyle& FSlateStyle::Get()
+	{
+		static FSlateStyle InsightsStyle;
+		return InsightsStyle;
+	}
+}
diff --git a/Engine/Plugins/AudioInsights/Source/AudioInsights/Public/AudioInsightsStyle.h b/Engine/Plugins/AudioInsights/Source/AudioInsights/Public/AudioInsightsStyle.h
--- a/Engine/Plugins/AudioInsights/Source/AudioInsights/Public/AudioInsightsStyle.h
+++ b/Engine/Plugins/AudioInsights/Source/AudioInsights/Public/AudioInsightsStyle.h
@@ -17,11 +17,7 @@ namespace UE::Audio::Insights
 	class FSlateStyle final : public FSlateStyleSet
 	{
 	public:
-		AUDIOINSIGHTS_API static FSlateStyle& Get()
-		{
-			static FSlateStyle InsightsStyle;
-			return InsightsStyle;
-		}
+		AUDIOINSIGHTS_API static FSlateStyle& Get();
 
 		static FName GetStyleName()
 		{
diff --git a/Engine/Plugins/Experimental/EditorDataStorageFeatures/Source/TedsEverythingPicker/Public/Widgets/SEverythingPicker.h b/Engine/Plugins/Experimental/EditorDataStorageFeatures/Source/TedsEverythingPicker/Public/Widgets/SEverythingPicker.h
--- a/Engine/Plugins/Experimental/EditorDataStorageFeatures/Source/TedsEverythingPicker/Public/Widgets/SEverythingPicker.h
+++ b/Engine/Plugins/Experimental/EditorDataStorageFeatures/Source/TedsEverythingPicker/Public/Widgets/SEverythingPicker.h
@@ -13,7 +13,7 @@ class SHorizontalBox;
 
 namespace UE::Editor::DataStorage::Picker
 {
-	class TEDSEVERYTHINGPICKER_API SEverythingPicker : public SCompoundWidget
+	class SEverythingPicker : public SCompoundWidget
 	{
 	public:
 		SLATE_BEGIN_ARGS(SEverythingPicker)
@@ -21,9 +21,9 @@ namespace UE::Editor::DataStorage::Picker
 			SLATE_SLOT_ARGUMENT(FPickerContext, Contexts)
 		SLATE_END_ARGS()
 
-		void Construct(const FArguments& InArgs);
+		TEDSEVERYTHINGPICKER_API void Construct(const FArguments& InArgs);
 
-		void AddContext(const FPickerContext::FSlotArguments& SlotArgs);
+		TEDSEVERYTHINGPICKER_API void AddContext(const FPickerContext::FSlotArguments& SlotArgs);
 
 		static FPickerContext::FSlotArguments Context()
 		{
diff --git a/Engine/Plugins/Runtime/Harmonix/Source/HarmonixDsp/Public/HarmonixDsp/Effects/TapeSS.h b/Engine/Plugins/Runtime/Harmonix/Source/HarmonixDsp/Public/HarmonixDsp/Effects/TapeSS.h
--- a/Engine/Plugins/Runtime/Harmonix/Source/HarmonixDsp/Public/HarmonixDsp/Effects/TapeSS.h
+++ b/Engine/Plugins/Runtime/Harmonix/Source/HarmonixDsp/Public/HarmonixDsp/Effects/TapeSS.h
@@ -39,8 +39,8 @@ namespace Harmonix::Dsp::Effects
 
 		HARMONIXDSP_API void  SnapState(EState State);
 		HARMONIXDSP_API void  SetState(EState State, float DurationSeconds);
-		HARMONIXDSP_API EState GetState() const { return CurrentState; }
-		HARMONIXDSP_API float GetSpeed() const { return SRRamper.GetCurrent(); }
+		EState GetState() const { return CurrentState; }
+		float GetSpeed() const { return SRRamper.GetCurrent(); }
 
 	private:
 		void CheckForStateUpdates();
