Decouple STedsHierarchyViewer from the FRowFilterNode

diff --git a/Engine/Plugins/Experimental/EditorDataStorageFeatures/Source/TedsEverythingPicker/Private/Context/TedsPickerContextUtil.cpp b/Engine/Plugins/Experimental/EditorDataStorageFeatures/Source/TedsEverythingPicker/Private/Context/TedsPickerContextUtil.cpp
--- a/Engine/Plugins/Experimental/EditorDataStorageFeatures/Source/TedsEverythingPicker/Private/Context/TedsPickerContextUtil.cpp
+++ b/Engine/Plugins/Experimental/EditorDataStorageFeatures/Source/TedsEverythingPicker/Private/Context/TedsPickerContextUtil.cpp
@@ -29,11 +29,23 @@ namespace UE::Editor::DataStorage::Picker
 			MakeShared<QueryStack::FRowQueryResultsNode>(*Storage,
 				ReferenceQueryNode, QueryStack::FRowQueryResultsNode::ESyncFlags::RefreshOnUpdate);
 
+		// Node that filters out the ReferenceQueryResultsNode to generate all the top level nodes
+		constexpr bool bRowsShouldHaveColumn = false;
+		TSharedRef<QueryStack::IRowNode> TopLevelNode =
+			MakeShared<QueryStack::FRowFilterNode<FTableRowParentColumn>> (Storage, ReferenceQueryResultsNode.ToSharedPtr(), bRowsShouldHaveColumn);
+
 		// Using the Teds Outliner Purpose for matching customization for now
 		IUiProvider::FPurposeID PurposeId = IUiProvider::FPurposeInfo("SceneOutliner", "RowLabel", NAME_None).GeneratePurposeID();
 
-		TSharedPtr<SHierarchyViewer> ReferenceView = SNew(SHierarchyViewer)
-			.QueryStack(ReferenceQueryResultsNode)
+		auto ParentExtractionFunction = [](const void* Column, UScriptStruct*)->RowHandle
+		{
+			const FTableRowParentColumn* ParentColumn = static_cast<const FTableRowParentColumn*>(Column);
+			return ParentColumn->Parent;
+		};
+
+		TSharedPtr<SHierarchyViewer> ReferenceView = SNew(SHierarchyViewer, FTableRowParentColumn::StaticStruct(), ParentExtractionFunction)
+			.AllNodeProvider(ReferenceQueryResultsNode)
+			.TopLevelNodeProvider(TopLevelNode)
 			.Columns({ FTypedElementLabelColumn::StaticStruct(), FTypedElementClassTypeInfoColumn::StaticStruct() })
 			.CellWidgetPurpose(PurposeId)
 			.OnSelectionChanged(InArgs._OnSelectionChanged)
diff --git a/Engine/Plugins/Experimental/EditorDataStorageFeatures/Source/TedsTableViewer/Private/Widgets/STedsHierarchyViewer.cpp b/Engine/Plugins/Experimental/EditorDataStorageFeatures/Source/TedsTableViewer/Private/Widgets/STedsHierarchyViewer.cpp
--- a/Engine/Plugins/Experimental/EditorDataStorageFeatures/Source/TedsTableViewer/Private/Widgets/STedsHierarchyViewer.cpp
+++ b/Engine/Plugins/Experimental/EditorDataStorageFeatures/Source/TedsTableViewer/Private/Widgets/STedsHierarchyViewer.cpp
@@ -18,8 +18,14 @@
 
 namespace UE::Editor::DataStorage
 {
-	void SHierarchyViewer::Construct(const FArguments& InArgs)
+	void SHierarchyViewer::Construct(
+		const FArguments& InArgs,
+		UScriptStruct* InHierarchyDataColumnType,
+		TFunction<RowHandle(const void*, UScriptStruct*)> InGetRowHandleFn)
 	{
+		checkf(InHierarchyDataColumnType != nullptr, TEXT("Nullptr ParentColumnType"));
+		HierarchyColumnType = InHierarchyDataColumnType;
+		GetParentFromColumnFn = MoveTemp(InGetRowHandleFn);
 		OnSelectionChanged = InArgs._OnSelectionChanged;
 		EmptyRowsMessage = InArgs._EmptyRowsMessage;
 		DefaultExpansionState = InArgs._DefaultExpansionState;
@@ -39,11 +45,10 @@ namespace UE::Editor::DataStorage
 			HeaderWidgetPurpose = IUiProvider::FPurposeInfo("General", "Header", NAME_None).GeneratePurposeID();
 		}
 
-		Model = MakeShared<FTedsTableViewerModel>(InArgs._QueryStack, InArgs._Columns, CellWigetPurpose, HeaderWidgetPurpose,
+		Model = MakeShared<FTedsTableViewerModel>(InArgs._AllNodeProvider, InArgs._Columns, CellWigetPurpose, HeaderWidgetPurpose,
 			FTedsTableViewerModel::FIsItemVisible::CreateSP(this, &SHierarchyViewer::IsItemVisible));
-
-		constexpr bool bRowsShouldHaveColumn = false;
-		HierarchyNode = MakeUnique<QueryStack::FRowFilterNode<FTableRowParentColumn>>(Model->GetDataStorageInterface(), InArgs._QueryStack, bRowsShouldHaveColumn);
+		
+		HierarchyNode = InArgs._TopLevelNodeProvider;
 		
 		HeaderRowWidget = SNew( SHeaderRow )
 							.CanSelectGeneratedColumn(true);
@@ -120,12 +125,13 @@ namespace UE::Editor::DataStorage
 
 	TableViewerItemPtr SHierarchyViewer::GetParentRow(TableViewerItemPtr InItem) const
 	{
-		if (FTableRowParentColumn* ParentColumn = Model->GetDataStorageInterface()->GetColumn<FTableRowParentColumn>(InItem))
+		if (const void* ParentColumn = Model->GetDataStorageInterface()->GetColumnData(InItem, HierarchyColumnType))
 		{
-			return FTedsRowHandle{.RowHandle = ParentColumn->Parent};
+			RowHandle ParentRow = GetParentFromColumnFn(ParentColumn, HierarchyColumnType);
+			return FTedsRowHandle(ParentRow);
 		}
 
-		return FTedsRowHandle{.RowHandle = InvalidRowHandle};
+		return FTedsRowHandle(InvalidRowHandle);
 	}
 
 	void SHierarchyViewer::CreateInternalWidget()
diff --git a/Engine/Plugins/Experimental/EditorDataStorageFeatures/Source/TedsTableViewer/Public/Widgets/STedsHierarchyViewer.h b/Engine/Plugins/Experimental/EditorDataStorageFeatures/Source/TedsTableViewer/Public/Widgets/STedsHierarchyViewer.h
--- a/Engine/Plugins/Experimental/EditorDataStorageFeatures/Source/TedsTableViewer/Public/Widgets/STedsHierarchyViewer.h
+++ b/Engine/Plugins/Experimental/EditorDataStorageFeatures/Source/TedsTableViewer/Public/Widgets/STedsHierarchyViewer.h
@@ -6,6 +6,7 @@
 #include "TedsRowFilterNode.h"
 #include "TedsTableViewerModel.h"
 #include "TedsQueryStackInterfaces.h"
+#include "Templates/Function.h"
 #include "Elements/Columns/TypedElementHiearchyColumns.h"
 #include "Widgets/SCompoundWidget.h"
 #include "Widgets/Views/STableRow.h"
@@ -53,10 +54,12 @@ namespace UE::Editor::DataStorage
 		{
 			
 		}
-		
+
+		// RowNode that provides the top level rows for the hierarchy
+		SLATE_ARGUMENT(TSharedPtr<QueryStack::IRowNode>, TopLevelNodeProvider)
 
 		// Query Stack that will supply the rows to be displayed
-		SLATE_ARGUMENT(TSharedPtr<QueryStack::IRowNode>, QueryStack)
+		SLATE_ARGUMENT(TSharedPtr<QueryStack::IRowNode>, AllNodeProvider)
 
 		// The Columns that this table viewer will display
 		// Table Viewer TODO: How do we specify column metadata (ReadOnly or ReadWrite)?
@@ -85,7 +88,10 @@ namespace UE::Editor::DataStorage
 
 	public:
 		
-		TEDSTABLEVIEWER_API void Construct(const FArguments& InArgs);
+		TEDSTABLEVIEWER_API void Construct(
+			const FArguments& InArgs,
+			UScriptStruct* InHierarchyDataColumnType,
+			TFunction<RowHandle(const void*, UScriptStruct*)> InGetRowHandleFn);
 
 		// Clear the current list of columns being displayed and set it to the given list
 		TEDSTABLEVIEWER_API void SetColumns(const TArray<TWeakObjectPtr<const UScriptStruct>>& Columns);
@@ -157,10 +163,16 @@ namespace UE::Editor::DataStorage
 		TAttribute<EExpansionState> DefaultExpansionState;
 
 		// A specialized query stack node to only contain the top level rows present in the actual query stack we are using.
-		TUniquePtr<QueryStack::FRowFilterNode<FTableRowParentColumn>> HierarchyNode;
+		TSharedPtr<QueryStack::IRowNode> HierarchyNode;
 		
 		// The cached list of top level rows
 		TArray<TableViewerItemPtr> TopLevelRows;
 
+		// Column type that holds hierarchy information for the row
+		UScriptStruct* HierarchyColumnType = nullptr;
+
+		// Function to get the parent from a the HierarchyColumnType of the row
+		TFunction<RowHandle(const void*, UScriptStruct*)> GetParentFromColumnFn;
+
 	};
 }
\ No newline at end of file
