Teds Outliner: Fix the GetChildren() function not checking whether the children actually match the query

diff --git a/Engine/Plugins/Experimental/EditorDataStorageFeatures/Source/TedsOutliner/Private/TedsOutlinerImpl.cpp b/Engine/Plugins/Experimental/EditorDataStorageFeatures/Source/TedsOutliner/Private/TedsOutlinerImpl.cpp
--- a/Engine/Plugins/Experimental/EditorDataStorageFeatures/Source/TedsOutliner/Private/TedsOutlinerImpl.cpp
+++ b/Engine/Plugins/Experimental/EditorDataStorageFeatures/Source/TedsOutliner/Private/TedsOutlinerImpl.cpp
@@ -676,11 +676,14 @@ void FTedsOutlinerImpl::CreateChildren(const FSceneOutlinerTreeItemPtr& Item, TA
 
 		GetChildrenRecursive(ItemRowHandle);
 	}
+	
+	// If the row doesn't exist in our query stack results, it doesn't match the query and should not be displayed
+	FRowHandleArrayView Rows(RowHandleSortNode->GetRows());
 
 	// Actually create the items for the child entities 
 	for (RowHandle ChildItemRowHandle : ChildItems)
 	{
-		if (!CanDisplayRow(ChildItemRowHandle))
+		if (!CanDisplayRow(ChildItemRowHandle) || !Rows.Contains(ChildItemRowHandle))
 		{
 			continue;
 		}
