[FNameColumn] Add a new column to Teds Asset Data to represent the name column and replace uses of the FNameColumn

diff --git a/Engine/Plugins/Experimental/EditorDataStorage/Source/TedsUI/Private/Widgets/NameWidget.cpp b/Engine/Plugins/Experimental/EditorDataStorage/Source/TedsUI/Private/Widgets/NameWidget.cpp
deleted file mode 100644
diff --git a/Engine/Plugins/Experimental/EditorDataStorage/Source/TedsUI/Public/Widgets/NameWidget.h b/Engine/Plugins/Experimental/EditorDataStorage/Source/TedsUI/Public/Widgets/NameWidget.h
deleted file mode 100644
diff --git a/Engine/Plugins/Experimental/EditorDataStorageFeatures/Source/TedsAssetData/Private/CB/TedsAssetDataCBDataSource.cpp b/Engine/Plugins/Experimental/EditorDataStorageFeatures/Source/TedsAssetData/Private/CB/TedsAssetDataCBDataSource.cpp
--- a/Engine/Plugins/Experimental/EditorDataStorageFeatures/Source/TedsAssetData/Private/CB/TedsAssetDataCBDataSource.cpp
+++ b/Engine/Plugins/Experimental/EditorDataStorageFeatures/Source/TedsAssetData/Private/CB/TedsAssetDataCBDataSource.cpp
@@ -11,7 +11,6 @@
 #include "Blueprint/BlueprintSupport.h"
 #include "ContentBrowserDataUtils.h"
 #include "DataStorage/Queries/Types.h"
-#include "Elements/Columns/TypedElementMiscColumns.h"
 #include "Elements/Framework/TypedElementQueryBuilder.h"
 #include "Elements/Interfaces/TypedElementDataStorageInterface.h"
 #include "Factories/Factory.h"
@@ -264,7 +263,7 @@ void FTedsAssetDataCBDataSource::AddAssetDataColumns(UE::Editor::DataStorage::IQ
 	AssetClassColumn.ClassPath = AssetData.AssetClassPath;
 	Context.AddColumn(Row, MoveTemp(AssetClassColumn));
 
-	FNameColumn ItemNameColumn;
+	FAssetNameColumn ItemNameColumn;
 	ItemNameColumn.Name = AssetData.AssetName;
 	Context.AddColumn(Row, MoveTemp(ItemNameColumn));
 
diff --git a/Engine/Plugins/Experimental/EditorDataStorageFeatures/Source/TedsAssetData/Private/TedsAssetData.cpp b/Engine/Plugins/Experimental/EditorDataStorageFeatures/Source/TedsAssetData/Private/TedsAssetData.cpp
--- a/Engine/Plugins/Experimental/EditorDataStorageFeatures/Source/TedsAssetData/Private/TedsAssetData.cpp
+++ b/Engine/Plugins/Experimental/EditorDataStorageFeatures/Source/TedsAssetData/Private/TedsAssetData.cpp
@@ -10,7 +10,6 @@
 #include "Containers/ChunkedArray.h"
 #include "DataStorage/Queries/Types.h"
 #include "Elements/Columns/TypedElementFolderColumns.h"
-#include "Elements/Columns/TypedElementMiscColumns.h"
 #include "Elements/Columns/TypedElementSlateWidgetColumns.h"
 #include "Elements/Framework/TypedElementQueryBuilder.h"
 #include "Elements/Interfaces/TypedElementDataStorageInterface.h"
@@ -112,7 +111,7 @@ FPopulatePathRowArgs ThreadSafe_PopulatePathRowArgs(FMapKey&& AssetRegistryPathK
 void PopulatePathDataTableRow(FPopulatePathRowArgs&& InPopulatePathRowArgs, ICoreProvider& Database, RowHandle InRowHandle)
 {
 	Database.GetColumn<FAssetPathColumn_Experimental>(InRowHandle)->Path = InPopulatePathRowArgs.AssetRegistryPath;
-	Database.GetColumn<FNameColumn>(InRowHandle)->Name = InPopulatePathRowArgs.AssetName;
+	Database.GetColumn<FAssetNameColumn>(InRowHandle)->Name = InPopulatePathRowArgs.AssetName;
 }
 
 
@@ -141,7 +140,7 @@ FTedsAssetData::FTedsAssetData(UE::Editor::DataStorage::ICoreProvider& InDatabas
 	PathsTable = Database.FindTable(FName(TEXT("Editor_AssetRegistryPathsTable")));
 	if (PathsTable == InvalidTableHandle)
 	{
-		PathsTable = Database.RegisterTable<FFolderTag, FAssetPathColumn_Experimental, FNameColumn, FUpdatedPathTag>(FName(TEXT("Editor_AssetRegistryPathsTable")));
+		PathsTable = Database.RegisterTable<FFolderTag, FAssetPathColumn_Experimental, FAssetNameColumn, FUpdatedPathTag>(FName(TEXT("Editor_AssetRegistryPathsTable")));
 	}
 
 	AssetsDataTable = Database.FindTable(FName(TEXT("Editor_AssetRegistryAssetDataTable")));
diff --git a/Engine/Plugins/Experimental/EditorDataStorageFeatures/Source/TedsAssetData/Private/Widgets/AssetDataLabelWidget.cpp b/Engine/Plugins/Experimental/EditorDataStorageFeatures/Source/TedsAssetData/Private/Widgets/AssetDataLabelWidget.cpp
--- a/Engine/Plugins/Experimental/EditorDataStorageFeatures/Source/TedsAssetData/Private/Widgets/AssetDataLabelWidget.cpp
+++ b/Engine/Plugins/Experimental/EditorDataStorageFeatures/Source/TedsAssetData/Private/Widgets/AssetDataLabelWidget.cpp
@@ -8,7 +8,6 @@
 #include "Elements/Framework/TypedElementAttributeBinding.h"
 #include "TedsAssetDataColumns.h"
 #include "TedsAssetDataHelper.h"
-#include "Elements/Columns/TypedElementMiscColumns.h"
 #include "Elements/Columns/TypedElementSlateWidgetColumns.h"
 #include "Widgets/Images/SImage.h"
 #include "Widgets/Layout/SBox.h"
@@ -24,7 +23,7 @@ void UAssetDataLabelWidgetFactory::RegisterWidgetConstructors(UE::Editor::DataSt
 	using namespace UE::Editor::DataStorage::Queries;
 	DataStorageUi.RegisterWidgetFactory<FAssetDataLabelWidgetConstructor>(
 		DataStorageUi.FindPurpose(IUiProvider::FPurposeInfo("General", "RowLabel", NAME_None).GeneratePurposeID()),
-		TColumn<FNameColumn>() && (TColumn<FAssetTag>() || TColumn<FAssetPathColumn_Experimental>()));
+		TColumn<FAssetNameColumn>() && (TColumn<FAssetTag>() || TColumn<FAssetPathColumn_Experimental>()));
 }
 
 FAssetDataLabelWidgetConstructor::FAssetDataLabelWidgetConstructor()
@@ -154,7 +153,7 @@ TSharedPtr<SWidget> FAssetDataLabelWidgetConstructor::CreateWidget(UE::Editor::D
 
 TArray<TWeakObjectPtr<const UScriptStruct>> FAssetDataLabelWidgetConstructor::GetLabelColumns()
 {
-	static TArray<TWeakObjectPtr<const UScriptStruct>> LabelColumns({ TWeakObjectPtr(FNameColumn::StaticStruct()) });
+	static TArray<TWeakObjectPtr<const UScriptStruct>> LabelColumns({ TWeakObjectPtr(FAssetNameColumn::StaticStruct()) });
 	return LabelColumns;
 }
 
diff --git a/Engine/Plugins/Experimental/EditorDataStorageFeatures/Source/TedsAssetData/Private/Widgets/AssetDataThumbnailWidget.cpp b/Engine/Plugins/Experimental/EditorDataStorageFeatures/Source/TedsAssetData/Private/Widgets/AssetDataThumbnailWidget.cpp
--- a/Engine/Plugins/Experimental/EditorDataStorageFeatures/Source/TedsAssetData/Private/Widgets/AssetDataThumbnailWidget.cpp
+++ b/Engine/Plugins/Experimental/EditorDataStorageFeatures/Source/TedsAssetData/Private/Widgets/AssetDataThumbnailWidget.cpp
@@ -11,7 +11,6 @@
 #include "Columns/SlateDelegateColumns.h"
 #include "Elements/Columns/TypedElementFolderColumns.h"
 #include "Elements/Columns/TypedElementHiearchyColumns.h"
-#include "Elements/Columns/TypedElementMiscColumns.h"
 #include "Elements/Columns/TypedElementSlateWidgetColumns.h"
 #include "Interfaces/IPluginManager.h"
 #include "Styling/StyleColors.h"
@@ -139,7 +138,7 @@ namespace Private
 						.VAlign(VAlign_Center)
 						[
 							SNew(STextBlock)
-							.Text(Binder.BindData(&FNameColumn::Name, [] (FName InName)
+							.Text(Binder.BindData(&FAssetNameColumn::Name, [] (FName InName)
 							{
 								return FText::FromString(TedsAssetDataHelper::RemoveSlashFromStart(InName.ToString()));
 							}))
diff --git a/Engine/Plugins/Experimental/EditorDataStorageFeatures/Source/TedsAssetData/Private/Widgets/AssetLabelWidget.cpp b/Engine/Plugins/Experimental/EditorDataStorageFeatures/Source/TedsAssetData/Private/Widgets/AssetLabelWidget.cpp
--- a/Engine/Plugins/Experimental/EditorDataStorageFeatures/Source/TedsAssetData/Private/Widgets/AssetLabelWidget.cpp
+++ b/Engine/Plugins/Experimental/EditorDataStorageFeatures/Source/TedsAssetData/Private/Widgets/AssetLabelWidget.cpp
@@ -5,7 +5,6 @@
 #include "Elements/Framework/TypedElementAttributeBinding.h"
 #include "TedsAssetDataColumns.h"
 #include "TedsAssetDataHelper.h"
-#include "Elements/Columns/TypedElementMiscColumns.h"
 #include "Widgets/Text/STextBlock.h"
 
 #define LOCTEXT_NAMESPACE "FAssetLabelWidgetConstructor"
@@ -16,7 +15,7 @@ void UAssetLabelWidgetFactory::RegisterWidgetConstructors(UE::Editor::DataStorag
 	using namespace UE::Editor::DataStorage::Queries;
 	DataStorageUi.RegisterWidgetFactory<FAssetLabelWidgetConstructor>(
 		DataStorageUi.FindPurpose(IUiProvider::FPurposeInfo("ContentBrowser", "RowLabel", NAME_None).GeneratePurposeID()),
-		TColumn<FNameColumn>());
+		TColumn<FAssetNameColumn>());
 }
 
 FAssetLabelWidgetConstructor::FAssetLabelWidgetConstructor()
@@ -38,7 +37,7 @@ TSharedPtr<SWidget> FAssetLabelWidgetConstructor::CreateWidget(UE::Editor::DataS
 
 	TAttribute<FText> AssetName;
 
-	if (const FNameColumn* AssetNameColumn = DataStorage->GetColumn<FNameColumn>(TargetRow))
+	if (const FAssetNameColumn* AssetNameColumn = DataStorage->GetColumn<FAssetNameColumn>(TargetRow))
 	{
 		const FString AssetNameString = TedsAssetDataHelper::RemoveSlashFromStart(AssetNameColumn->Name.ToString());
 		AssetName = FText::FromString(AssetNameString);
@@ -50,7 +49,7 @@ TSharedPtr<SWidget> FAssetLabelWidgetConstructor::CreateWidget(UE::Editor::DataS
 				.ToolTipText(Binder.BindTextFormat(
 								LOCTEXT("AssetLabelTooltip", 
 									"{Name}\n\nVirtual path: {VirtualPath}\n  Asset path: {AssetPath}\n  Verse path: {VersePath}"))
-								.Arg(TEXT("Name"), &FNameColumn::Name)
+								.Arg(TEXT("Name"), &FAssetNameColumn::Name)
 								.Arg(TEXT("VirtualPath"), &FVirtualPathColumn_Experimental::VirtualPath, LOCTEXT("PathNotSet", "<not set>"))
 								.Arg(TEXT("AssetPath"), &FAssetPathColumn_Experimental::Path, LOCTEXT("PathNotSet", "<not set>"))
 								.Arg(TEXT("VersePath"), &FVersePathColumn::VersePath, 
diff --git a/Engine/Plugins/Experimental/EditorDataStorageFeatures/Source/TedsAssetData/Private/Widgets/AssetNameWidget.cpp b/Engine/Plugins/Experimental/EditorDataStorageFeatures/Source/TedsAssetData/Private/Widgets/AssetNameWidget.cpp
new file mode 100644
--- /dev/null
+++ b/Engine/Plugins/Experimental/EditorDataStorageFeatures/Source/TedsAssetData/Private/Widgets/AssetNameWidget.cpp
@@ -0,0 +1,91 @@
+// Copyright Epic Games, Inc. All Rights Reserved.
+
+#include "Widgets/AssetNameWidget.h"
+
+#include "TedsAssetDataColumns.h"
+#include "Elements/Framework/TypedElementAttributeBinding.h"
+#include "Internationalization/Internationalization.h"
+#include "Widgets/Text/STextBlock.h"
+
+#define LOCTEXT_NAMESPACE "AssetNameWidget"
+
+namespace UE::Editor::DataStorage
+{
+	//
+	// No slash
+	//
+
+	FText FAssetNameWidgetSorter_NoSlash::GetShortName() const
+	{
+		return LOCTEXT("FAssetNameWidgetSorter_NoSlashName", "No slash");
+	}
+	
+	int32 FAssetNameWidgetSorter_NoSlash::Compare(const FAssetNameColumn& Left, const FAssetNameColumn& Right) const
+	{
+		return TSortNameView(TSortByName<ESortByNameFlags::RemoveLeadingSlash>{}, Left.Name).Compare(Right.Name);
+	}
+
+	FPrefixInfo FAssetNameWidgetSorter_NoSlash::CalculatePrefix(const FAssetNameColumn& Column, uint32 ByteIndex) const
+	{
+		return CreateSortPrefix(ByteIndex, TSortNameView(TSortByName<ESortByNameFlags::RemoveLeadingSlash>{}, Column.Name));
+	}
+
+	//
+	// With slash
+	//
+
+	FText FAssetNameWidgetSorter_WithSlash::GetShortName() const
+	{
+		return LOCTEXT("FAssetNameWidgetSorter_WithSlashName", "With slash");
+	}
+	
+	int32 FAssetNameWidgetSorter_WithSlash::Compare(const FAssetNameColumn& Left, const FAssetNameColumn& Right) const
+	{
+		return TSortNameView(TSortByName<>{}, Left.Name).Compare(Right.Name);
+	}
+
+	FPrefixInfo FAssetNameWidgetSorter_WithSlash::CalculatePrefix(const FAssetNameColumn& Column, uint32 ByteIndex) const
+	{
+		return CreateSortPrefix(ByteIndex, TSortNameView(TSortByName<>{}, Column.Name));
+	}
+} // namespace UE::Editor::DataStorage
+
+void UAssetNameWidgetFactory::RegisterWidgetConstructors(UE::Editor::DataStorage::ICoreProvider& DataStorage,
+	UE::Editor::DataStorage::IUiProvider& DataStorageUi) const
+{
+	using namespace UE::Editor::DataStorage::Queries;
+	DataStorageUi.RegisterWidgetFactory<FAssetNameWidgetConstructor>(
+		DataStorageUi.FindPurpose(DataStorageUi.GetGeneralWidgetPurposeID()),
+		TColumn<FAssetNameColumn>());
+}
+
+FAssetNameWidgetConstructor::FAssetNameWidgetConstructor()
+	: FSimpleWidgetConstructor(StaticStruct())
+{
+}
+
+TSharedPtr<SWidget> FAssetNameWidgetConstructor::CreateWidget(UE::Editor::DataStorage::ICoreProvider* DataStorage,
+	UE::Editor::DataStorage::IUiProvider* DataStorageUi, UE::Editor::DataStorage::RowHandle TargetRow, UE::Editor::DataStorage::RowHandle WidgetRow,
+	const UE::Editor::DataStorage::FMetaDataView& Arguments)
+{
+	UE::Editor::DataStorage::FAttributeBinder Binder(TargetRow, DataStorage);
+
+	return SNew(STextBlock)
+				.Text(Binder.BindText(&FAssetNameColumn::Name))
+				.ToolTipText(Binder.BindText(&FAssetNameColumn::Name));
+}
+
+TArray<TSharedPtr<const UE::Editor::DataStorage::FColumnSorterInterface>> FAssetNameWidgetConstructor::ConstructColumnSorters(
+	UE::Editor::DataStorage::ICoreProvider* DataStorage,
+	UE::Editor::DataStorage::IUiProvider* DataStorageUi,
+	const UE::Editor::DataStorage::FMetaDataView& Arguments)
+{
+	using namespace UE::Editor::DataStorage;
+	return TArray<TSharedPtr<const FColumnSorterInterface>>(
+		{
+			MakeShared<FAssetNameWidgetSorter_NoSlash>(),
+			MakeShared<FAssetNameWidgetSorter_WithSlash>()
+		});
+}
+
+#undef LOCTEXT_NAMESPACE
diff --git a/Engine/Plugins/Experimental/EditorDataStorageFeatures/Source/TedsAssetData/Private/Widgets/AssetPreview/STedsAssetPreviewWidget.cpp b/Engine/Plugins/Experimental/EditorDataStorageFeatures/Source/TedsAssetData/Private/Widgets/AssetPreview/STedsAssetPreviewWidget.cpp
--- a/Engine/Plugins/Experimental/EditorDataStorageFeatures/Source/TedsAssetData/Private/Widgets/AssetPreview/STedsAssetPreviewWidget.cpp
+++ b/Engine/Plugins/Experimental/EditorDataStorageFeatures/Source/TedsAssetData/Private/Widgets/AssetPreview/STedsAssetPreviewWidget.cpp
@@ -312,7 +312,7 @@ TArray<TWeakObjectPtr<const UScriptStruct>> STedsAssetPreviewWidget::GetBasicInf
 	{
 		TWeakObjectPtr(FAssetTag::StaticStruct()),
 		TWeakObjectPtr(FFolderTag::StaticStruct()),
-		TWeakObjectPtr(FNameColumn::StaticStruct())
+		TWeakObjectPtr(FAssetNameColumn::StaticStruct())
 	});
 	return BasicInfoColumns;
 }
diff --git a/Engine/Plugins/Experimental/EditorDataStorageFeatures/Source/TedsAssetData/Private/Widgets/AssetPreview/Widgets/BaseInfo/AssetPreviewBaseInfoWidget.cpp b/Engine/Plugins/Experimental/EditorDataStorageFeatures/Source/TedsAssetData/Private/Widgets/AssetPreview/Widgets/BaseInfo/AssetPreviewBaseInfoWidget.cpp
--- a/Engine/Plugins/Experimental/EditorDataStorageFeatures/Source/TedsAssetData/Private/Widgets/AssetPreview/Widgets/BaseInfo/AssetPreviewBaseInfoWidget.cpp
+++ b/Engine/Plugins/Experimental/EditorDataStorageFeatures/Source/TedsAssetData/Private/Widgets/AssetPreview/Widgets/BaseInfo/AssetPreviewBaseInfoWidget.cpp
@@ -1,4 +1,4 @@
-﻿// Copyright Epic Games, Inc. All Rights Reserved.
+// Copyright Epic Games, Inc. All Rights Reserved.
 
 #include "AssetPreviewBaseInfoWidget.h"
 
@@ -25,7 +25,7 @@ void UAssetPreviewBaseInfoWidgetFactory::RegisterWidgetConstructors(UE::Editor::
 
 	DataStorageUi.RegisterWidgetFactory<FAssetPreviewBaseInfoWidgetConstructor>(
 		DataStorageUi.FindPurpose(IUiProvider::FPurposeInfo("AssetPreview", "Default", NAME_None).GeneratePurposeID()),
-		(TColumn<FAssetTag>() || TColumn<FFolderTag>()) && TColumn<FNameColumn>());
+		(TColumn<FAssetTag>() || TColumn<FFolderTag>()) && TColumn<FAssetNameColumn>());
 }
 
 FAssetPreviewBaseInfoWidgetConstructor::FAssetPreviewBaseInfoWidgetConstructor()
@@ -43,7 +43,7 @@ TSharedPtr<SWidget> FAssetPreviewBaseInfoWidgetConstructor::CreateWidget(UE::Edi
 
 	// Asset Name
 	TSharedRef<SVerticalBox> BaseInfoBox = SNew(SVerticalBox);
-	if (DataStorage->HasColumns<FNameColumn>(TargetRow))
+	if (DataStorage->HasColumns<FAssetNameColumn>(TargetRow))
 	{
 		const FName NamePurposeTableName = TEXT("Editor_WidgetTable");
 
@@ -142,7 +142,7 @@ TSharedPtr<SWidget> FAssetPreviewBaseInfoWidgetConstructor::CreateWidget(UE::Edi
 
 TArray<TWeakObjectPtr<const UScriptStruct>> FAssetPreviewBaseInfoWidgetConstructor::GetLabelColumns()
 {
-	static TArray<TWeakObjectPtr<const UScriptStruct>> LabelColumns({ TWeakObjectPtr(FNameColumn::StaticStruct()) });
+	static TArray<TWeakObjectPtr<const UScriptStruct>> LabelColumns({ TWeakObjectPtr(FAssetNameColumn::StaticStruct()) });
 	return LabelColumns;
 }
 
diff --git a/Engine/Plugins/Experimental/EditorDataStorageFeatures/Source/TedsAssetData/Private/Widgets/ContentBrowserListViewNameWidget.cpp b/Engine/Plugins/Experimental/EditorDataStorageFeatures/Source/TedsAssetData/Private/Widgets/ContentBrowserListViewNameWidget.cpp
--- a/Engine/Plugins/Experimental/EditorDataStorageFeatures/Source/TedsAssetData/Private/Widgets/ContentBrowserListViewNameWidget.cpp
+++ b/Engine/Plugins/Experimental/EditorDataStorageFeatures/Source/TedsAssetData/Private/Widgets/ContentBrowserListViewNameWidget.cpp
@@ -40,7 +40,7 @@ void UContentBrowserListViewNameWidgetFactory::RegisterWidgetConstructors(UE::Ed
 	using namespace UE::Editor::DataStorage::Queries;
 	DataStorageUi.RegisterWidgetFactory<FContentBrowserListViewNameWidgetConstructor>(
 		DataStorageUi.FindPurpose(WidgetPurpose),
-		TColumn<FNameColumn>() && (TColumn<FAssetTag>() || TColumn<FFolderTag>()));
+		TColumn<FAssetNameColumn>() && (TColumn<FAssetTag>() || TColumn<FFolderTag>()));
 }
 
 FContentBrowserListViewNameWidgetConstructor::FContentBrowserListViewNameWidgetConstructor()
@@ -180,7 +180,7 @@ TSharedPtr<SWidget> FContentBrowserListViewNameWidgetConstructor::CreateWidget(U
 
 TArray<TWeakObjectPtr<const UScriptStruct>> FContentBrowserListViewNameWidgetConstructor::GetLabelColumns()
 {
-	static TArray<TWeakObjectPtr<const UScriptStruct>> LabelColumns({ TWeakObjectPtr(FNameColumn::StaticStruct()) });
+	static TArray<TWeakObjectPtr<const UScriptStruct>> LabelColumns({ TWeakObjectPtr(FAssetNameColumn::StaticStruct()) });
 	return LabelColumns;
 }
 
diff --git a/Engine/Plugins/Experimental/EditorDataStorageFeatures/Source/TedsAssetData/Public/TedsAssetDataColumns.h b/Engine/Plugins/Experimental/EditorDataStorageFeatures/Source/TedsAssetData/Public/TedsAssetDataColumns.h
--- a/Engine/Plugins/Experimental/EditorDataStorageFeatures/Source/TedsAssetData/Public/TedsAssetDataColumns.h
+++ b/Engine/Plugins/Experimental/EditorDataStorageFeatures/Source/TedsAssetData/Public/TedsAssetDataColumns.h
@@ -76,6 +76,17 @@ struct FPublicAssetTag final : public FEditorDataStorageTag
 	GENERATED_BODY()
 };
 
+// Column to store the name of an asset
+USTRUCT(meta = (DisplayName = "Name"))
+struct FAssetNameColumn final : public FEditorDataStorageColumn
+{
+	GENERATED_BODY()
+
+	// The name of the asset
+	UPROPERTY()
+	FName Name;
+};
+
 // Column to store the type of an asset
 USTRUCT(meta = (DisplayName = "Type"))
 struct FAssetClassColumn final : public FEditorDataStorageColumn
diff --git a/Engine/Plugins/Experimental/EditorDataStorageFeatures/Source/TedsAssetData/Public/Widgets/AssetNameWidget.h b/Engine/Plugins/Experimental/EditorDataStorageFeatures/Source/TedsAssetData/Public/Widgets/AssetNameWidget.h
new file mode 100644
--- /dev/null
+++ b/Engine/Plugins/Experimental/EditorDataStorageFeatures/Source/TedsAssetData/Public/Widgets/AssetNameWidget.h
@@ -0,0 +1,71 @@
+// Copyright Epic Games, Inc. All Rights Reserved.
+
+#pragma once
+
+#include "TedsAssetDataColumns.h"
+#include "Elements/Interfaces/TypedElementDataStorageFactory.h"
+#include "Elements/Interfaces/TypedElementDataStorageUiInterface.h"
+#include "Internationalization/Text.h"
+#include "UObject/ObjectMacros.h"
+
+#include "AssetNameWidget.generated.h"
+
+class IEditorDataStorageProvider;
+class UScriptStruct;
+
+namespace UE::Editor::DataStorage
+{
+	class FAssetNameWidgetSorter_NoSlash final : public TColumnSorterInterface<FColumnSorterInterface::ESortType::HybridSort, FAssetNameColumn>
+	{
+	public:
+		virtual FText GetShortName() const override;
+
+	protected:
+		virtual int32 Compare(const FAssetNameColumn& Left, const FAssetNameColumn& Right) const override;
+		virtual FPrefixInfo CalculatePrefix(const FAssetNameColumn& Column, uint32 ByteIndex) const override;
+	};
+
+	class FAssetNameWidgetSorter_WithSlash final : public TColumnSorterInterface<FColumnSorterInterface::ESortType::HybridSort, FAssetNameColumn>
+	{
+	public:
+		virtual FText GetShortName() const override;
+
+	protected:
+		virtual int32 Compare(const FAssetNameColumn& Left, const FAssetNameColumn& Right) const override;
+		virtual FPrefixInfo CalculatePrefix(const FAssetNameColumn& Column, uint32 ByteIndex) const override;
+	};
+}
+
+UCLASS()
+class UAssetNameWidgetFactory : public UEditorDataStorageFactory
+{
+	GENERATED_BODY()
+
+public:
+	virtual ~UAssetNameWidgetFactory() override = default;
+
+	TEDSASSETDATA_API virtual void RegisterWidgetConstructors(UE::Editor::DataStorage::ICoreProvider& DataStorage,
+		UE::Editor::DataStorage::IUiProvider& DataStorageUi) const override;
+};
+
+// Label for assets in TEDS
+USTRUCT()
+struct FAssetNameWidgetConstructor : public FSimpleWidgetConstructor
+{
+	GENERATED_BODY()
+
+public:
+	TEDSASSETDATA_API FAssetNameWidgetConstructor();
+
+	TEDSASSETDATA_API virtual TSharedPtr<SWidget> CreateWidget(
+		UE::Editor::DataStorage::ICoreProvider* DataStorage,
+		UE::Editor::DataStorage::IUiProvider* DataStorageUi,
+		UE::Editor::DataStorage::RowHandle TargetRow,
+		UE::Editor::DataStorage::RowHandle WidgetRow, 
+		const UE::Editor::DataStorage::FMetaDataView& Arguments) override;
+
+	virtual TArray<TSharedPtr<const UE::Editor::DataStorage::FColumnSorterInterface>> ConstructColumnSorters(
+		UE::Editor::DataStorage::ICoreProvider* DataStorage,
+		UE::Editor::DataStorage::IUiProvider* DataStorageUi,
+		const UE::Editor::DataStorage::FMetaDataView& Arguments) override;
+};
diff --git a/Engine/Plugins/Experimental/EditorDataStorageFeatures/Source/TedsContentBrowser/Private/TedsContentBrowserModule.cpp b/Engine/Plugins/Experimental/EditorDataStorageFeatures/Source/TedsContentBrowser/Private/TedsContentBrowserModule.cpp
--- a/Engine/Plugins/Experimental/EditorDataStorageFeatures/Source/TedsContentBrowser/Private/TedsContentBrowserModule.cpp
+++ b/Engine/Plugins/Experimental/EditorDataStorageFeatures/Source/TedsContentBrowser/Private/TedsContentBrowserModule.cpp
@@ -9,7 +9,6 @@
 #include "DataStorage/Features.h"
 #include "DataStorage/Queries/Types.h"
 #include "Elements/Columns/TypedElementFolderColumns.h"
-#include "Elements/Columns/TypedElementMiscColumns.h"
 #include "Elements/Interfaces/TypedElementDataStorageInterface.h"
 #include "Experimental/ContentBrowserViewExtender.h"
 #include "Modules/ModuleManager.h"
@@ -231,7 +230,7 @@ namespace UE::Editor::ContentBrowser
 					.CellWidgetPurpose(IUiProvider::FPurposeInfo("ContentBrowser", "RowLabel", NAME_None).GeneratePurposeID())
 					.HeaderWidgetPurpose(IUiProvider::FPurposeInfo("General", "Header", NAME_None).GeneratePurposeID())
 					// Default list of columns to display
-					.Columns({ FNameColumn::StaticStruct(), FTedsAlertColumn::StaticStruct(),
+					.Columns({ FAssetNameColumn::StaticStruct(), FTedsAlertColumn::StaticStruct(),
 						FAssetClassColumn::StaticStruct(), FAssetTag::StaticStruct(), FFolderTag::StaticStruct(), FAssetPathColumn_Experimental::StaticStruct(),
 						FDiskSizeColumn::StaticStruct(), FVirtualPathColumn_Experimental::StaticStruct(),
 						DynamicSkeletalMeshSkeletonColumn, DynamicBlueprintColumn, DynamicSourceTextureColumn,
diff --git a/Engine/Plugins/Experimental/EditorDataStorageFeatures/Source/TedsContentBrowser/Private/Widgets/ContentBrowserTileViewWidget.cpp b/Engine/Plugins/Experimental/EditorDataStorageFeatures/Source/TedsContentBrowser/Private/Widgets/ContentBrowserTileViewWidget.cpp
--- a/Engine/Plugins/Experimental/EditorDataStorageFeatures/Source/TedsContentBrowser/Private/Widgets/ContentBrowserTileViewWidget.cpp
+++ b/Engine/Plugins/Experimental/EditorDataStorageFeatures/Source/TedsContentBrowser/Private/Widgets/ContentBrowserTileViewWidget.cpp
@@ -191,7 +191,7 @@ TSharedPtr<SWidget> FContentBrowserTileViewWidgetConstructor::CreateWidget(UE::E
 									SNew(STextBlock)
 									.ColorAndOpacity(FStyleColors::White)
 									.OverflowPolicy(ETextOverflowPolicy::Ellipsis)
-									.Text(RowBinder.BindData(&FNameColumn::Name, [] (FName InName)
+									.Text(RowBinder.BindData(&FAssetNameColumn::Name, [] (FName InName)
 									{
 										return FText::FromString(TedsAssetDataHelper::RemoveSlashFromStart(InName.ToString()));
 									}))
@@ -219,7 +219,7 @@ TSharedPtr<SWidget> FContentBrowserTileViewWidgetConstructor::CreateWidget(UE::E
 									[
 										SNew(SInlineEditableTextBlock)
 											.Font(GetThumbnailFont())
-											.Text(RowBinder.BindData(&FNameColumn::Name, [] (FName InName)
+											.Text(RowBinder.BindData(&FAssetNameColumn::Name, [] (FName InName)
 											{
 												return FText::FromString(TedsAssetDataHelper::RemoveSlashFromStart(InName.ToString()));
 											}))
