More loc text identifier copy/pasta
#rnx

diff --git a/Engine/Plugins/Experimental/EditorDataStorageFeatures/Source/TedsOutliner/Private/Widgets/OutlinerIconWidget.cpp b/Engine/Plugins/Experimental/EditorDataStorageFeatures/Source/TedsOutliner/Private/Widgets/OutlinerIconWidget.cpp
--- a/Engine/Plugins/Experimental/EditorDataStorageFeatures/Source/TedsOutliner/Private/Widgets/OutlinerIconWidget.cpp
+++ b/Engine/Plugins/Experimental/EditorDataStorageFeatures/Source/TedsOutliner/Private/Widgets/OutlinerIconWidget.cpp
@@ -72,7 +72,7 @@ namespace UE::OutlinerIconWidget::Local
 			return LOCTEXT("OverrideAddedTooltip", "This entity has been added.");
 
 		case EOverriddenState::AllOverridden:
-			return LOCTEXT("OverrideAddedTooltip", "This entity has all its properties overridden.");
+			return LOCTEXT("OverrideAllTooltip", "This entity has all its properties overridden.");
 
 		case EOverriddenState::HasOverrides:
 			return LOCTEXT("OverrideInsideTooltip", "At least one property or child has an override.");
