[TEDS Debugger] Add an override for the CreateWidgetDisplayNameText for Row Handles as the column header was dispalying 'TEDS Column" in the debugger

diff --git a/Engine/Plugins/Experimental/EditorDataStorageFeatures/Source/TedsTableViewer/Private/Widgets/RowHandleWidget.cpp b/Engine/Plugins/Experimental/EditorDataStorageFeatures/Source/TedsTableViewer/Private/Widgets/RowHandleWidget.cpp
--- a/Engine/Plugins/Experimental/EditorDataStorageFeatures/Source/TedsTableViewer/Private/Widgets/RowHandleWidget.cpp
+++ b/Engine/Plugins/Experimental/EditorDataStorageFeatures/Source/TedsTableViewer/Private/Widgets/RowHandleWidget.cpp
@@ -107,6 +107,12 @@ bool FRowHandleWidgetConstructor::FinalizeWidget(UE::Editor::DataStorage::ICoreP
 	return true;
 }
 
+FText FRowHandleWidgetConstructor::CreateWidgetDisplayNameText(UE::Editor::DataStorage::ICoreProvider* DataStorage, 
+	UE::Editor::DataStorage::RowHandle Row) const
+{
+	return LOCTEXT("RowHandleColumnName", "Row Handle");
+}
+
 TArray<TSharedPtr<const UE::Editor::DataStorage::FColumnSorterInterface>> FRowHandleWidgetConstructor::ConstructColumnSorters(
 	UE::Editor::DataStorage::ICoreProvider* DataStorage,
 	UE::Editor::DataStorage::IUiProvider* DataStorageUi,
diff --git a/Engine/Plugins/Experimental/EditorDataStorageFeatures/Source/TedsTableViewer/Public/Widgets/RowHandleWidget.h b/Engine/Plugins/Experimental/EditorDataStorageFeatures/Source/TedsTableViewer/Public/Widgets/RowHandleWidget.h
--- a/Engine/Plugins/Experimental/EditorDataStorageFeatures/Source/TedsTableViewer/Public/Widgets/RowHandleWidget.h
+++ b/Engine/Plugins/Experimental/EditorDataStorageFeatures/Source/TedsTableViewer/Public/Widgets/RowHandleWidget.h
@@ -60,6 +60,9 @@ protected:
 		UE::Editor::DataStorage::IUiProvider* DataStorageUi,
 		UE::Editor::DataStorage::RowHandle Row,
 		const TSharedPtr<SWidget>& Widget) override;
+		
+	UE_API virtual FText CreateWidgetDisplayNameText(UE::Editor::DataStorage::ICoreProvider* DataStorage, 
+    		UE::Editor::DataStorage::RowHandle Row = UE::Editor::DataStorage::InvalidRowHandle) const override;
 
 	UE_API virtual TArray<TSharedPtr<const UE::Editor::DataStorage::FColumnSorterInterface>> ConstructColumnSorters(
 		UE::Editor::DataStorage::ICoreProvider* DataStorage,
