[TEDS] Replaced the implementation of the filter node with the new FilterByRow call.

diff --git a/Engine/Plugins/Experimental/EditorDataStorageFeatures/Source/TedsEverythingPicker/Private/Context/TedsPickerContextUtil.cpp b/Engine/Plugins/Experimental/EditorDataStorageFeatures/Source/TedsEverythingPicker/Private/Context/TedsPickerContextUtil.cpp
--- a/Engine/Plugins/Experimental/EditorDataStorageFeatures/Source/TedsEverythingPicker/Private/Context/TedsPickerContextUtil.cpp
+++ b/Engine/Plugins/Experimental/EditorDataStorageFeatures/Source/TedsEverythingPicker/Private/Context/TedsPickerContextUtil.cpp
@@ -8,6 +8,8 @@
 #include "Elements/Columns/TypedElementLabelColumns.h"
 #include "Elements/Columns/TypedElementTypeInfoColumns.h"
 
+#include "Elements/Framework/TypedElementQueryContext.h"
+
 #include "Widgets/STedsHierarchyViewer.h"
 #include "Widgets/SEverythingPicker.h"
 
@@ -18,6 +20,8 @@ namespace UE::Editor::DataStorage::Picker
 {
 	void SObjectReferenceContextView::Construct(const FArguments& InArgs)
 	{
+		using namespace UE::Editor::DataStorage::Queries;
+
 		ICoreProvider* Storage = GetMutableDataStorageFeature<ICoreProvider>(UE::Editor::DataStorage::StorageFeatureName);
 
 		ensure(Storage);
@@ -30,9 +34,12 @@ namespace UE::Editor::DataStorage::Picker
 				ReferenceQueryNode, QueryStack::FRowQueryResultsNode::ESyncFlags::RefreshOnUpdate);
 
 		// Node that filters out the ReferenceQueryResultsNode to generate all the top level nodes
-		constexpr bool bRowsShouldHaveColumn = false;
 		TSharedRef<QueryStack::IRowNode> TopLevelNode =
-			MakeShared<QueryStack::FRowFilterNode<FTableRowParentColumn>> (Storage, ReferenceQueryResultsNode.ToSharedPtr(), bRowsShouldHaveColumn);
+			MakeShared<QueryStack::FRowFilterNode>(Storage, ReferenceQueryResultsNode.ToSharedPtr(),
+				[](TQueryContext<SingleRowInfo>& Context)
+				{
+					return !Context.CurrentRowHasColumns<FTableRowParentColumn>();
+				});
 
 		// Using the Teds Outliner Purpose for matching customization for now
 		IUiProvider::FPurposeID PurposeId = IUiProvider::FPurposeInfo("SceneOutliner", "RowLabel", NAME_None).GeneratePurposeID();
diff --git a/Engine/Plugins/Experimental/EditorDataStorageFeatures/Source/TedsQueryStack/Private/TedsRowFilterNode.cpp b/Engine/Plugins/Experimental/EditorDataStorageFeatures/Source/TedsQueryStack/Private/TedsRowFilterNode.cpp
new file mode 100644
--- /dev/null
+++ b/Engine/Plugins/Experimental/EditorDataStorageFeatures/Source/TedsQueryStack/Private/TedsRowFilterNode.cpp
@@ -0,0 +1,32 @@
+// Copyright Epic Games, Inc. All Rights Reserved
+
+#include "TedsRowFilterNode.h"
+
+namespace UE::Editor::DataStorage::QueryStack
+{
+	FRowHandleArrayView FRowFilterNode::GetRows() const
+	{
+		return Rows.GetRows();
+	}
+
+	FRowHandleArray& FRowFilterNode::GetMutableRows()
+	{
+		return Rows;
+	}
+
+	IRowNode::RevisionId FRowFilterNode::GetRevision() const
+	{
+		return CachedParentRevisionID;
+	}
+
+	void FRowFilterNode::Update()
+	{
+		if (CachedParentRevisionID != ParentRowNode->GetRevision())
+		{
+			Rows.Reset();
+			Storage->FilterRowsBy(Rows, ParentRowNode->GetRows(), Filter);
+
+			CachedParentRevisionID = ParentRowNode->GetRevision();
+		}
+	}
+} // namespace UE::Editor::DataStorage::QueryStack
diff --git a/Engine/Plugins/Experimental/EditorDataStorageFeatures/Source/TedsQueryStack/Public/TedsRowFilterNode.h b/Engine/Plugins/Experimental/EditorDataStorageFeatures/Source/TedsQueryStack/Public/TedsRowFilterNode.h
--- a/Engine/Plugins/Experimental/EditorDataStorageFeatures/Source/TedsQueryStack/Public/TedsRowFilterNode.h
+++ b/Engine/Plugins/Experimental/EditorDataStorageFeatures/Source/TedsQueryStack/Public/TedsRowFilterNode.h
@@ -4,6 +4,7 @@
 
 #include "DataStorage/CommonTypes.h"
 #include "Elements/Interfaces/TypedElementDataStorageInterface.h"
+#include "Elements/Framework/TypedElementQueryFunctions.h"
 #include "TedsQueryStackInterfaces.h"
 
 namespace UE::Editor::DataStorage
@@ -14,83 +15,34 @@ namespace UE::Editor::DataStorage
 namespace UE::Editor::DataStorage::QueryStack
 {
 	/**
-	 * A specialized query stack node that takes another row node as input and only keeps rows that contain (or don't contain) a specific row
+	 * Query stack node that filters out rows based on the provided filter.
+	 * Performance-wise this will be slower than using filters on queries, but this node can be used when using
+	 * query level filters is not possible.
 	 */
-	template<TColumnType ColumnType>
-	class FRowFilterNode : public IRowNode
+	class FRowFilterNode final : public IRowNode
 	{
 	public:
-
-		/**
-		 * Construct an FHierarchyTopLevelRowNode
-		 *
-		 * @param InStorage	The root elements of the new graph to be generated
-		 * @param InParentRowNode Different visualization settings, such as whether it should display the referencers or the dependencies of the NewGraphRootIdentifiers
-		 * @param bInRowsShouldHaveColumn If true, only keep rows that have the required column. If false, only keep rows that DON'T have the required column 
-		 */
-		FRowFilterNode(ICoreProvider* InStorage, const TSharedPtr<IRowNode>& InParentRowNode, bool bInRowsShouldHaveColumn)
-		: Storage(InStorage)
-		, ParentRowNode(InParentRowNode)
-		, CachedParentRevisionID(InParentRowNode->GetRevision())
-		, bRowsShouldHaveColumn(bInRowsShouldHaveColumn)
-		{
-			
-		}
-		
-		virtual ~FRowFilterNode() override = default;
-
-		// IRowNode interface
-		virtual FRowHandleArrayView GetRows() const override
-		{
-			return Rows.GetRows();
-		}
-		
-		virtual FRowHandleArray& GetMutableRows() override
-		{
-			return Rows;
-		}
-		
-		virtual RevisionId GetRevision() const override
-		{
-			return ParentRowNode->GetRevision();
-		}
-		
-		virtual void Update() override
+		template<typename FilterFunction>
+		FRowFilterNode(ICoreProvider* Storage, const TSharedPtr<IRowNode>& ParentRowNode, FilterFunction&& Filter)
+			: Storage(Storage)
+			, Filter(Queries::BuildQueryFunction<bool>(Forward<FilterFunction>(Filter)))
+			, ParentRowNode(ParentRowNode)
+			, CachedParentRevisionID(ParentRowNode->GetRevision())
 		{
-			if (CachedParentRevisionID != ParentRowNode->GetRevision())
-			{
-				UpdateRows();
-				CachedParentRevisionID = ParentRowNode->GetRevision();
-			}
 		}
-		// ~IRowNode interface
-
-	protected:
-
-		// Update our list of top level rows
-		void UpdateRows()
-		{
-			Rows.Reset();
-
-			FRowHandleArrayView ParentRows = ParentRowNode->GetRows();
 
-			for (RowHandle Row : ParentRows)
-			{
-				const bool bHasColumn = Storage->HasColumns<ColumnType>(Row);
-
-				if (bRowsShouldHaveColumn ? bHasColumn : !bHasColumn)
-				{
-					Rows.Add(Row);
-				}
-			}
-		}
+		virtual ~FRowFilterNode() override = default;
 
-	protected:
+		TEDSQUERYSTACK_API virtual FRowHandleArrayView GetRows() const override;
+		TEDSQUERYSTACK_API virtual FRowHandleArray& GetMutableRows() override;
+		TEDSQUERYSTACK_API virtual RevisionId GetRevision() const override;
+		TEDSQUERYSTACK_API virtual void Update() override;
 
-		FRowHandleArray Rows;
+	private:
 		ICoreProvider* Storage;
+		Queries::TQueryFunction<bool> Filter;
+		FRowHandleArray Rows;
 		TSharedPtr<IRowNode> ParentRowNode;
 		RevisionId CachedParentRevisionID;
-		bool bRowsShouldHaveColumn;
 	};
 }
