Move TypedElement Bridge into EditorDataStorageFeatures
* Helps to separate TEDS and TypedElement

diff --git a/Engine/Plugins/Experimental/EditorDataStorage/Source/TedsCore/Private/Compatibility/TedsTypedElementActorHandleFactory.cpp b/Engine/Plugins/Experimental/EditorDataStorage/Source/TedsCore/Private/Compatibility/TedsTypedElementActorHandleFactory.cpp
deleted file mode 100644
diff --git a/Engine/Plugins/Experimental/EditorDataStorage/Source/TedsCore/Private/Compatibility/TedsTypedElementActorHandleFactory.h b/Engine/Plugins/Experimental/EditorDataStorage/Source/TedsCore/Private/Compatibility/TedsTypedElementActorHandleFactory.h
deleted file mode 100644
diff --git a/Engine/Plugins/Experimental/EditorDataStorage/Source/TedsCore/Private/Compatibility/TedsTypedElementBridge.cpp b/Engine/Plugins/Experimental/EditorDataStorage/Source/TedsCore/Private/Compatibility/TedsTypedElementBridge.cpp
deleted file mode 100644
diff --git a/Engine/Plugins/Experimental/EditorDataStorage/Source/TedsCore/Private/Compatibility/TedsTypedElementBridgeQueries.cpp b/Engine/Plugins/Experimental/EditorDataStorage/Source/TedsCore/Private/Compatibility/TedsTypedElementBridgeQueries.cpp
deleted file mode 100644
diff --git a/Engine/Plugins/Experimental/EditorDataStorage/Source/TedsCore/Private/Compatibility/TedsTypedElementBridgeQueries.h b/Engine/Plugins/Experimental/EditorDataStorage/Source/TedsCore/Private/Compatibility/TedsTypedElementBridgeQueries.h
deleted file mode 100644
diff --git a/Engine/Plugins/Experimental/EditorDataStorage/Source/TedsCore/Public/Compatibility/Columns/TypedElement.h b/Engine/Plugins/Experimental/EditorDataStorage/Source/TedsCore/Public/Compatibility/Columns/TypedElement.h
deleted file mode 100644
diff --git a/Engine/Plugins/Experimental/EditorDataStorage/Source/TedsCore/Public/Compatibility/TedsTypedElementBridge.h b/Engine/Plugins/Experimental/EditorDataStorage/Source/TedsCore/Public/Compatibility/TedsTypedElementBridge.h
deleted file mode 100644
diff --git a/Engine/Plugins/Experimental/EditorDataStorageFeatures/EditorDataStorageFeatures.uplugin b/Engine/Plugins/Experimental/EditorDataStorageFeatures/EditorDataStorageFeatures.uplugin
--- a/Engine/Plugins/Experimental/EditorDataStorageFeatures/EditorDataStorageFeatures.uplugin
+++ b/Engine/Plugins/Experimental/EditorDataStorageFeatures/EditorDataStorageFeatures.uplugin
@@ -62,6 +62,10 @@
 		{
 			"Name": "TedsEverythingPicker",
 			"Type": "Editor"
+		},
+		{
+			"Name": "TedsTypedElementBridge",
+			"Type": "Editor"
 		}
 	],
 	"Plugins": [
diff --git a/Engine/Plugins/Experimental/EditorDataStorageFeatures/Source/TedsTypedElementBridge/Private/TedsTypedElementBridge/TedsTypedElementActorHandleFactory.cpp b/Engine/Plugins/Experimental/EditorDataStorageFeatures/Source/TedsTypedElementBridge/Private/TedsTypedElementBridge/TedsTypedElementActorHandleFactory.cpp
new file mode 100644
--- /dev/null
+++ b/Engine/Plugins/Experimental/EditorDataStorageFeatures/Source/TedsTypedElementBridge/Private/TedsTypedElementBridge/TedsTypedElementActorHandleFactory.cpp
@@ -0,0 +1,125 @@
+// Copyright Epic Games, Inc. All Rights Reserved.
+
+#include "TedsTypedElementActorHandleFactory.h"
+
+#include "TedsTypedElementBridge/TedsTypedElementBridgeCapabilities.h"
+#include "TedsTypedElementBridge/TedsTypedElementBridgeColumns.h"
+#include "DataStorage/Features.h"
+#include "Elements/Columns/TypedElementCompatibilityColumns.h"
+#include "Elements/Framework/EngineElementsLibrary.h"
+#include "Elements/Framework/TypedElementList.h"
+#include "Elements/Framework/TypedElementQueryBuilder.h"
+#include "Elements/Interfaces/TypedElementDataStorageInterface.h"
+#include "GameFramework/Actor.h"
+
+void UTypedElementActorHandleDataStorageFactory::PreRegister(UE::Editor::DataStorage::ICoreProvider& DataStorage)
+{
+	Super::PreRegister(DataStorage);
+	
+	BridgeEnableDelegateHandle = UE::Editor::DataStorage::Compatibility::OnTypedElementBridgeEnabled().AddUObject(this, &UTypedElementActorHandleDataStorageFactory::HandleBridgeEnabled);
+}
+
+void UTypedElementActorHandleDataStorageFactory::PreShutdown(UE::Editor::DataStorage::ICoreProvider& DataStorage)
+{
+	UE::Editor::DataStorage::Compatibility::OnTypedElementBridgeEnabled().Remove(BridgeEnableDelegateHandle);
+	BridgeEnableDelegateHandle.Reset();
+}
+
+void UTypedElementActorHandleDataStorageFactory::RegisterQueries(UE::Editor::DataStorage::ICoreProvider& DataStorage)
+{
+	Super::RegisterQueries(DataStorage);
+
+	if (UE::Editor::DataStorage::Compatibility::IsTypedElementBridgeEnabled())
+	{
+		RegisterQuery_ActorHandlePopulate(DataStorage);
+	}
+
+	using namespace UE::Editor::DataStorage::Queries;
+	GetAllActorsQuery = DataStorage.RegisterQuery(
+	Select()
+		.ReadOnly<FTypedElementUObjectColumn>()
+	.Where()
+		.All<FTypedElementActorTag>()
+	.Compile());
+}
+
+void UTypedElementActorHandleDataStorageFactory::RegisterQuery_ActorHandlePopulate(UE::Editor::DataStorage::ICoreProvider& DataStorage)
+{
+	using namespace UE::Editor::DataStorage::Queries;
+
+	if (!ensureMsgf(ActorHandlePopulateQuery == InvalidQueryHandle, TEXT("Already registered query")))
+	{
+		return;
+	}
+	
+	ActorHandlePopulateQuery = DataStorage.RegisterQuery(
+	Select(TEXT("Populate actor typed element handles"),
+		FObserver::OnAdd<FTypedElementUObjectColumn>(),
+		[](IQueryContext& Context, RowHandle Row, const FTypedElementUObjectColumn& ObjectColumn)
+		{
+			if (UObject* Object = ObjectColumn.Object.Get())
+			{
+				checkSlow(Cast<AActor>(Object));
+				FTypedElementHandle Handle = UEngineElementsLibrary::AcquireEditorActorElementHandle(static_cast<AActor*>(Object));
+				Context.AddColumn(Row, Compatibility::FTypedElementColumn
+				{
+					.Handle = Handle
+				});
+			}
+		})
+	.Where()
+		.All<FTypedElementActorTag>()
+	.Compile());
+}
+
+void UTypedElementActorHandleDataStorageFactory::HandleBridgeEnabled(bool bEnabled)
+{
+	using namespace UE::Editor::DataStorage;
+
+	ICoreProvider* DataStorage = GetMutableDataStorageFeature<ICoreProvider>(StorageFeatureName);
+
+	if (bEnabled)
+	{
+		using namespace Queries;
+		
+		// Populate all the rows
+		TArray<RowHandle> CollatedRowHandles;
+		TArray<TWeakObjectPtr<const AActor>> Actors;
+		
+		DataStorage->RunQuery(GetAllActorsQuery, CreateDirectQueryCallbackBinding([&CollatedRowHandles, &Actors](IDirectQueryContext& Context, const FTypedElementUObjectColumn* Fragments)
+		{
+			TConstArrayView<RowHandle> RowHandles = Context.GetRowHandles();
+
+			CollatedRowHandles.Append(RowHandles);
+
+			TConstArrayView<const FTypedElementUObjectColumn> FragmentView(Fragments, Context.GetRowCount());
+
+			Actors.Reserve(Actors.Num() + FragmentView.Num());
+			for (const FTypedElementUObjectColumn& Fragment : FragmentView)
+			{
+				const AActor* Actor = Cast<AActor>(Fragment.Object);
+				Actors.Add(Actor);
+			}
+		}));
+
+		for (int32 Index = 0, End = CollatedRowHandles.Num(); Index < End; ++Index)
+		{
+			if (const AActor* Actor = Actors[Index].Get())
+			{
+				FTypedElementHandle Handle = UEngineElementsLibrary::AcquireEditorActorElementHandle(Actor);
+				DataStorage->AddColumn(CollatedRowHandles[Index], Compatibility::FTypedElementColumn
+				{
+					.Handle = Handle
+				});
+			}
+		}
+
+		RegisterQuery_ActorHandlePopulate(*DataStorage);
+	}
+	else
+	{
+		DataStorage->UnregisterQuery(ActorHandlePopulateQuery);
+		ActorHandlePopulateQuery = InvalidQueryHandle;
+	}
+}
+
diff --git a/Engine/Plugins/Experimental/EditorDataStorageFeatures/Source/TedsTypedElementBridge/Private/TedsTypedElementBridge/TedsTypedElementActorHandleFactory.h b/Engine/Plugins/Experimental/EditorDataStorageFeatures/Source/TedsTypedElementBridge/Private/TedsTypedElementBridge/TedsTypedElementActorHandleFactory.h
new file mode 100644
--- /dev/null
+++ b/Engine/Plugins/Experimental/EditorDataStorageFeatures/Source/TedsTypedElementBridge/Private/TedsTypedElementBridge/TedsTypedElementActorHandleFactory.h
@@ -0,0 +1,36 @@
+// Copyright Epic Games, Inc. All Rights Reserved.
+
+#pragma once
+
+#include "Elements/Framework/TypedElementHandle.h"
+#include "Elements/Interfaces/TypedElementDataStorageInterface.h"
+#include "Elements/Interfaces/TypedElementDataStorageFactory.h"
+#include "UObject/UObjectGlobals.h"
+
+#include "TedsTypedElementActorHandleFactory.generated.h"
+
+class UTypedElementRegistry;
+
+/**
+ * This class is responsible for acquiring and registering Actor Typed Element Handles
+ * with TEDS/
+ */
+UCLASS(Transient)
+class UTypedElementActorHandleDataStorageFactory : public UEditorDataStorageFactory
+{
+	GENERATED_BODY()
+public:
+
+	virtual void PreRegister(UE::Editor::DataStorage::ICoreProvider& DataStorage) override;
+	virtual void RegisterQueries(UE::Editor::DataStorage::ICoreProvider& DataStorage) override;
+	virtual void PreShutdown(UE::Editor::DataStorage::ICoreProvider& DataStorage) override;
+private:
+	virtual void RegisterQuery_ActorHandlePopulate(UE::Editor::DataStorage::ICoreProvider& DataStorage);
+	void HandleBridgeEnabled(bool bEnabled);
+
+	FDelegateHandle BridgeEnableDelegateHandle;
+	UE::Editor::DataStorage::QueryHandle ActorHandlePopulateQuery = UE::Editor::DataStorage::InvalidQueryHandle;
+	UE::Editor::DataStorage::QueryHandle GetAllActorsQuery = UE::Editor::DataStorage::InvalidQueryHandle;
+};
+
+
diff --git a/Engine/Plugins/Experimental/EditorDataStorageFeatures/Source/TedsTypedElementBridge/Private/TedsTypedElementBridge/TedsTypedElementBridgeCapabilities.cpp b/Engine/Plugins/Experimental/EditorDataStorageFeatures/Source/TedsTypedElementBridge/Private/TedsTypedElementBridge/TedsTypedElementBridgeCapabilities.cpp
new file mode 100644
--- /dev/null
+++ b/Engine/Plugins/Experimental/EditorDataStorageFeatures/Source/TedsTypedElementBridge/Private/TedsTypedElementBridge/TedsTypedElementBridgeCapabilities.cpp
@@ -0,0 +1,20 @@
+// Copyright Epic Games, Inc. All Rights Reserved.
+
+#include "TedsTypedElementBridge/TedsTypedElementBridgeCapabilities.h"
+#include "TedsTypedElementBridge/TedsTypedElementBridgeQueries.h"
+
+#include "HAL/IConsoleManager.h"
+
+namespace UE::Editor::DataStorage::Compatibility
+{
+	FOnTypedElementBridgeEnabled& OnTypedElementBridgeEnabled()
+	{
+		static FOnTypedElementBridgeEnabled OnBridgeEnabled;
+		return OnBridgeEnabled;
+	}
+
+	bool IsTypedElementBridgeEnabled()
+	{
+		return UTypedElementBridgeDataStorageFactory::IsEnabled();
+	}
+} // namespace UE::Editor::DataStorage::Compatibility
diff --git a/Engine/Plugins/Experimental/EditorDataStorageFeatures/Source/TedsTypedElementBridge/Private/TedsTypedElementBridge/TedsTypedElementBridgeModule.cpp b/Engine/Plugins/Experimental/EditorDataStorageFeatures/Source/TedsTypedElementBridge/Private/TedsTypedElementBridge/TedsTypedElementBridgeModule.cpp
new file mode 100644
--- /dev/null
+++ b/Engine/Plugins/Experimental/EditorDataStorageFeatures/Source/TedsTypedElementBridge/Private/TedsTypedElementBridge/TedsTypedElementBridgeModule.cpp
@@ -0,0 +1,20 @@
+﻿// Copyright Epic Games, Inc. All Rights Reserved.
+
+#include "Modules/ModuleManager.h"
+#include "Modules/ModuleInterface.h"
+
+class FTedsTypedElementBridgeModule : public IModuleInterface
+{
+public:
+
+	virtual void StartupModule() override
+	{
+	}
+
+	virtual void ShutdownModule() override
+	{
+	}
+};
+
+IMPLEMENT_MODULE(FTedsTypedElementBridgeModule, TedsTypedElementBridge);
+
diff --git a/Engine/Plugins/Experimental/EditorDataStorageFeatures/Source/TedsTypedElementBridge/Private/TedsTypedElementBridge/TedsTypedElementBridgeQueries.cpp b/Engine/Plugins/Experimental/EditorDataStorageFeatures/Source/TedsTypedElementBridge/Private/TedsTypedElementBridge/TedsTypedElementBridgeQueries.cpp
new file mode 100644
--- /dev/null
+++ b/Engine/Plugins/Experimental/EditorDataStorageFeatures/Source/TedsTypedElementBridge/Private/TedsTypedElementBridge/TedsTypedElementBridgeQueries.cpp
@@ -0,0 +1,110 @@
+// Copyright Epic Games, Inc. All Rights Reserved.
+
+#include "TedsTypedElementBridge/TedsTypedElementBridgeQueries.h"
+
+#include "TedsTypedElementBridge/TedsTypedElementBridgeCapabilities.h"
+#include "TedsTypedElementBridge/TedsTypedElementBridgeColumns.h"
+#include "DataStorage/Features.h"
+#include "Elements/Columns/TypedElementCompatibilityColumns.h"
+#include "Elements/Framework/EngineElementsLibrary.h"
+#include "Elements/Framework/TypedElementList.h"
+#include "Elements/Framework/TypedElementQueryBuilder.h"
+#include "Elements/Interfaces/TypedElementDataStorageInterface.h"
+#include "HAL/IConsoleManager.h"
+
+namespace UE::Editor::DataStorage::Compatibility::Private
+{
+	bool bBridgeEnabled = false;
+	FAutoConsoleVariableRef CVarBridgeEnabled(
+		TEXT("TEDS.TypedElementBridge.Enable"),
+		bBridgeEnabled,
+		TEXT("Automatically populated TEDS with TypedElementHandles"));
+} // namespace UE::Editor::DataStorage::Compatibility::Private
+
+uint8 UTypedElementBridgeDataStorageFactory::GetOrder() const
+{
+	return 110;
+}
+
+void UTypedElementBridgeDataStorageFactory::PreRegister(UE::Editor::DataStorage::ICoreProvider& DataStorage)
+{
+	Super::PreRegister(DataStorage);
+	DebugEnabledDelegateHandle = UE::Editor::DataStorage::Compatibility::Private::CVarBridgeEnabled->OnChangedDelegate().AddUObject(this, &UTypedElementBridgeDataStorageFactory::HandleOnEnabled);
+}
+
+void UTypedElementBridgeDataStorageFactory::PreShutdown(UE::Editor::DataStorage::ICoreProvider& DataStorage)
+{
+	UE::Editor::DataStorage::Compatibility::Private::CVarBridgeEnabled->OnChangedDelegate().Remove(DebugEnabledDelegateHandle);
+	DebugEnabledDelegateHandle.Reset();
+	CleanupTypedElementColumns(DataStorage);
+
+	Super::PreShutdown(DataStorage);
+}
+
+void UTypedElementBridgeDataStorageFactory::RegisterQueries(UE::Editor::DataStorage::ICoreProvider& DataStorage)
+{
+	Super::RegisterQueries(DataStorage);
+
+	if (IsEnabled())
+	{
+		RegisterQuery_NewUObject(DataStorage);
+	}
+}
+
+bool UTypedElementBridgeDataStorageFactory::IsEnabled()
+{
+	return UE::Editor::DataStorage::Compatibility::Private::CVarBridgeEnabled->GetBool();
+}
+
+void UTypedElementBridgeDataStorageFactory::RegisterQuery_NewUObject(UE::Editor::DataStorage::ICoreProvider& DataStorage)
+{
+	using namespace UE::Editor::DataStorage::Queries;
+	
+	RemoveTypedElementRowHandleQuery = DataStorage.RegisterQuery(
+		Select()
+			.ReadOnly<UE::Editor::DataStorage::Compatibility::FTypedElementColumn>()
+		.Compile());
+}
+
+void UTypedElementBridgeDataStorageFactory::UnregisterQuery_NewUObject(UE::Editor::DataStorage::ICoreProvider& DataStorage)
+{
+	DataStorage.UnregisterQuery(RemoveTypedElementRowHandleQuery);
+}
+
+void UTypedElementBridgeDataStorageFactory::CleanupTypedElementColumns(UE::Editor::DataStorage::ICoreProvider& DataStorage)
+{
+	using namespace UE::Editor::DataStorage::Queries;
+
+	// Remove any TEv1 handles
+	{
+		TArray<RowHandle> Handles;
+		DataStorage.RunQuery(
+			RemoveTypedElementRowHandleQuery,
+			CreateDirectQueryCallbackBinding(
+				[&Handles](IDirectQueryContext& Context, const RowHandle*)
+				{
+					Handles.Append(Context.GetRowHandles());
+				}));
+		
+		DataStorage.BatchAddRemoveColumns(TConstArrayView<RowHandle>(Handles), {}, {Compatibility::FTypedElementColumn::StaticStruct()});
+	}
+}
+
+void UTypedElementBridgeDataStorageFactory::HandleOnEnabled(IConsoleVariable* CVar)
+{
+	using namespace UE::Editor::DataStorage;
+	ICoreProvider* DataStorage = GetMutableDataStorageFeature<ICoreProvider>(StorageFeatureName);
+	bool bIsEnabled = CVar->GetBool();
+
+	if (bIsEnabled)
+	{
+		RegisterQuery_NewUObject(*DataStorage);
+		Compatibility::OnTypedElementBridgeEnabled().Broadcast(bIsEnabled);
+	}
+	else
+	{
+		Compatibility::OnTypedElementBridgeEnabled().Broadcast(bIsEnabled);
+		UnregisterQuery_NewUObject(*DataStorage);
+		CleanupTypedElementColumns(*DataStorage);
+	}
+}
diff --git a/Engine/Plugins/Experimental/EditorDataStorageFeatures/Source/TedsTypedElementBridge/Private/TedsTypedElementBridge/TedsTypedElementBridgeQueries.h b/Engine/Plugins/Experimental/EditorDataStorageFeatures/Source/TedsTypedElementBridge/Private/TedsTypedElementBridge/TedsTypedElementBridgeQueries.h
new file mode 100644
--- /dev/null
+++ b/Engine/Plugins/Experimental/EditorDataStorageFeatures/Source/TedsTypedElementBridge/Private/TedsTypedElementBridge/TedsTypedElementBridgeQueries.h
@@ -0,0 +1,37 @@
+// Copyright Epic Games, Inc. All Rights Reserved.
+
+#pragma once
+
+#include "Elements/Framework/TypedElementHandle.h"
+#include "Elements/Interfaces/TypedElementDataStorageInterface.h"
+#include "Elements/Interfaces/TypedElementDataStorageFactory.h"
+#include "UObject/UObjectGlobals.h"
+
+#include "TedsTypedElementBridgeQueries.generated.h"
+
+/**
+ * This class is responsible for running queries that will ensure Typed Element Handles
+ * are cleaned up when TEDS is shut down.
+ */
+UCLASS(Transient)
+class UTypedElementBridgeDataStorageFactory : public UEditorDataStorageFactory
+{
+	GENERATED_BODY()
+public:
+	~UTypedElementBridgeDataStorageFactory() override = default;
+	virtual uint8 GetOrder() const override;
+	virtual void PreRegister(UE::Editor::DataStorage::ICoreProvider& DataStorage) override;
+	virtual void PreShutdown(UE::Editor::DataStorage::ICoreProvider& DataStorage) override;
+	virtual void RegisterQueries(UE::Editor::DataStorage::ICoreProvider& DataStorage) override;
+
+	static bool IsEnabled();
+
+private:
+	void RegisterQuery_NewUObject(UE::Editor::DataStorage::ICoreProvider& DataStorage);
+	void UnregisterQuery_NewUObject(UE::Editor::DataStorage::ICoreProvider& DataStorage);
+	void CleanupTypedElementColumns(UE::Editor::DataStorage::ICoreProvider& DataStorage);
+	void HandleOnEnabled(IConsoleVariable* CVar);
+	
+	UE::Editor::DataStorage::QueryHandle RemoveTypedElementRowHandleQuery = UE::Editor::DataStorage::InvalidQueryHandle;
+	FDelegateHandle DebugEnabledDelegateHandle;
+};
diff --git a/Engine/Plugins/Experimental/EditorDataStorageFeatures/Source/TedsTypedElementBridge/Public/TedsTypedElementBridge/TedsTypedElementBridgeCapabilities.h b/Engine/Plugins/Experimental/EditorDataStorageFeatures/Source/TedsTypedElementBridge/Public/TedsTypedElementBridge/TedsTypedElementBridgeCapabilities.h
new file mode 100644
--- /dev/null
+++ b/Engine/Plugins/Experimental/EditorDataStorageFeatures/Source/TedsTypedElementBridge/Public/TedsTypedElementBridge/TedsTypedElementBridgeCapabilities.h
@@ -0,0 +1,13 @@
+// Copyright Epic Games, Inc. All Rights Reserved.
+
+#pragma once
+
+#include "Delegates/Delegate.h"
+
+namespace UE::Editor::DataStorage::Compatibility
+{
+	DECLARE_MULTICAST_DELEGATE_OneParam(FOnTypedElementBridgeEnabled, bool /*bEnabled*/);
+
+	TEDSTYPEDELEMENTBRIDGE_API FOnTypedElementBridgeEnabled& OnTypedElementBridgeEnabled();
+	TEDSTYPEDELEMENTBRIDGE_API bool IsTypedElementBridgeEnabled();
+} // namespace UE::Editor::DataStorage::Compatibility
diff --git a/Engine/Plugins/Experimental/EditorDataStorageFeatures/Source/TedsTypedElementBridge/Public/TedsTypedElementBridge/TedsTypedElementBridgeColumns.h b/Engine/Plugins/Experimental/EditorDataStorageFeatures/Source/TedsTypedElementBridge/Public/TedsTypedElementBridge/TedsTypedElementBridgeColumns.h
new file mode 100644
--- /dev/null
+++ b/Engine/Plugins/Experimental/EditorDataStorageFeatures/Source/TedsTypedElementBridge/Public/TedsTypedElementBridge/TedsTypedElementBridgeColumns.h
@@ -0,0 +1,23 @@
+// Copyright Epic Games, Inc. All Rights Reserved.
+
+#pragma once
+
+#include "DataStorage/CommonTypes.h"
+#include "Elements/Framework/TypedElementHandle.h"
+
+#include "TedsTypedElementBridgeColumns.generated.h"
+
+// A column which contains a TypedElementHandle from the TypedElementFramework
+// Prefer to use the namespaced alias when referencing in code - UE::Editor::DataStorage::Compatibility::FTypedElementColumn
+USTRUCT()
+struct FTedsTypedElementColumn : public FEditorDataStorageColumn
+{
+	GENERATED_BODY()
+
+	FTypedElementHandle Handle;
+};
+
+namespace UE::Editor::DataStorage::Compatibility
+{
+	using FTypedElementColumn = FTedsTypedElementColumn;
+}  // namespace UE::Editor::DataStorage::Compatibility
diff --git a/Engine/Plugins/Experimental/EditorDataStorageFeatures/Source/TedsTypedElementBridge/TedsTypedElementBridge.Build.cs b/Engine/Plugins/Experimental/EditorDataStorageFeatures/Source/TedsTypedElementBridge/TedsTypedElementBridge.Build.cs
new file mode 100644
--- /dev/null
+++ b/Engine/Plugins/Experimental/EditorDataStorageFeatures/Source/TedsTypedElementBridge/TedsTypedElementBridge.Build.cs
@@ -0,0 +1,38 @@
+// Copyright Epic Games, Inc. All Rights Reserved.
+
+using UnrealBuildTool;
+
+public class TedsTypedElementBridge : ModuleRules
+{
+	public TedsTypedElementBridge(ReadOnlyTargetRules Target) : base(Target)
+	{
+		// Enable truncation warnings in this module
+		CppCompileWarningSettings.UnsafeTypeCastWarningLevel = WarningLevel.Warning;
+
+		if (Target.bBuildEditor)
+		{
+			PublicIncludePaths.AddRange(new string[] {});
+			PrivateIncludePaths.AddRange(new string[] {});
+
+			PublicDependencyModuleNames.AddRange(
+				new string[]
+				{
+					"TypedElementFramework",
+					"TedsQueryStack"
+				});
+
+			PrivateDependencyModuleNames.AddRange(
+				new string[]
+				{
+					"Core",
+					"CoreUObject",
+					"Engine",
+					"InputCore",
+					"Slate",
+					"SlateCore",
+					"WorkspaceMenuStructure",
+					"ToolWidgets"
+				});
+		}
+	}
+}
\ No newline at end of file
