STedsHierarchyViewer: Add functionality to bind commonly used columns (e.g context menu) to the widget row

diff --git a/Engine/Plugins/Experimental/EditorDataStorageFeatures/Source/TedsTableViewer/Private/Widgets/STedsHierarchyViewer.cpp b/Engine/Plugins/Experimental/EditorDataStorageFeatures/Source/TedsTableViewer/Private/Widgets/STedsHierarchyViewer.cpp
--- a/Engine/Plugins/Experimental/EditorDataStorageFeatures/Source/TedsTableViewer/Private/Widgets/STedsHierarchyViewer.cpp
+++ b/Engine/Plugins/Experimental/EditorDataStorageFeatures/Source/TedsTableViewer/Private/Widgets/STedsHierarchyViewer.cpp
@@ -67,7 +67,7 @@ namespace UE::Editor::DataStorage
 		
 		AddWidgetColumns();
 
-		TreeView = SNew(STedsTreeView, STedsTreeView::FOnGetParent::CreateSP(this, &SHierarchyViewer::GetParentRow))
+		TreeView = SNew(STedsTreeView, STedsTreeView::FOnGetParent::CreateSP(this, &SHierarchyViewer::GetParentRow), GetWidgetRowHandle())
 			.HeaderRow(HeaderRowWidget)
 			.TopLevelRowsSource(&TopLevelRows)
 			.RowsSource(&Model->GetItems())
diff --git a/Engine/Plugins/Experimental/EditorDataStorageFeatures/Source/TedsTableViewer/Private/Widgets/STedsTreeView.cpp b/Engine/Plugins/Experimental/EditorDataStorageFeatures/Source/TedsTableViewer/Private/Widgets/STedsTreeView.cpp
--- a/Engine/Plugins/Experimental/EditorDataStorageFeatures/Source/TedsTableViewer/Private/Widgets/STedsTreeView.cpp
+++ b/Engine/Plugins/Experimental/EditorDataStorageFeatures/Source/TedsTableViewer/Private/Widgets/STedsTreeView.cpp
@@ -2,15 +2,22 @@
 
 #include "STedsTreeView.h"
 
+#include "Columns/SlateDelegateColumns.h"
+#include "Elements/Framework/TypedElementAttributeBinding.h"
+
 namespace UE::Editor::DataStorage
 {
-	void STedsTreeView::Construct( const FArguments& InArgs, const FOnGetParent& InOnGetParent )
+	void STedsTreeView::Construct( const FArguments& InArgs, const FOnGetParent& InOnGetParent, RowHandle InWidgetRow )
 	{
 		checkf(InOnGetParent.IsBound(), TEXT("Must provide FOnGetParent delegate."));
 		checkf(InArgs._RowsSource, TEXT("Must provide a source for the rows to be displayed."));
 
 		OnGetParent = InOnGetParent;
 		RowsSource = InArgs._RowsSource;
+		WidgetRow = InWidgetRow;
+		
+		// Attribute binder to bind widget columns to attributes on the ListView
+		FAttributeBinder Binder(WidgetRow);
 			
 		STreeView::Construct(STreeView<TableViewerItemPtr>::FArguments()
 			.HeaderRow(InArgs._HeaderRow)
@@ -18,6 +25,9 @@ namespace UE::Editor::DataStorage
 			.OnGenerateRow(InArgs._OnGenerateRow)
 			.OnSelectionChanged(InArgs._OnSelectionChanged)
 			.SelectionMode(InArgs._SelectionMode)
+			.OnContextMenuOpening(Binder.BindEvent(&FWidgetContextMenuColumn::OnContextMenuOpening))
+			.OnItemScrolledIntoView(Binder.BindEvent(&FWidgetRowScrolledIntoView::OnItemScrolledIntoView))
+			.OnMouseButtonDoubleClick(Binder.BindEvent(&FWidgetDoubleClickedColumn::OnMouseButtonDoubleClick))
 			.OnGetChildren(this, &STedsTreeView::GetChildren_Internal)
 			);
 	}
@@ -76,4 +86,4 @@ namespace UE::Editor::DataStorage
 			OutChildren = *FoundArray; 
 		}
 	}
-}
\ No newline at end of file
+}
diff --git a/Engine/Plugins/Experimental/EditorDataStorageFeatures/Source/TedsTableViewer/Private/Widgets/STedsTreeView.h b/Engine/Plugins/Experimental/EditorDataStorageFeatures/Source/TedsTableViewer/Private/Widgets/STedsTreeView.h
--- a/Engine/Plugins/Experimental/EditorDataStorageFeatures/Source/TedsTableViewer/Private/Widgets/STedsTreeView.h
+++ b/Engine/Plugins/Experimental/EditorDataStorageFeatures/Source/TedsTableViewer/Private/Widgets/STedsTreeView.h
@@ -39,7 +39,7 @@ namespace UE::Editor::DataStorage
 			
 		SLATE_END_ARGS()
 		
-		void Construct( const FArguments& InArgs, const FOnGetParent& InOnGetParent );
+		void Construct( const FArguments& InArgs, const FOnGetParent& InOnGetParent, RowHandle InWidgetRow );
 
 		virtual void Tick( const FGeometry& AllottedGeometry, const double InCurrentTime, const float InDeltaTime ) override;
 		virtual void RequestListRefresh() override;
@@ -55,6 +55,8 @@ namespace UE::Editor::DataStorage
 		protected:
 		
 		FOnGetParent OnGetParent;
+		
+		RowHandle WidgetRow = InvalidRowHandle;
 
 		// All rows that can be shown in this tree, including the full hieararchy regardless of whether it is currently expanded
 		const TArray<TableViewerItemPtr>* RowsSource = nullptr;
