[TEDS Outliner] Make all current outliner columns sortable

diff --git a/Engine/Plugins/Experimental/EditorDataStorageFeatures/Source/TedsAlerts/Private/TedsAlertWidget.cpp b/Engine/Plugins/Experimental/EditorDataStorageFeatures/Source/TedsAlerts/Private/TedsAlertWidget.cpp
--- a/Engine/Plugins/Experimental/EditorDataStorageFeatures/Source/TedsAlerts/Private/TedsAlertWidget.cpp
+++ b/Engine/Plugins/Experimental/EditorDataStorageFeatures/Source/TedsAlerts/Private/TedsAlertWidget.cpp
@@ -14,6 +14,21 @@
 
 #define LOCTEXT_NAMESPACE "TedsAlertWidget"
 
+namespace UE::Editor::DataStorage
+{
+	using namespace UE::Editor::DataStorage::Columns;
+
+	FText FAlertWidgetSorter::GetShortName() const
+	{
+		return LOCTEXT("AlertWidgetSorter", "Alert Sorter");
+	}
+
+	FPrefixInfo FAlertWidgetSorter::CalculatePrefix(const FAlertColumn& Column, uint32 ByteIndex) const
+	{
+		return CreateSortPrefix(ByteIndex, -static_cast<int32>(Column.AlertType), Column.Priority);
+	}
+}
+
 namespace UE::Editor::DataStorage::Widgets::Private
 {
 	void UpdateWidget(const TSharedPtr<SWidget>& Widget, const FText& Alert, bool bIsWarning, uint16 ErrorCount, uint16 WarningCount,
@@ -399,6 +414,18 @@ bool FAlertWidgetConstructor::FinalizeWidget(UE::Editor::DataStorage::ICoreProvi
 	return true;
 }
 
+TArray<TSharedPtr<const UE::Editor::DataStorage::FColumnSorterInterface>> FAlertWidgetConstructor::ConstructColumnSorters(
+	UE::Editor::DataStorage::ICoreProvider* DataStorage,
+	UE::Editor::DataStorage::IUiProvider* DataStorageUi,
+	const UE::Editor::DataStorage::FMetaDataView& Arguments)
+{
+	using namespace UE::Editor::DataStorage;
+	return TArray<TSharedPtr<const FColumnSorterInterface>>(
+		{
+			MakeShared<FAlertWidgetSorter>()
+		});
+}
+
 
 
 //
diff --git a/Engine/Plugins/Experimental/EditorDataStorageFeatures/Source/TedsAlerts/Private/TedsAlertWidget.h b/Engine/Plugins/Experimental/EditorDataStorageFeatures/Source/TedsAlerts/Private/TedsAlertWidget.h
--- a/Engine/Plugins/Experimental/EditorDataStorageFeatures/Source/TedsAlerts/Private/TedsAlertWidget.h
+++ b/Engine/Plugins/Experimental/EditorDataStorageFeatures/Source/TedsAlerts/Private/TedsAlertWidget.h
@@ -5,9 +5,23 @@
 #include "DataStorage/CommonTypes.h"
 #include "Elements/Interfaces/TypedElementDataStorageFactory.h"
 #include "Elements/Interfaces/TypedElementDataStorageUiInterface.h"
+#include "TedsAlertColumns.h"
 
 #include "TedsAlertWidget.generated.h"
 
+namespace UE::Editor::DataStorage
+{
+	using namespace UE::Editor::DataStorage::Columns;
+
+	class FAlertWidgetSorter final : public TColumnSorterInterface<FColumnSorterInterface::ESortType::FixedSize64, FAlertColumn>
+	{
+	public:
+		virtual FText GetShortName() const override;
+	protected:
+		virtual FPrefixInfo CalculatePrefix(const FAlertColumn& Column, uint32 ByteIndex) const override;
+	};
+}
+
 UCLASS()
 class UAlertWidgetFactory final : public UEditorDataStorageFactory
 {
@@ -49,6 +63,10 @@ private:
 	bool FinalizeWidget(UE::Editor::DataStorage::ICoreProvider* DataStorage,
 		UE::Editor::DataStorage::IUiProvider* DataStorageUi,
 		UE::Editor::DataStorage::RowHandle Row, const TSharedPtr<SWidget>& Widget) override;
+	TArray<TSharedPtr<const UE::Editor::DataStorage::FColumnSorterInterface>> ConstructColumnSorters(
+		UE::Editor::DataStorage::ICoreProvider* DataStorage,
+		UE::Editor::DataStorage::IUiProvider* DataStorageUi,
+		const UE::Editor::DataStorage::FMetaDataView& Arguments) override;
 };
 
 USTRUCT(meta = (DisplayName = "General purpose alert"))
diff --git a/Engine/Plugins/Experimental/EditorDataStorageFeatures/Source/TedsOutliner/Private/Widgets/TypeInfoWidget.cpp b/Engine/Plugins/Experimental/EditorDataStorageFeatures/Source/TedsOutliner/Private/Widgets/TypeInfoWidget.cpp
--- a/Engine/Plugins/Experimental/EditorDataStorageFeatures/Source/TedsOutliner/Private/Widgets/TypeInfoWidget.cpp
+++ b/Engine/Plugins/Experimental/EditorDataStorageFeatures/Source/TedsOutliner/Private/Widgets/TypeInfoWidget.cpp
@@ -8,11 +8,32 @@
 #include "Elements/Columns/TypedElementSlateWidgetColumns.h"
 #include "Elements/Columns/TypedElementTypeInfoColumns.h"
 #include "Elements/Interfaces/Capabilities/TypedElementUiTextCapability.h"
+#include "Elements/Interfaces/TypedElementDataStorageUiInterface.h"
 #include "Styling/SlateIconFinder.h"
 #include "TedsTableViewerUtils.h"
 #include "Widgets/Text/STextBlock.h"
 #include "Widgets/Images/SImage.h"
 
+#define LOCTEXT_NAMESPACE "TypeInfoWidget"
+
+namespace UE::Editor::DataStorage
+{
+	FText FTypeInfoWidgetSorter::GetShortName() const
+	{
+		return LOCTEXT("TypeInfoWidgetSorter", "Type Info Sorter");
+	}
+
+	int32 FTypeInfoWidgetSorter::Compare(const FTypedElementClassTypeInfoColumn& Left, const FTypedElementClassTypeInfoColumn& Right) const
+	{
+		return Left.TypeInfo->GetName().Compare(Right.TypeInfo->GetName(), ESearchCase::IgnoreCase);
+	}
+
+	FPrefixInfo FTypeInfoWidgetSorter::CalculatePrefix(const FTypedElementClassTypeInfoColumn& Column, uint32 ByteIndex) const
+	{
+		return CreateSortPrefix(ByteIndex, TSortStringView(FSortCaseInsensitive{}, Column.TypeInfo->GetName()));
+	}
+}
+
 void UTypeInfoWidgetFactory::RegisterWidgetConstructors(UE::Editor::DataStorage::ICoreProvider& DataStorage,
 	UE::Editor::DataStorage::IUiProvider& DataStorageUi) const
 {
@@ -125,3 +146,15 @@ bool FTypeInfoWidgetConstructor::FinalizeWidget(UE::Editor::DataStorage::ICorePr
 
 	return true;
 }
+
+TArray<TSharedPtr<const UE::Editor::DataStorage::FColumnSorterInterface>> FTypeInfoWidgetConstructor::ConstructColumnSorters(
+	UE::Editor::DataStorage::ICoreProvider* DataStorage,
+	UE::Editor::DataStorage::IUiProvider* DataStorageUi,
+	const UE::Editor::DataStorage::FMetaDataView& Arguments)
+{
+	using namespace UE::Editor::DataStorage;
+	return TArray<TSharedPtr<const FColumnSorterInterface>>(
+		{
+			MakeShared<FTypeInfoWidgetSorter>()
+		});
+}
\ No newline at end of file
diff --git a/Engine/Plugins/Experimental/EditorDataStorageFeatures/Source/TedsOutliner/Public/Widgets/TypeInfoWidget.h b/Engine/Plugins/Experimental/EditorDataStorageFeatures/Source/TedsOutliner/Public/Widgets/TypeInfoWidget.h
--- a/Engine/Plugins/Experimental/EditorDataStorageFeatures/Source/TedsOutliner/Public/Widgets/TypeInfoWidget.h
+++ b/Engine/Plugins/Experimental/EditorDataStorageFeatures/Source/TedsOutliner/Public/Widgets/TypeInfoWidget.h
@@ -5,11 +5,24 @@
 #include "Elements/Interfaces/TypedElementDataStorageFactory.h"
 #include "Elements/Interfaces/TypedElementDataStorageInterface.h"
 #include "Elements/Interfaces/TypedElementDataStorageUiInterface.h"
+#include "Elements/Columns/TypedElementTypeInfoColumns.h"
 
 #include "TypeInfoWidget.generated.h"
 
 #define UE_API TEDSOUTLINER_API
 
+namespace UE::Editor::DataStorage
+{
+	class FTypeInfoWidgetSorter final : public TColumnSorterInterface<FColumnSorterInterface::ESortType::HybridSort, FTypedElementClassTypeInfoColumn>
+	{
+	public:
+		virtual FText GetShortName() const override;
+	protected:
+		virtual int32 Compare(const FTypedElementClassTypeInfoColumn& Left, const FTypedElementClassTypeInfoColumn& Right) const override;
+		virtual FPrefixInfo CalculatePrefix(const FTypedElementClassTypeInfoColumn& Column, uint32 ByteIndex) const override;
+	};
+}
+
 UCLASS()
 class UTypeInfoWidgetFactory : public UEditorDataStorageFactory
 {
@@ -37,6 +50,10 @@ protected:
 	UE_API bool FinalizeWidget(UE::Editor::DataStorage::ICoreProvider* DataStorage,
 		UE::Editor::DataStorage::IUiProvider* DataStorageUi,
 		UE::Editor::DataStorage::RowHandle Row, const TSharedPtr<SWidget>& Widget) override;
+	UE_API TArray<TSharedPtr<const UE::Editor::DataStorage::FColumnSorterInterface>> ConstructColumnSorters(
+		UE::Editor::DataStorage::ICoreProvider* DataStorage,
+		UE::Editor::DataStorage::IUiProvider* DataStorageUi,
+		const UE::Editor::DataStorage::FMetaDataView& Arguments) override;
 
 protected:
 
