[TEDS Outliner] Add the level column to the TEDS outliner

diff --git a/Engine/Plugins/Experimental/EditorDataStorageFeatures/Source/TedsActorCompatibility/Private/Compatibility/TedsActorLevelQueries.cpp b/Engine/Plugins/Experimental/EditorDataStorageFeatures/Source/TedsActorCompatibility/Private/Compatibility/TedsActorLevelQueries.cpp
new file mode 100644
--- /dev/null
+++ b/Engine/Plugins/Experimental/EditorDataStorageFeatures/Source/TedsActorCompatibility/Private/Compatibility/TedsActorLevelQueries.cpp
@@ -0,0 +1,69 @@
+// Copyright Epic Games, Inc. All Rights Reserved.
+
+#include "Compatibility/TedsActorLevelQueries.h"
+
+#include "Elements/Columns/TypedElementCompatibilityColumns.h"
+#include "Elements/Columns/TypedElementMiscColumns.h"
+#include "Elements/Framework/TypedElementQueryBuilder.h"
+#include "Engine/Level.h"
+#include "GameFramework/Actor.h"
+#include "LevelInstance/LevelInstanceInterface.h"
+
+void UActorLevelDataStorageFactory::RegisterQueries(UE::Editor::DataStorage::ICoreProvider& DataStorage)
+{
+	RegisterActorAddLevelColumn(DataStorage);
+	RegisterActorLevelToColumnQuery(DataStorage);
+}
+
+void UActorLevelDataStorageFactory::RegisterActorAddLevelColumn(UE::Editor::DataStorage::ICoreProvider& DataStorage) const
+{
+	using namespace UE::Editor::DataStorage::Queries;
+
+	DataStorage.RegisterQuery(
+		Select(
+			TEXT("Actor Level to New Column"),
+			FProcessor(EQueryTickPhase::PrePhysics, DataStorage.GetQueryTickGroupName(EQueryTickGroups::SyncExternalToDataStorage))
+			.SetExecutionMode(EExecutionMode::GameThread),
+			[](IQueryContext& Context, RowHandle Row, const FTypedElementUObjectColumn& Object)
+			{
+				if (const AActor* ActorInstance = Cast<AActor>(Object.Object); ActorInstance != nullptr)
+				{
+					if (ULevel* ActorLevel = ActorInstance->GetLevel(); ActorLevel != nullptr)
+					{
+						Context.AddColumn(Row, FLevelColumn{ .Level = ActorLevel });
+					}
+				}
+			})
+		.Where()
+			.All<FTypedElementActorTag, FTypedElementSyncFromWorldTag>()
+			.None<FLevelColumn>()
+		.Compile()
+	);
+}
+
+void UActorLevelDataStorageFactory::RegisterActorLevelToColumnQuery(UE::Editor::DataStorage::ICoreProvider& DataStorage) const
+{
+	using namespace UE::Editor::DataStorage::Queries;
+
+	DataStorage.RegisterQuery(
+		Select(
+			TEXT("Actor Level to Column"),
+			FProcessor(EQueryTickPhase::PrePhysics, DataStorage.GetQueryTickGroupName(EQueryTickGroups::SyncExternalToDataStorage))
+			.SetExecutionMode(EExecutionMode::GameThread),
+			[](IQueryContext& Context, RowHandle Row, const FTypedElementUObjectColumn& Object, FLevelColumn& LevelColumn)
+			{
+				if (const AActor* ActorInstance = Cast<AActor>(Object.Object); ActorInstance != nullptr)
+				{
+					if (ULevel* ActorLevel = ActorInstance->GetLevel(); ActorLevel != nullptr)
+					{
+						LevelColumn.Level = ActorLevel;
+						return;
+					}
+				}
+				Context.RemoveColumns<FLevelColumn>(Row);
+			})
+		.Where()
+			.All<FTypedElementActorTag, FTypedElementSyncFromWorldTag>()
+		.Compile()
+	);
+}
\ No newline at end of file
diff --git a/Engine/Plugins/Experimental/EditorDataStorageFeatures/Source/TedsActorCompatibility/Private/Compatibility/TedsActorLevelQueries.h b/Engine/Plugins/Experimental/EditorDataStorageFeatures/Source/TedsActorCompatibility/Private/Compatibility/TedsActorLevelQueries.h
new file mode 100644
--- /dev/null
+++ b/Engine/Plugins/Experimental/EditorDataStorageFeatures/Source/TedsActorCompatibility/Private/Compatibility/TedsActorLevelQueries.h
@@ -0,0 +1,29 @@
+// Copyright Epic Games, Inc.All Rights Reserved.
+
+#pragma once
+
+#include "Elements/Interfaces/TypedElementDataStorageFactory.h"
+#include "UObject/ObjectMacros.h"
+
+#include "TedsActorLevelQueries.generated.h"
+
+UCLASS()
+class UActorLevelDataStorageFactory : public UEditorDataStorageFactory
+{
+	GENERATED_BODY()
+
+public:
+	~UActorLevelDataStorageFactory() override = default;
+
+	void RegisterQueries(UE::Editor::DataStorage::ICoreProvider& DataStorage) override;
+private:
+
+	/**
+	 * Adds the level column to new actors that do not have one already.
+	 */
+	void RegisterActorAddLevelColumn(UE::Editor::DataStorage::ICoreProvider& DataStorage) const;
+	/**
+	 * Takes the level of the actor and copies it to the Data Storage if they differ.
+	 */
+	void RegisterActorLevelToColumnQuery(UE::Editor::DataStorage::ICoreProvider& DataStorage) const;
+};
diff --git a/Engine/Plugins/Experimental/EditorDataStorageFeatures/Source/TedsActorCompatibility/Private/TedsActorCompatibilityFactory.cpp b/Engine/Plugins/Experimental/EditorDataStorageFeatures/Source/TedsActorCompatibility/Private/TedsActorCompatibilityFactory.cpp
--- a/Engine/Plugins/Experimental/EditorDataStorageFeatures/Source/TedsActorCompatibility/Private/TedsActorCompatibilityFactory.cpp
+++ b/Engine/Plugins/Experimental/EditorDataStorageFeatures/Source/TedsActorCompatibility/Private/TedsActorCompatibilityFactory.cpp
@@ -42,7 +42,8 @@ void UTedsActorCompatibilityFactory::RegisterTable(UE::Editor::DataStorage::ICor
 	StandardActorTable = DataStorage.RegisterTable(TTypedElementColumnTypeList<
 		FTypedElementUObjectColumn, FTypedElementUObjectIdColumn, FUObjectIdNameColumn,
 		FTypedElementClassTypeInfoColumn, FTypedElementLabelColumn, FTypedElementLabelHashColumn, 
-		FTedsActorMobilityColumn, FTypedElementActorTag, FTypedElementSyncFromWorldTag>(),
+		FTedsActorMobilityColumn, FTypedElementActorTag, FTypedElementSyncFromWorldTag,
+		FLevelColumn>(),
 		FName("Editor_StandardActorTable"));
 
 	StandardActorWithTransformTable = DataStorage.RegisterTable(StandardActorTable,
diff --git a/Engine/Plugins/Experimental/EditorDataStorageFeatures/Source/TedsOutliner/Private/Compatibility/SceneOutlinerTedsBridge.cpp b/Engine/Plugins/Experimental/EditorDataStorageFeatures/Source/TedsOutliner/Private/Compatibility/SceneOutlinerTedsBridge.cpp
--- a/Engine/Plugins/Experimental/EditorDataStorageFeatures/Source/TedsOutliner/Private/Compatibility/SceneOutlinerTedsBridge.cpp
+++ b/Engine/Plugins/Experimental/EditorDataStorageFeatures/Source/TedsOutliner/Private/Compatibility/SceneOutlinerTedsBridge.cpp
@@ -35,7 +35,7 @@ void USceneOutlinerTedsBridgeFactory::PreRegister(UE::Editor::DataStorage::ICore
 	TEDSToOutlinerDefaultColumnMapping.Add(FUObjectIdNameColumn::StaticStruct(), FSceneOutlinerBuiltInColumnTypes::IDName());
 	TEDSToOutlinerDefaultColumnMapping.Add(FTedsActorMobilityColumn::StaticStruct(), FSceneOutlinerBuiltInColumnTypes::Mobility());
 	TEDSToOutlinerDefaultColumnMapping.Add(FTypedElementLabelColumn::StaticStruct(), FSceneOutlinerBuiltInColumnTypes::Label());
-
+	TEDSToOutlinerDefaultColumnMapping.Add(FLevelColumn::StaticStruct(), FSceneOutlinerBuiltInColumnTypes::Level());
 }
 
 void USceneOutlinerTedsBridgeFactory::RegisterTables(UE::Editor::DataStorage::ICoreProvider& DataStorage)
diff --git a/Engine/Plugins/Experimental/EditorDataStorageFeatures/Source/TedsOutliner/Private/TedsOutlinerModule.cpp b/Engine/Plugins/Experimental/EditorDataStorageFeatures/Source/TedsOutliner/Private/TedsOutlinerModule.cpp
--- a/Engine/Plugins/Experimental/EditorDataStorageFeatures/Source/TedsOutliner/Private/TedsOutlinerModule.cpp
+++ b/Engine/Plugins/Experimental/EditorDataStorageFeatures/Source/TedsOutliner/Private/TedsOutlinerModule.cpp
@@ -128,6 +128,7 @@ UE::Editor::DataStorage::FQueryDescription FTedsOutlinerModule::GetLevelEditorTe
 					  FChildAlertColumn,
 					  FUObjectIdNameColumn,
 					  FTedsActorMobilityColumn,
+					  FLevelColumn,
 					  FVisibleInEditorColumn>()
 		.Compile();
 
@@ -139,6 +140,7 @@ UE::Editor::DataStorage::FQueryDescription FTedsOutlinerModule::GetLevelEditorTe
 					  FAlertColumn, 
 					  FUObjectIdNameColumn,
 					  FTedsActorMobilityColumn,
+					  FLevelColumn,
 					  FVisibleInEditorColumn>()
 		.Compile();
 
diff --git a/Engine/Plugins/Experimental/EditorDataStorageFeatures/Source/TedsOutliner/Private/Widgets/OutlinerLevelWidget.cpp b/Engine/Plugins/Experimental/EditorDataStorageFeatures/Source/TedsOutliner/Private/Widgets/OutlinerLevelWidget.cpp
new file mode 100644
--- /dev/null
+++ b/Engine/Plugins/Experimental/EditorDataStorageFeatures/Source/TedsOutliner/Private/Widgets/OutlinerLevelWidget.cpp
@@ -0,0 +1,101 @@
+// Copyright Epic Games, Inc. All Rights Reserved.
+
+#include "Widgets/OutlinerLevelWidget.h"
+
+#include "SceneOutlinerHelpers.h"
+#include "Elements/Columns/TypedElementCompatibilityColumns.h"
+#include "Elements/Columns/TypedElementMiscColumns.h"
+#include "Elements/Framework/TypedElementAttributeBinding.h"
+#include "Elements/Interfaces/TypedElementDataStorageUiInterface.h"
+#include "Engine/Level.h"
+#include "Widgets/Text/STextBlock.h"
+
+#define LOCTEXT_NAMESPACE "OutlinerLevelWidget"
+
+namespace UE::Editor::DataStorage
+{
+	FText FLevelWidgetSorter::GetShortName() const
+	{
+		return LOCTEXT("LevelWidgetSorter", "Level Sorter");
+	}
+
+	int32 FLevelWidgetSorter::Compare(const FLevelColumn& Left, const FLevelColumn& Right) const
+	{
+		return Left.Level->GetOuter()->GetName().Compare(Right.Level->GetOuter()->GetName());
+	}
+
+	FPrefixInfo FLevelWidgetSorter::CalculatePrefix(const FLevelColumn& Column, uint32 ByteIndex) const
+	{
+		return CreateSortPrefix(ByteIndex, TSortStringView(FSortCaseSensitive{}, Column.Level->GetOuter()->GetName()));
+	}
+}
+
+void UOutlinerLevelWidgetFactory::RegisterWidgetConstructors(UE::Editor::DataStorage::ICoreProvider& DataStorage,
+	UE::Editor::DataStorage::IUiProvider& DataStorageUi) const
+{
+	using namespace UE::Editor::DataStorage::Queries;
+
+	DataStorageUi.RegisterWidgetFactory<FOutlinerLevelWidgetConstructor>(
+		DataStorageUi.FindPurpose(UE::Editor::DataStorage::IUiProvider::FPurposeInfo("SceneOutliner", "Cell", NAME_None).GeneratePurposeID()),
+		TColumn<FLevelColumn>());
+	DataStorageUi.RegisterWidgetFactory<FOutlinerLevelWidgetConstructor>(
+			DataStorageUi.FindPurpose(DataStorageUi.GetGeneralWidgetPurposeID()),
+			TColumn<FLevelColumn>());
+}
+
+FOutlinerLevelWidgetConstructor::FOutlinerLevelWidgetConstructor()
+	: FSimpleWidgetConstructor(StaticStruct())
+{
+}
+
+TSharedPtr<SWidget> FOutlinerLevelWidgetConstructor::CreateWidget(
+	UE::Editor::DataStorage::ICoreProvider* DataStorage,
+	UE::Editor::DataStorage::IUiProvider* DataStorageUi,
+	UE::Editor::DataStorage::RowHandle TargetRow,
+	UE::Editor::DataStorage::RowHandle WidgetRow,
+	const UE::Editor::DataStorage::FMetaDataView& Arguments)
+{
+	using namespace UE::Editor::DataStorage;
+	if (!DataStorage->IsRowAvailable(TargetRow))
+	{
+		return SNew(STextBlock)
+			.Text(LOCTEXT("MissingRowReferenceColumn", "Unable to retrieve row reference."));
+	}
+
+	FAttributeBinder Binder(TargetRow, DataStorage);
+
+	return SNew(SHorizontalBox)
+		+ SHorizontalBox::Slot()
+		.AutoWidth()
+		.VAlign(VAlign_Center)
+		.Padding(8, 0, 0, 0)
+		[
+			SNew(STextBlock)
+				.ColorAndOpacity(FSlateColor::UseSubduedForeground())
+				.Text(Binder.BindData(&FLevelColumn::Level, [](TWeakObjectPtr<ULevel> InLevel)
+				{
+					if (InLevel.IsValid())
+					{
+						return FText::FromString(InLevel->GetOuter()->GetName());
+					}
+					else
+					{
+						return LOCTEXT("OutlinerLevelUnknown", "<unknown>");
+					}
+				}))
+		];
+}
+
+TArray<TSharedPtr<const UE::Editor::DataStorage::FColumnSorterInterface>> FOutlinerLevelWidgetConstructor::ConstructColumnSorters(
+	UE::Editor::DataStorage::ICoreProvider* DataStorage,
+	UE::Editor::DataStorage::IUiProvider* DataStorageUi,
+	const UE::Editor::DataStorage::FMetaDataView& Arguments)
+{
+	using namespace UE::Editor::DataStorage;
+	return TArray<TSharedPtr<const FColumnSorterInterface>>(
+		{
+			MakeShared<FLevelWidgetSorter>()
+		});
+}
+
+#undef LOCTEXT_NAMESPACE
diff --git a/Engine/Plugins/Experimental/EditorDataStorageFeatures/Source/TedsOutliner/Public/Widgets/OutlinerLevelWidget.h b/Engine/Plugins/Experimental/EditorDataStorageFeatures/Source/TedsOutliner/Public/Widgets/OutlinerLevelWidget.h
new file mode 100644
--- /dev/null
+++ b/Engine/Plugins/Experimental/EditorDataStorageFeatures/Source/TedsOutliner/Public/Widgets/OutlinerLevelWidget.h
@@ -0,0 +1,55 @@
+// Copyright Epic Games, Inc. All Rights Reserved.
+
+#pragma once
+
+#include "Elements/Interfaces/TypedElementDataStorageFactory.h"
+#include "Elements/Interfaces/TypedElementDataStorageInterface.h"
+#include "Elements/Interfaces/TypedElementDataStorageUiInterface.h"
+#include "Elements/Columns/TypedElementCompatibilityColumns.h"
+
+#include "OutlinerLevelWidget.generated.h"
+
+namespace UE::Editor::DataStorage
+{
+	class FLevelWidgetSorter final : public TColumnSorterInterface<FColumnSorterInterface::ESortType::HybridSort, FLevelColumn>
+	{
+	public:
+		virtual FText GetShortName() const override;
+	protected:
+		virtual int32 Compare(const FLevelColumn& Left, const FLevelColumn& Right) const override;
+		virtual FPrefixInfo CalculatePrefix(const FLevelColumn& Column, uint32 ByteIndex) const override;
+	};
+}
+
+UCLASS()
+class UOutlinerLevelWidgetFactory : public UEditorDataStorageFactory
+{
+	GENERATED_BODY()
+
+public:
+	~UOutlinerLevelWidgetFactory() override = default;
+
+	void RegisterWidgetConstructors(UE::Editor::DataStorage::ICoreProvider& DataStorage,
+		UE::Editor::DataStorage::IUiProvider& DataStorageUi) const override;
+};
+
+USTRUCT()
+struct TEDSOUTLINER_API FOutlinerLevelWidgetConstructor : public FSimpleWidgetConstructor
+{
+	GENERATED_BODY()
+
+public:
+	FOutlinerLevelWidgetConstructor();
+	~FOutlinerLevelWidgetConstructor() override = default;
+
+	virtual TSharedPtr<SWidget> CreateWidget(
+		UE::Editor::DataStorage::ICoreProvider* DataStorage,
+		UE::Editor::DataStorage::IUiProvider* DataStorageUi,
+		UE::Editor::DataStorage::RowHandle TargetRow,
+		UE::Editor::DataStorage::RowHandle WidgetRow,
+		const UE::Editor::DataStorage::FMetaDataView& Arguments) override;
+	virtual TArray<TSharedPtr<const UE::Editor::DataStorage::FColumnSorterInterface>> ConstructColumnSorters(
+		UE::Editor::DataStorage::ICoreProvider* DataStorage,
+		UE::Editor::DataStorage::IUiProvider* DataStorageUi,
+		const UE::Editor::DataStorage::FMetaDataView& Arguments) override;
+};
\ No newline at end of file
diff --git a/Engine/Source/Runtime/TypedElementFramework/Public/Elements/Columns/TypedElementCompatibilityColumns.h b/Engine/Source/Runtime/TypedElementFramework/Public/Elements/Columns/TypedElementCompatibilityColumns.h
--- a/Engine/Source/Runtime/TypedElementFramework/Public/Elements/Columns/TypedElementCompatibilityColumns.h
+++ b/Engine/Source/Runtime/TypedElementFramework/Public/Elements/Columns/TypedElementCompatibilityColumns.h
@@ -84,6 +84,17 @@ struct FTypedElementWorldColumn final : public FEditorDataStorageColumn
 	TWeakObjectPtr<UWorld> World;
 };
 
+/**
+ * Column that stores a reference to a level.
+ */
+USTRUCT(meta = (DisplayName = "Level"))
+struct FTypedElementLevelColumn final : public FEditorDataStorageColumn
+{
+	GENERATED_BODY()
+
+	TWeakObjectPtr<ULevel> Level;
+};
+
 /**
  * Tag to signal that a row represents a property bag "placeholder-typed" object reference.
  * This object has an unknown base type and will generally be associated with a "property
@@ -133,4 +144,5 @@ struct FEditorDataStorageUObjectIdNameColumn final : public FEditorDataStorageCo
 namespace UE::Editor::DataStorage
 {
 	using FUObjectIdNameColumn = FEditorDataStorageUObjectIdNameColumn;
+	using FLevelColumn = FTypedElementLevelColumn;
 }  // namespace UE::Editor::DataStorage
\ No newline at end of file
