ActorComponent TEDS Registration

diff --git a/Engine/Plugins/Experimental/EditorDataStorageFeatures/Source/TedsActorCompatibility/Private/ActorComponentDebugHierarchyWidget/ActorComponentTree.cpp b/Engine/Plugins/Experimental/EditorDataStorageFeatures/Source/TedsActorCompatibility/Private/ActorComponentDebugHierarchyWidget/ActorComponentTree.cpp
new file mode 100644
--- /dev/null
+++ b/Engine/Plugins/Experimental/EditorDataStorageFeatures/Source/TedsActorCompatibility/Private/ActorComponentDebugHierarchyWidget/ActorComponentTree.cpp
@@ -0,0 +1,97 @@
+﻿// Copyright Epic Games, Inc. All Rights Reserved.
+
+#include "ActorComponentTree.h"
+
+#include "SActorComponentTreeWidget.h"
+#include "TedsQueryNode.h"
+#include "TedsQueryNode.h"
+#include "TedsRowQueryResultsNode.h"
+#include "WorkspaceMenuStructure.h"
+#include "WorkspaceMenuStructureModule.h"
+#include "WorkspaceMenuStructureModule.h"
+#include "Editor/LevelEditor/Private/SLevelEditor.h"
+#include "Elements/Columns/TypedElementCompatibilityColumns.h"
+#include "Elements/Columns/TypedElementSelectionColumns.h"
+#include "DataStorage/Features.h"
+#include "Framework/Docking/TabManager.h"
+#include "QueryStackNodes/RowQueryCallbackResultsNode.h"
+#include "Widgets/STedsTableViewer.h"
+#include "Widgets/Docking/SDockTab.h"
+
+#define LOCTEXT_NAMESPACE "ActorComponentTree"
+
+namespace UE::Editor::DataStorage::ActorCompatibility
+{
+	const FLazyName ActorComponentTreeTab = TEXT("ActorComponent Tree Tab");
+	
+	void RegisterActionComponentDebugHierarchyWidget()
+	{
+		const FName TabId = ActorComponentTreeTab.Resolve();
+		FTabSpawnerEntry& TabEntry = FGlobalTabmanager::Get()->RegisterNomadTabSpawner(
+			TabId,
+			FOnSpawnTab::CreateLambda([](const FSpawnTabArgs& Args)
+			{
+				using namespace  UE::Editor::DataStorage;
+				using namespace UE::Editor::DataStorage::Queries;
+				const TSharedRef<SDockTab> MajorTab = SNew(SDockTab)
+					.TabRole(ETabRole::MajorTab);
+
+				ICoreProvider* DataStorageInterface = GetMutableDataStorageFeature<ICoreProvider>(StorageFeatureName);
+				
+				FHierarchyHandle HierarchyHandle = DataStorageInterface->FindHierarchyByName(TEXT("ActorComponent"));
+
+				if (DataStorageInterface->IsValidHierachyHandle(HierarchyHandle))
+				{
+					// Query node which will enumerate all selected actor rows and provide access to their hierarchy
+					TSharedRef<QueryStack::FQueryNode> QueryNode = MakeShared<QueryStack::FQueryNode>(*DataStorageInterface);
+					QueryNode->SetQuery(
+						Select()
+						.ReadOnly<FTypedElementSelectionColumn>()
+						.Where()
+							.All<FTypedElementActorTag>()
+						.Compile());
+
+					auto FilterRowsBySelectionSet = [](
+						IDirectQueryContext& Context,
+						TArrayView<const RowHandle> Rows,
+						QueryStack::FRowQueryCallbackResultsNode::EmitRowFn EmitRows)
+					{
+						const FTypedElementSelectionColumn* SelectionColumnArrayPtr = Context.GetColumn<FTypedElementSelectionColumn>();
+					
+						TArrayView<const FTypedElementSelectionColumn> SelectionColumns(SelectionColumnArrayPtr, Context.GetRowCount());
+					
+						for (int32 RowIndex = 0; RowIndex < Rows.Num(); ++RowIndex)
+						{
+							if (SelectionColumns[RowIndex].SelectionSet == NAME_None)
+							{
+								EmitRows(MakeArrayView(&Rows[RowIndex], 1));
+							}
+						}
+					};
+				
+					TSharedPtr<QueryStack::IRowNode> TopLevelRowNode = MakeShared<QueryStack::FRowQueryCallbackResultsNode>(
+						*DataStorageInterface,
+						QueryNode,
+						QueryStack::FRowQueryResultsNode::ESyncFlags::IncrementWhenDifferent | QueryStack::FRowQueryResultsNode::ESyncFlags::RefreshOnUpdate,
+						FilterRowsBySelectionSet
+						);
+
+					// Create a widget to display the view model, widget internally binds the results of the query and displays appropriately
+					TSharedRef<SActorComponentTreeWidget> TreeWidget = SNew(SActorComponentTreeWidget, HierarchyHandle, TopLevelRowNode);
+			
+					MajorTab->SetContent(TreeWidget);
+				}
+				else
+				{
+					MajorTab->SetContent(SNullWidget::NullWidget);
+				}
+				return MajorTab;
+			}));
+
+		TabEntry
+		.SetGroup(WorkspaceMenu::GetMenuStructure().GetDeveloperToolsDebugCategory())
+		.SetDisplayName(LOCTEXT("ActorComponentTree_DisplayName", "ActorComponent Tree"))
+		.SetTooltipText(LOCTEXT("ActorComponentTree_OpenTree", "Opens the ActorComponent Tree"))
+		.SetIcon(FSlateIcon(FAppStyle::GetAppStyleSetName(), "Debug"));
+	}
+}
diff --git a/Engine/Plugins/Experimental/EditorDataStorageFeatures/Source/TedsActorCompatibility/Private/ActorComponentDebugHierarchyWidget/ActorComponentTree.h b/Engine/Plugins/Experimental/EditorDataStorageFeatures/Source/TedsActorCompatibility/Private/ActorComponentDebugHierarchyWidget/ActorComponentTree.h
new file mode 100644
--- /dev/null
+++ b/Engine/Plugins/Experimental/EditorDataStorageFeatures/Source/TedsActorCompatibility/Private/ActorComponentDebugHierarchyWidget/ActorComponentTree.h
@@ -0,0 +1,8 @@
+﻿// Copyright Epic Games, Inc. All Rights Reserved.
+
+#pragma once
+
+namespace UE::Editor::DataStorage::ActorCompatibility
+{
+	void RegisterActionComponentDebugHierarchyWidget();
+}
\ No newline at end of file
diff --git a/Engine/Plugins/Experimental/EditorDataStorageFeatures/Source/TedsActorCompatibility/Private/ActorComponentDebugHierarchyWidget/QueryStackNodes/RowQueryCallbackResultsNode.cpp b/Engine/Plugins/Experimental/EditorDataStorageFeatures/Source/TedsActorCompatibility/Private/ActorComponentDebugHierarchyWidget/QueryStackNodes/RowQueryCallbackResultsNode.cpp
new file mode 100644
--- /dev/null
+++ b/Engine/Plugins/Experimental/EditorDataStorageFeatures/Source/TedsActorCompatibility/Private/ActorComponentDebugHierarchyWidget/QueryStackNodes/RowQueryCallbackResultsNode.cpp
@@ -0,0 +1,38 @@
+﻿// Copyright Epic Games, Inc. All Rights Reserved.
+
+#include "RowQueryCallbackResultsNode.h"
+
+#include "Elements/Framework/TypedElementQueryBuilder.h"
+#include "Elements/Interfaces/TypedElementQueryStorageInterfaces.h"
+
+namespace UE::Editor::DataStorage::QueryStack
+{
+	FRowQueryCallbackResultsNode::FRowQueryCallbackResultsNode(
+		ICoreProvider& Storage,
+		TSharedPtr<IQueryNode> InQueryNode,
+		ESyncFlags SyncFlags,
+		CallbackFn InCallbackFn)
+		: FRowQueryResultsNode(Storage, InQueryNode, SyncFlags)
+		, Callback(InCallbackFn)
+	{
+	}
+
+	void FRowQueryCallbackResultsNode::RefreshInternal(FRowHandleArray& TargetRows)
+	{
+		using namespace UE::Editor::DataStorage::Queries;
+		Storage.RunQuery(QueryNode->GetQuery(), 
+			EDirectQueryExecutionFlags::AllowBoundQueries,
+			CreateDirectQueryCallbackBinding(
+		[&TargetRows, this](IDirectQueryContext& Context, const RowHandle* QueryRowResults)
+		{
+			auto EmitRows = [&TargetRows](TArrayView<const RowHandle> InRows)
+			{
+				TargetRows.Append(InRows);
+			};
+
+			TArrayView<const RowHandle> QueryRowResultsView = MakeArrayView(QueryRowResults, Context.GetRowCount());
+
+			Callback(Context, QueryRowResultsView, EmitRows);
+		}));
+	}
+}
diff --git a/Engine/Plugins/Experimental/EditorDataStorageFeatures/Source/TedsActorCompatibility/Private/ActorComponentDebugHierarchyWidget/QueryStackNodes/RowQueryCallbackResultsNode.h b/Engine/Plugins/Experimental/EditorDataStorageFeatures/Source/TedsActorCompatibility/Private/ActorComponentDebugHierarchyWidget/QueryStackNodes/RowQueryCallbackResultsNode.h
new file mode 100644
--- /dev/null
+++ b/Engine/Plugins/Experimental/EditorDataStorageFeatures/Source/TedsActorCompatibility/Private/ActorComponentDebugHierarchyWidget/QueryStackNodes/RowQueryCallbackResultsNode.h
@@ -0,0 +1,36 @@
+﻿// Copyright Epic Games, Inc. All Rights Reserved.
+
+#pragma once
+#include "TedsRowQueryResultsNode.h"
+
+namespace UE::Editor::DataStorage
+{
+	struct IDirectQueryContext;
+}
+
+namespace UE::Editor::DataStorage::QueryStack
+{
+	/**
+	* A node that is configured with a query and a callback.  The callback can be used to
+	 * read and filter data from the rows that match the query and can emit 0 or more rows that
+	 * will be added to the node.
+	 * This node can be useful for collating rows referenced from other rows.
+	 */
+	class FRowQueryCallbackResultsNode : public FRowQueryResultsNode
+	{
+	public:
+		using EmitRowFn = TFunctionRef<void(TArrayView<const RowHandle>)>;
+		using CallbackFn = TFunction<void(IDirectQueryContext&, TArrayView<const RowHandle>, EmitRowFn EmitRows)>;
+		FRowQueryCallbackResultsNode(
+			ICoreProvider& Storage,
+			TSharedPtr<IQueryNode> InQueryNode,
+			ESyncFlags SyncFlags,
+			CallbackFn InCallback);
+	
+	protected:
+		virtual void RefreshInternal(FRowHandleArray& TargetRows) override;
+
+	private:
+		CallbackFn Callback;
+	};
+}
diff --git a/Engine/Plugins/Experimental/EditorDataStorageFeatures/Source/TedsActorCompatibility/Private/ActorComponentDebugHierarchyWidget/SActorComponentTreeWidget.cpp b/Engine/Plugins/Experimental/EditorDataStorageFeatures/Source/TedsActorCompatibility/Private/ActorComponentDebugHierarchyWidget/SActorComponentTreeWidget.cpp
new file mode 100644
--- /dev/null
+++ b/Engine/Plugins/Experimental/EditorDataStorageFeatures/Source/TedsActorCompatibility/Private/ActorComponentDebugHierarchyWidget/SActorComponentTreeWidget.cpp
@@ -0,0 +1,66 @@
+﻿// Copyright Epic Games, Inc. All Rights Reserved.
+
+#include "SActorComponentTreeWidget.h"
+
+#include "TedsHierarchyNode.h"
+#include "TedsQueryNode.h"
+#include "Elements/Columns/TypedElementCompatibilityColumns.h"
+#include "Elements/Columns/TypedElementLabelColumns.h"
+#include "DataStorage/Features.h"
+#include "Elements/Framework/TypedElementAttributeBinding.h"
+#include "Widgets/STedsHierarchyViewer.h"
+#include "Widgets/Text/STextBlock.h"
+
+#define LOCTEXT_NAMESPACE "ActorComponentTree"
+
+SActorComponentTreeWidget::~SActorComponentTreeWidget() = default;
+
+void SActorComponentTreeWidget::Construct(
+	const FArguments& InArgs,
+	UE::Editor::DataStorage::FHierarchyHandle InHierarchyHandle,
+	TSharedPtr<UE::Editor::DataStorage::QueryStack::IRowNode>& InTopLevelRowProvider)
+{
+	using namespace UE::Editor::DataStorage;
+	using namespace UE::Editor::DataStorage::Queries;
+
+	TopLevelRowProvider = InTopLevelRowProvider;
+	HierarchyHandle = InHierarchyHandle;
+
+	ICoreProvider* DataStorageInterface = GetMutableDataStorageFeature<ICoreProvider>(StorageFeatureName);
+	
+	TFunction<RowHandle(const void*, const UScriptStruct*)> GetParentRowFromColumnPredicate = DataStorageInterface->CreateParentExtractionFunction(HierarchyHandle);
+	
+	const UScriptStruct* HierarchyStruct = DataStorageInterface->GetHierarchyDataColumnType(HierarchyHandle);
+
+	TSharedRef<QueryStack::IRowNode> HierarchyRowNode = MakeShared<FHierarchyRowNode>(
+		*DataStorageInterface,
+		HierarchyHandle,
+		TopLevelRowProvider,
+		FHierarchyRowNode::ESyncFlags::IncrementWhenDifferent);
+	
+	TSharedPtr<QueryStack::IRowNode> HierarchyRowNodePtr = HierarchyRowNode.ToSharedPtr();
+	
+	TSharedRef<SHierarchyViewer> HierarchyViewer =
+		SNew(SHierarchyViewer,
+			HierarchyStruct,
+			MoveTemp(GetParentRowFromColumnPredicate))
+			.AllNodeProvider(HierarchyRowNode)
+			.TopLevelNodeProvider(TopLevelRowProvider)
+			.EmptyRowsMessage(LOCTEXT("EmptyRowsMessage", "The provided query has no results."))
+			.CellWidgetPurpose(IUiProvider::FPurposeInfo("General", "RowLabel", NAME_None).GeneratePurposeID())
+			.Columns({FTypedElementLabelColumn::StaticStruct()});
+
+	ChildSlot
+	[
+		HierarchyViewer
+	];
+}
+
+bool SActorComponentTreeWidget::TickProvider(float)
+{
+	TopLevelRowProvider->Update();
+	return true; // Keep ticking
+}
+
+#undef LOCTEXT_NAMESPACE
+
diff --git a/Engine/Plugins/Experimental/EditorDataStorageFeatures/Source/TedsActorCompatibility/Private/ActorComponentDebugHierarchyWidget/SActorComponentTreeWidget.h b/Engine/Plugins/Experimental/EditorDataStorageFeatures/Source/TedsActorCompatibility/Private/ActorComponentDebugHierarchyWidget/SActorComponentTreeWidget.h
new file mode 100644
--- /dev/null
+++ b/Engine/Plugins/Experimental/EditorDataStorageFeatures/Source/TedsActorCompatibility/Private/ActorComponentDebugHierarchyWidget/SActorComponentTreeWidget.h
@@ -0,0 +1,47 @@
+﻿// Copyright Epic Games, Inc. All Rights Reserved.
+
+#pragma once
+#include "Containers/Ticker.h"
+#include "DataStorage/Handles.h"
+#include "Elements/Interfaces/TypedElementDataStorageFactory.h"
+#include "Elements/Interfaces/TypedElementDataStorageUiInterface.h"
+#include "Widgets/DeclarativeSyntaxSupport.h"
+#include "Widgets/SCompoundWidget.h"
+#include "Widgets/SWidget.h"
+
+namespace UE::Editor::DataStorage
+{
+	class STedsTableViewer;
+
+	namespace QueryStack
+	{
+		class IRowNode;
+		class FRowViewNode;
+	}
+}
+
+class SActorComponentTreeWidget : public SCompoundWidget
+{
+public:
+	SLATE_BEGIN_ARGS(SActorComponentTreeWidget) {}
+		
+	SLATE_END_ARGS();
+
+public:
+	virtual ~SActorComponentTreeWidget() override;
+
+	void Construct(
+		const FArguments& InArgs,
+		UE::Editor::DataStorage::FHierarchyHandle InHierarchyHandle,
+		TSharedPtr<UE::Editor::DataStorage::QueryStack::IRowNode>& InTopLevelRowProvider);
+
+private:
+
+	bool TickProvider(float);
+
+	// Provides the top level rows
+	TSharedPtr<UE::Editor::DataStorage::QueryStack::IRowNode> TopLevelRowProvider;
+	// Provides the rows that will be shown in the associated tree view
+	TSharedPtr<UE::Editor::DataStorage::QueryStack::IRowNode> TreeRowProvider;
+	UE::Editor::DataStorage::FHierarchyHandle HierarchyHandle;
+};
\ No newline at end of file
diff --git a/Engine/Plugins/Experimental/EditorDataStorageFeatures/Source/TedsActorCompatibility/Private/TedsActorCompatibilityModule.cpp b/Engine/Plugins/Experimental/EditorDataStorageFeatures/Source/TedsActorCompatibility/Private/TedsActorCompatibilityModule.cpp
--- a/Engine/Plugins/Experimental/EditorDataStorageFeatures/Source/TedsActorCompatibility/Private/TedsActorCompatibilityModule.cpp
+++ b/Engine/Plugins/Experimental/EditorDataStorageFeatures/Source/TedsActorCompatibility/Private/TedsActorCompatibilityModule.cpp
@@ -2,15 +2,209 @@
 
 #include "TedsActorCompatibilityModule.h"
 
+#include "ActorComponentDebugHierarchyWidget/ActorComponentTree.h"
+#include "Columns/TedsActorComponentCompatibilityColumns.h"
+#include "GameFramework/Actor.h"
+#include "Components/ActorComponent.h"
+#include "Elements/Columns/TypedElementLabelColumns.h"
+#include "DataStorage/Features.h"
+#include "Elements/Columns/TypedElementCompatibilityColumns.h"
+#include "Elements/Columns/TypedElementMiscColumns.h"
+#include "Elements/Columns/TypedElementTypeInfoColumns.h"
+#include "Elements/Framework/TypedElementQueryBuilder.h"
+#include "Elements/Interfaces/TypedElementDataStorageCompatibilityInterface.h"
+#include "Elements/Interfaces/TypedElementDataStorageInterface.h"
+#include "HAL/IConsoleManager.h"
 #include "Modules/ModuleManager.h"
 
 namespace UE::Editor::DataStorage
 {
+	static bool bEnableActorComponentCompatibility = true;
+	FAutoConsoleVariableRef CVarEnableActorComponentRegistrator(
+		TEXT("TEDS.Feature.ActorCompatibility.ActorComponents.Enable"),
+		bEnableActorComponentCompatibility,
+		TEXT("Enables ActorComponent tracking in TEDS"),
+		ECVF_ReadOnly);
+
+	static bool bEnableDemoTab = false;
+	FAutoConsoleVariableRef CVarAddDemoTab(
+		TEXT("TEDS.Debug.ActorCompatibility.ActorComponents.EnableHierarchyViewer"),
+		bEnableDemoTab,
+		TEXT("Enables a debug tab to demonstrate usage of hierarchy views with ActorComponents"),
+		ECVF_ReadOnly);
+	
 	void FTedsActorCompatibilityModule::StartupModule()
 	{
 		IModuleInterface::StartupModule();
 
 		FModuleManager::Get().LoadModule(TEXT("TypedElementFramework"));
+		FModuleManager::Get().LoadModule(TEXT("EditorDataStorageHierarchy"));
+
+		if (bEnableActorComponentCompatibility)
+		{
+			using namespace UE::Editor::DataStorage;
+			using namespace UE::Editor::DataStorage::Queries;
+
+			UE::Editor::DataStorage::OnEditorDataStorageFeaturesEnabled().AddLambda([this]()
+			{
+				using namespace UE::Editor::DataStorage;
+				using namespace UE::Editor::DataStorage::Queries;
+				
+				ICoreProvider* CoreProvider = GetMutableDataStorageFeature<ICoreProvider>(StorageFeatureName);
+		
+				const FHierarchyRegistrationParams HierarchyParams
+				{
+					.Name = TEXT("ActorComponent"),
+				};
+
+				CoreProvider->RegisterHierarchy(HierarchyParams);
+		
+				TableHandle ActorComponentTable = CoreProvider->RegisterTable(
+					{
+						FTypedElementUObjectColumn::StaticStruct(),
+						FTypedElementUObjectIdColumn::StaticStruct(),
+						FTypedElementClassTypeInfoColumn::StaticStruct(),
+						FActorComponentTypeTag::StaticStruct(),
+						FTypedElementLabelColumn::StaticStruct(),
+						FTypedElementLabelHashColumn::StaticStruct(),
+						FTypedElementSyncFromWorldTag::StaticStruct()
+					},
+					FName("Editor_ActorComponentTable"));
+
+				ICompatibilityProvider* CompatibilityProvider = GetMutableDataStorageFeature<ICompatibilityProvider>(CompatibilityFeatureName);
+				CompatibilityProvider->RegisterTypeTableAssociation(UActorComponent::StaticClass(), ActorComponentTable);
+
+				// 1. For actors that need syncing, add any components they have to the compatibility layer
+				// This solves a lot of issues around trying to intercept all the places components are added and removed
+				// In a perfect world, this may not be required
+				struct FRegisterActorComponentsWithCompatibility
+				{
+					TWeakObjectPtr<AActor> WkActor;
+					RowHandle ActorRow;
+					TArray<TWeakObjectPtr<UActorComponent>> ComponentsToRegister;
+
+					void operator()()
+					{
+						AActor* Actor = WkActor.Get();
+						if (!Actor)
+						{
+							return;
+						}
+						ICompatibilityProvider* CompatibilityProvider = GetMutableDataStorageFeature<ICompatibilityProvider>(CompatibilityFeatureName);
+
+						for (TWeakObjectPtr<UActorComponent> Component : ComponentsToRegister)
+						{
+							UActorComponent* NewComponent = Component.Get();
+							if (!NewComponent)
+							{
+								CompatibilityProvider->AddCompatibleObject(NewComponent);
+							}
+						}
+					};
+				};
+				CoreProvider->RegisterQuery(
+					Select(
+						TEXT("Register ActorComponent Ownership"),
+						FProcessor(EQueryTickPhase::DuringPhysics, CoreProvider->GetQueryTickGroupName(EQueryTickGroups::Default)),
+						[](IQueryContext& Context, RowHandle Row, const FTypedElementUObjectColumn& ObjectColumn)
+						{
+							AActor* Actor = static_cast<AActor*>(ObjectColumn.Object.Get());
+
+							FRegisterActorComponentsWithCompatibility Command;
+							Command.WkActor = Actor;
+							Command.ActorRow = Row;
+							
+							constexpr bool bIncludeFromChildActors = false;
+							Actor->ForEachComponent(bIncludeFromChildActors, [&Context, &Command](UActorComponent* Component)
+							{
+								RowHandle ComponentRow = Context.LookupMappedRow(ICompatibilityProvider::ObjectMappingDomain, FMapKeyView(static_cast<const UObject*>(Component)));
+								if (ComponentRow == InvalidRowHandle)
+								{
+									Command.ComponentsToRegister.Add(Component);
+								}
+							});
+							
+							if (!Command.ComponentsToRegister.IsEmpty())
+							{
+								Context.PushCommand<FRegisterActorComponentsWithCompatibility>(MoveTemp(Command));
+							}
+						})
+					.Where()
+						.All<FTypedElementActorTag>()
+						.Any<FTypedElementSyncFromWorldTag>()
+					.Compile());
+
+				// 2. For components that aren't a child, which means they don't have an associated owner, set up that owner
+				//    If a component is not associated with an owner, then remove that association in TEDS - in practice this is usally only a
+				//    transitory state
+				//    Note: ActorComponent ownership can change due to renaming.  
+				CoreProvider->RegisterQuery(
+					Select(
+						TEXT("Set ActorComponent Hierarchy"),
+						FProcessor(EQueryTickPhase::PostPhysics, CoreProvider->GetQueryTickGroupName(EQueryTickGroups::Default)),
+						[](IQueryContext& Context, RowHandle ComponentRow, const FTypedElementUObjectColumn& ObjectColumn)
+						{
+							if (UActorComponent* ActorComponent = Cast<UActorComponent>(ObjectColumn.Object.Get()))
+							{
+								// Confusingly, an ActorComponent could indicate that it has an owner while its
+								// owner does not consider it owned.
+								// Defer to the owner's source of truth for this
+								const AActor* OwnerActor = ActorComponent->GetOwner();
+								RowHandle ActorRow;
+								if (OwnerActor && OwnerActor->GetComponents().Contains(ActorComponent))
+								{
+									ActorRow = Context.LookupMappedRow(ICompatibilityProvider::ObjectMappingDomain, FMapKeyView(static_cast<const UObject*>(OwnerActor)));
+								}
+								else
+								{
+									ActorRow = InvalidRowHandle;
+								}
+								
+								const bool bNotAssociatedWithActor = OwnerActor == nullptr && ActorRow == InvalidRowHandle;
+								const bool bActorRowAssigned = Context.IsRowAssigned(ActorRow);
+
+								RowHandle ParentRowOfComponent = Context.GetParentRow(ComponentRow);
+								
+								if (bNotAssociatedWithActor || (bActorRowAssigned && ParentRowOfComponent != ActorRow))
+								{
+									Context.SetParentRow(ComponentRow, ActorRow);
+								}
+							}
+						})
+						.AccessesHierarchy(TEXT("ActorComponent"))
+						.Where()
+							.All<FActorComponentTypeTag, FTypedElementSyncFromWorldTag>()
+						.Compile());
+
+				// 
+				CoreProvider->RegisterQuery(
+					Select(
+						TEXT("Sync actorcomponent label to column"),
+						FProcessor(EQueryTickPhase::PrePhysics, CoreProvider->GetQueryTickGroupName(EQueryTickGroups::SyncExternalToDataStorage)),
+						[](const FTypedElementUObjectColumn& Actor, FTypedElementLabelColumn& Label, FTypedElementLabelHashColumn& LabelHash)
+						{
+							if (const UActorComponent* ComponentInstance = Cast<UActorComponent>(Actor.Object); ComponentInstance != nullptr)
+							{
+								const FString Name = ComponentInstance->GetName();
+								uint64 NameHash = CityHash64(reinterpret_cast<const char*>(*Name), Name.Len() * sizeof(**Name));
+								if (LabelHash.LabelHash != NameHash)
+								{
+									Label.Label = Name;
+									LabelHash.LabelHash = NameHash;
+								}
+							}
+						}
+					)
+					.Where()
+						.All<FActorComponentTypeTag, FTypedElementSyncFromWorldTag>()
+					.Compile());
+			});
+		}
+
+		if (bEnableDemoTab)
+		{
+			UE::Editor::DataStorage::ActorCompatibility::RegisterActionComponentDebugHierarchyWidget();
+		}
 	}
 
 	void FTedsActorCompatibilityModule::ShutdownModule()
diff --git a/Engine/Plugins/Experimental/EditorDataStorageFeatures/Source/TedsActorCompatibility/Public/Columns/TedsActorComponentCompatibilityColumns.h b/Engine/Plugins/Experimental/EditorDataStorageFeatures/Source/TedsActorCompatibility/Public/Columns/TedsActorComponentCompatibilityColumns.h
new file mode 100644
--- /dev/null
+++ b/Engine/Plugins/Experimental/EditorDataStorageFeatures/Source/TedsActorCompatibility/Public/Columns/TedsActorComponentCompatibilityColumns.h
@@ -0,0 +1,12 @@
+// Copyright Epic Games, Inc. All Rights Reserved.
+
+#pragma once
+#include "DataStorage/CommonTypes.h"
+
+#include "TedsActorComponentCompatibilityColumns.generated.h"
+
+USTRUCT()
+struct FActorComponentTypeTag : public FEditorDataStorageTag
+{
+	GENERATED_BODY()
+};
\ No newline at end of file
diff --git a/Engine/Plugins/Experimental/EditorDataStorageFeatures/Source/TedsActorCompatibility/TedsActorCompatibility.Build.cs b/Engine/Plugins/Experimental/EditorDataStorageFeatures/Source/TedsActorCompatibility/TedsActorCompatibility.Build.cs
--- a/Engine/Plugins/Experimental/EditorDataStorageFeatures/Source/TedsActorCompatibility/TedsActorCompatibility.Build.cs
+++ b/Engine/Plugins/Experimental/EditorDataStorageFeatures/Source/TedsActorCompatibility/TedsActorCompatibility.Build.cs
@@ -24,7 +24,12 @@ public class TedsActorCompatibility : ModuleRules
 					"CoreUObject",
 					"Engine",
 					"TypedElementFramework",
-					"UnrealEd"
+					"UnrealEd",
+					"TedsQueryStack",
+					"WorkspaceMenuStructure",
+					"TedsTableViewer",
+					"SlateCore",
+					"Slate",
 				});
 
 			DynamicallyLoadedModuleNames.AddRange(new string[] { });
diff --git a/Engine/Plugins/Experimental/EditorDataStorageFeatures/Source/TedsQueryStack/Private/TedsHierarchyNode.cpp b/Engine/Plugins/Experimental/EditorDataStorageFeatures/Source/TedsQueryStack/Private/TedsHierarchyNode.cpp
--- a/Engine/Plugins/Experimental/EditorDataStorageFeatures/Source/TedsQueryStack/Private/TedsHierarchyNode.cpp
+++ b/Engine/Plugins/Experimental/EditorDataStorageFeatures/Source/TedsQueryStack/Private/TedsHierarchyNode.cpp
@@ -10,7 +10,7 @@ namespace UE::Editor::DataStorage
 	FHierarchyRowNode::FHierarchyRowNode(
 		ICoreProvider& Storage,
 		FHierarchyHandle InHierarchyHandle,
-		const TSharedPtr<QueryStack::IRowNode>& InTopLevelRows,
+		TSharedPtr<QueryStack::IRowNode>& InTopLevelRows,
 		ESyncFlags InSyncFlags)
 		: TopLevelRows(InTopLevelRows)
 		, Storage(Storage)
@@ -26,33 +26,46 @@ namespace UE::Editor::DataStorage
 
 	void FHierarchyRowNode::Update()
 	{
-		
-		if (TopLevelRows->GetRevision() != ParentRevision)
-		{
-			Rows.Empty();
-			RefreshInternal(Storage, Rows);
-			++Revision;
-			ParentRevision = TopLevelRows->GetRevision();
-		}
-		if (EnumHasAllFlags(SyncFlags, ESyncFlags::IncrementWhenDifferent))
+		TopLevelRows->Update();
+		switch (SyncFlags)
 		{
-			FRowHandleArray& TempArray = GetTempRowHandleArray();
-			RefreshInternal(Storage, TempArray);
-			TempArray.Sort();
-			Rows.Sort();
-			if (!TempArray.IsSame(Rows))
+			case ESyncFlags::ParentRevisionDifferent:
 			{
-				Rows = MoveTemp(TempArray);
+				if (TopLevelRows->GetRevision() != ParentRevision)
+				{
+					Rows.Empty();
+					RefreshInternal(Storage, Rows);
+					++Revision;
+					ParentRevision = TopLevelRows->GetRevision();
+				}
+				break;
+			}
+			case ESyncFlags::IncrementWhenDifferent:
+			{
+				FRowHandleArray& TempArray = GetTempRowHandleArray();
+				RefreshInternal(Storage, TempArray);
+				TempArray.Sort();
+				Rows.Sort();
+				if (!TempArray.IsSame(Rows))
+				{
+					Rows = MoveTemp(TempArray);
+					++Revision;
+					ParentRevision = TopLevelRows->GetRevision();
+				}
+				break;
+			}
+			case ESyncFlags::Always:
+			{
+				Rows.Empty();
+				RefreshInternal(Storage, Rows);
 				++Revision;
 				ParentRevision = TopLevelRows->GetRevision();
+				break;
+			}
+			default:
+			{
+				// Do nothing
 			}
-		}
-		else
-		{
-			Rows.Empty();
-			RefreshInternal(Storage, Rows);
-			++Revision;
-			ParentRevision = TopLevelRows->GetRevision();
 		}
 	}
 
diff --git a/Engine/Plugins/Experimental/EditorDataStorageFeatures/Source/TedsQueryStack/Public/TedsHierarchyNode.h b/Engine/Plugins/Experimental/EditorDataStorageFeatures/Source/TedsQueryStack/Public/TedsHierarchyNode.h
--- a/Engine/Plugins/Experimental/EditorDataStorageFeatures/Source/TedsQueryStack/Public/TedsHierarchyNode.h
+++ b/Engine/Plugins/Experimental/EditorDataStorageFeatures/Source/TedsQueryStack/Public/TedsHierarchyNode.h
@@ -13,25 +13,25 @@ namespace UE::Editor::DataStorage
 	public:
 		enum class ESyncFlags
 		{
-			None = 0,
 			/** 
-			 * Update the row list whenever "Update" is called. This can take several milliseconds for large lists. It's
-			 * recommended to use a monitoring node to detect changes for large lists.
+			 * Increment the revision only when the parent's revision changed
 			 */
-			RefreshOnUpdate = 1 << 0,
-			/** Update the row list whenever the parent changes. */
-			RefreshOnParentChange = 1 << 1,
+			ParentRevisionDifferent = 0,
 			/** 
 			 * Compares the previous row list with the current and only updates the revision if there are differences.
 			 */
-			IncrementWhenDifferent = 1 << 2
+			IncrementWhenDifferent = 1,
+			/**
+			 * Update the row list whenever "Update" is called. 
+			 */
+			Always
 		};
 		using EmitRowFn = TFunctionRef<void(TArrayView<const RowHandle>)>;
 		
 		TEDSQUERYSTACK_API FHierarchyRowNode(
 			ICoreProvider& Storage,
 			FHierarchyHandle InHierarchyHandle,
-			const TSharedPtr<QueryStack::IRowNode>& InTopLevelRows,
+			TSharedPtr<QueryStack::IRowNode>& InTopLevelRows,
 			ESyncFlags SyncFlags);
 		
 		TEDSQUERYSTACK_API virtual RevisionId GetRevision() const override;
@@ -45,7 +45,7 @@ namespace UE::Editor::DataStorage
 		
 		FRowHandleArray Rows;
 		FRowHandleArray NewRows;
-		const TSharedPtr<QueryStack::IRowNode> TopLevelRows;
+		TSharedPtr<QueryStack::IRowNode> TopLevelRows;
 		RevisionId ParentRevision = 0;
 		
 		ICoreProvider& Storage;
diff --git a/Engine/Plugins/Experimental/EditorDataStorageFeatures/Source/TedsQueryStack/Public/TedsRowQueryResultsNode.h b/Engine/Plugins/Experimental/EditorDataStorageFeatures/Source/TedsQueryStack/Public/TedsRowQueryResultsNode.h
--- a/Engine/Plugins/Experimental/EditorDataStorageFeatures/Source/TedsQueryStack/Public/TedsRowQueryResultsNode.h
+++ b/Engine/Plugins/Experimental/EditorDataStorageFeatures/Source/TedsQueryStack/Public/TedsRowQueryResultsNode.h
@@ -50,8 +50,8 @@ namespace UE::Editor::DataStorage::QueryStack
 		TEDSQUERYSTACK_API virtual FRowHandleArrayView GetRows() const override;
 		TEDSQUERYSTACK_API virtual FRowHandleArray& GetMutableRows() override;
 
-	private:
-		void RefreshInternal(FRowHandleArray& TargetRows);
+	protected:
+		TEDSQUERYSTACK_API virtual void RefreshInternal(FRowHandleArray& TargetRows);
 
 		FRowHandleArray Rows;
 		TSharedPtr<IQueryNode> QueryNode;
diff --git a/Engine/Source/Runtime/Engine/Private/Actor.cpp b/Engine/Source/Runtime/Engine/Private/Actor.cpp
--- a/Engine/Source/Runtime/Engine/Private/Actor.cpp
+++ b/Engine/Source/Runtime/Engine/Private/Actor.cpp
@@ -72,11 +72,15 @@
 #endif
 
 #if WITH_EDITOR
+#include "Features/IModularFeatures.h"
 #include "FoliageHelper.h"
 #include "LevelInstance/LevelInstanceSubsystem.h"
 #include "LevelInstance/LevelInstanceInterface.h"
 #endif
 
+#include "Elements/Columns/TypedElementMiscColumns.h"
+#include "DataStorage/Features.h"
+#include "Elements/Interfaces/TypedElementDataStorageCompatibilityInterface.h"
 #include "WorldPartition/WorldPartitionLog.h"
 #include "WorldPartition/DataLayer/WorldDataLayers.h"
 #include "WorldPartition/WorldPartitionActorDescUtils.h"
@@ -514,7 +518,7 @@ void AActor::ResetOwnedComponents()
 	OwnedComponents.Reset();
 	ReplicatedComponents.Reset();
 	ReplicatedComponentsInfo.Reset();
-
+	
 	ForEachObjectWithOuter(this, [this](UObject* Child)
 	{
 		UActorComponent* Component = Cast<UActorComponent>(Child);
@@ -530,6 +534,19 @@ void AActor::ResetOwnedComponents()
 			}
 		}
 	}, true, RF_NoFlags, EInternalObjectFlags::Garbage);
+
+#if WITH_EDITOR
+	{
+		using namespace UE::Editor::DataStorage;
+		if (ICompatibilityProvider* Storage = UE::Editor::DataStorage::GetMutableDataStorageFeature<ICompatibilityProvider>(CompatibilityFeatureName))
+		{
+			for (UActorComponent* Component : OwnedComponents)
+			{
+				Storage->AddCompatibleObject(const_cast<UActorComponent*>(Component));
+			}
+		}
+	}
+#endif
 }
 
 void AActor::PostInitProperties()
@@ -3843,6 +3860,16 @@ void AActor::AddOwnedComponent(UActorComponent* Component)
 		{
 			InstanceComponents.Add(Component);
 		}
+#if WITH_EDITOR
+		{
+			using namespace UE::Editor::DataStorage;
+			if (ICoreProvider* Storage = UE::Editor::DataStorage::GetMutableDataStorageFeature<ICoreProvider>(StorageFeatureName))
+			{
+				RowHandle ComponentRow = Storage->LookupMappedRow(ICompatibilityProvider::ObjectMappingDomain, FMapKeyView(Component));
+				Storage->AddColumns<FTypedElementSyncFromWorldTag>(ComponentRow);
+			}
+		}
+#endif
 	}
 }
 
@@ -3866,6 +3893,18 @@ void AActor::RemoveOwnedComponent(UActorComponent* Component)
 		{
 			InstanceComponents.RemoveSingleSwap(Component);
 		}
+
+		
+#if WITH_EDITOR
+		{
+			using namespace UE::Editor::DataStorage;
+			if (ICoreProvider* Storage = UE::Editor::DataStorage::GetMutableDataStorageFeature<ICoreProvider>(StorageFeatureName))
+			{
+				RowHandle ComponentRow = Storage->LookupMappedRow(ICompatibilityProvider::ObjectMappingDomain, FMapKeyView(Component));
+				Storage->AddColumns<FTypedElementSyncFromWorldTag>(ComponentRow);
+			}
+		}
+#endif
 	}
 }
 
diff --git a/Engine/Source/Runtime/Engine/Private/Components/ActorComponent.cpp b/Engine/Source/Runtime/Engine/Private/Components/ActorComponent.cpp
--- a/Engine/Source/Runtime/Engine/Private/Components/ActorComponent.cpp
+++ b/Engine/Source/Runtime/Engine/Private/Components/ActorComponent.cpp
@@ -48,6 +48,9 @@
 #include "Kismet2/ComponentEditorUtils.h"
 #include "ObjectCacheEventSink.h"
 #include "StaticMeshCompiler.h"
+#include "Elements/Columns/TypedElementMiscColumns.h"
+#include "DataStorage/Features.h"
+#include "Elements/Interfaces/TypedElementDataStorageCompatibilityInterface.h"
 #endif
 
 #if UE_WITH_IRIS
@@ -593,6 +596,11 @@ void UActorComponent::PostInitProperties()
 
 #if WITH_EDITOR
 	UEngineElementsLibrary::CreateEditorComponentElement(this);
+	using namespace UE::Editor::DataStorage;
+	if (ICompatibilityProvider* Storage = UE::Editor::DataStorage::GetMutableDataStorageFeature<ICompatibilityProvider>(CompatibilityFeatureName))
+	{
+		Storage->AddCompatibleObject(this);
+	}
 #endif	// WITH_EDITOR
 
 	// Instance components will be added during the owner's initialization
@@ -888,6 +896,15 @@ void UActorComponent::PostRename(UObject* OldOuter, const FName OldName)
 				OwnerPrivate->AddOwnedComponent(this);
 			}
 
+#if WITH_EDITOR
+			using namespace UE::Editor::DataStorage;
+			if (ICoreProvider* Storage = UE::Editor::DataStorage::GetMutableDataStorageFeature<ICoreProvider>(StorageFeatureName))
+			{
+				RowHandle ComponentRow = Storage->LookupMappedRow(ICompatibilityProvider::ObjectMappingDomain, FMapKeyView(this));
+				Storage->AddColumns<FTypedElementSyncFromWorldTag>(ComponentRow);
+			}
+#endif
+
 			TArray<UObject*> Children;
 			GetObjectsWithOuter(this, Children, /*bIncludeNestedObjects=*/false);
 
@@ -1113,6 +1130,11 @@ void UActorComponent::BeginDestroy()
 	}
 
 #if WITH_EDITOR
+	using namespace UE::Editor::DataStorage;
+	if (ICompatibilityProvider* Storage = UE::Editor::DataStorage::GetMutableDataStorageFeature<ICompatibilityProvider>(CompatibilityFeatureName))
+	{
+		Storage->RemoveCompatibleObject(this);
+	}
 	UEngineElementsLibrary::DestroyEditorComponentElement(this);
 #endif	// WITH_EDITOR
 
@@ -2099,6 +2121,16 @@ void UActorComponent::OnComponentDestroyed(bool bDestroyingHierarchy)
 {
 	// @TODO: Would be nice to ensure(bHasBeenCreated), but there are still many places where components are created without calling OnComponentCreated
 	bHasBeenCreated = false;
+
+#if WITH_EDITOR
+	{
+		using namespace UE::Editor::DataStorage;
+		if (ICompatibilityProvider* Storage = UE::Editor::DataStorage::GetMutableDataStorageFeature<ICompatibilityProvider>(CompatibilityFeatureName))
+		{
+			Storage->RemoveCompatibleObject(this);
+		}
+	}
+#endif
 }
 
 void UActorComponent::K2_DestroyComponent(UObject* Object)
diff --git a/Engine/Source/Runtime/Engine/Private/Elements/Framework/EngineElementsLibrary.cpp b/Engine/Source/Runtime/Engine/Private/Elements/Framework/EngineElementsLibrary.cpp
--- a/Engine/Source/Runtime/Engine/Private/Elements/Framework/EngineElementsLibrary.cpp
+++ b/Engine/Source/Runtime/Engine/Private/Elements/Framework/EngineElementsLibrary.cpp
@@ -383,9 +383,15 @@ void UEngineElementsLibrary::RegisterActorElement(const AActor* InActor)
 {
 	using namespace UE::Editor::DataStorage;
 	ICompatibilityProvider* Storage = GetMutableDataStorageFeature<ICompatibilityProvider>(CompatibilityFeatureName);
+	RowHandle AddedObjectRow = InvalidRowHandle;
 	if (Storage)
 	{
-		Storage->AddCompatibleObject(const_cast<AActor*>(InActor));
+		AddedObjectRow = Storage->AddCompatibleObject(const_cast<AActor*>(InActor));
+		constexpr bool bIncludeChildActors = false;
+		InActor->ForEachComponent(bIncludeChildActors, [&Storage](UActorComponent* Component)
+		{
+			Storage->AddCompatibleObject(const_cast<UActorComponent*>(Component));
+		});
 	}
 }
 
diff --git a/Engine/Source/Runtime/Engine/Private/Level.cpp b/Engine/Source/Runtime/Engine/Private/Level.cpp
--- a/Engine/Source/Runtime/Engine/Private/Level.cpp
+++ b/Engine/Source/Runtime/Engine/Private/Level.cpp
@@ -1094,6 +1094,11 @@ void ULevel::AddLoadedActors(const TArray<AActor*>& ActorList, const FTransform*
 			if (TedsCompat)
 			{
 				TedsCompat->AddCompatibleObject(Actor);
+				constexpr bool bIncludeChildActors = false;
+				Actor->ForEachComponent(bIncludeChildActors, [&TedsCompat](UActorComponent* Component)
+				{
+					TedsCompat->AddCompatibleObject(Component);
+				});
 			}
 		}
 
@@ -1207,6 +1212,12 @@ void ULevel::RemoveLoadedActors(const TArray<AActor*>& ActorList, const FTransfo
 			// but for our purposes it is considered dead already
 			if (TedsCompat)
 			{
+				constexpr bool bIncludeChildActors = false;
+				Actor->ForEachComponent(bIncludeChildActors, [&TedsCompat](UActorComponent* Component)
+				{
+					TedsCompat->RemoveCompatibleObject(Component);
+				});
+				
 				TedsCompat->RemoveCompatibleObject(Actor);
 			}
 
