[Backout] - CL44050924
#rnx
Original CL Desc
-----------------------------------------------------------------
ActorComponent TEDS Registration

diff --git a/Engine/Plugins/Experimental/EditorDataStorageFeatures/Source/TedsActorCompatibility/Private/ActorComponentDebugHierarchyWidget/ActorComponentTree.cpp b/Engine/Plugins/Experimental/EditorDataStorageFeatures/Source/TedsActorCompatibility/Private/ActorComponentDebugHierarchyWidget/ActorComponentTree.cpp
deleted file mode 100644
diff --git a/Engine/Plugins/Experimental/EditorDataStorageFeatures/Source/TedsActorCompatibility/Private/ActorComponentDebugHierarchyWidget/ActorComponentTree.h b/Engine/Plugins/Experimental/EditorDataStorageFeatures/Source/TedsActorCompatibility/Private/ActorComponentDebugHierarchyWidget/ActorComponentTree.h
deleted file mode 100644
diff --git a/Engine/Plugins/Experimental/EditorDataStorageFeatures/Source/TedsActorCompatibility/Private/ActorComponentDebugHierarchyWidget/QueryStackNodes/RowQueryCallbackResultsNode.cpp b/Engine/Plugins/Experimental/EditorDataStorageFeatures/Source/TedsActorCompatibility/Private/ActorComponentDebugHierarchyWidget/QueryStackNodes/RowQueryCallbackResultsNode.cpp
deleted file mode 100644
diff --git a/Engine/Plugins/Experimental/EditorDataStorageFeatures/Source/TedsActorCompatibility/Private/ActorComponentDebugHierarchyWidget/QueryStackNodes/RowQueryCallbackResultsNode.h b/Engine/Plugins/Experimental/EditorDataStorageFeatures/Source/TedsActorCompatibility/Private/ActorComponentDebugHierarchyWidget/QueryStackNodes/RowQueryCallbackResultsNode.h
deleted file mode 100644
diff --git a/Engine/Plugins/Experimental/EditorDataStorageFeatures/Source/TedsActorCompatibility/Private/ActorComponentDebugHierarchyWidget/SActorComponentTreeWidget.cpp b/Engine/Plugins/Experimental/EditorDataStorageFeatures/Source/TedsActorCompatibility/Private/ActorComponentDebugHierarchyWidget/SActorComponentTreeWidget.cpp
deleted file mode 100644
diff --git a/Engine/Plugins/Experimental/EditorDataStorageFeatures/Source/TedsActorCompatibility/Private/ActorComponentDebugHierarchyWidget/SActorComponentTreeWidget.h b/Engine/Plugins/Experimental/EditorDataStorageFeatures/Source/TedsActorCompatibility/Private/ActorComponentDebugHierarchyWidget/SActorComponentTreeWidget.h
deleted file mode 100644
diff --git a/Engine/Plugins/Experimental/EditorDataStorageFeatures/Source/TedsActorCompatibility/Private/TedsActorCompatibilityModule.cpp b/Engine/Plugins/Experimental/EditorDataStorageFeatures/Source/TedsActorCompatibility/Private/TedsActorCompatibilityModule.cpp
--- a/Engine/Plugins/Experimental/EditorDataStorageFeatures/Source/TedsActorCompatibility/Private/TedsActorCompatibilityModule.cpp
+++ b/Engine/Plugins/Experimental/EditorDataStorageFeatures/Source/TedsActorCompatibility/Private/TedsActorCompatibilityModule.cpp
@@ -2,209 +2,15 @@
 
 #include "TedsActorCompatibilityModule.h"
 
-#include "ActorComponentDebugHierarchyWidget/ActorComponentTree.h"
-#include "Columns/TedsActorComponentCompatibilityColumns.h"
-#include "GameFramework/Actor.h"
-#include "Components/ActorComponent.h"
-#include "Elements/Columns/TypedElementLabelColumns.h"
-#include "DataStorage/Features.h"
-#include "Elements/Columns/TypedElementCompatibilityColumns.h"
-#include "Elements/Columns/TypedElementMiscColumns.h"
-#include "Elements/Columns/TypedElementTypeInfoColumns.h"
-#include "Elements/Framework/TypedElementQueryBuilder.h"
-#include "Elements/Interfaces/TypedElementDataStorageCompatibilityInterface.h"
-#include "Elements/Interfaces/TypedElementDataStorageInterface.h"
-#include "HAL/IConsoleManager.h"
 #include "Modules/ModuleManager.h"
 
 namespace UE::Editor::DataStorage
 {
-	static bool bEnableActorComponentCompatibility = true;
-	FAutoConsoleVariableRef CVarEnableActorComponentRegistrator(
-		TEXT("TEDS.Feature.ActorCompatibility.ActorComponents.Enable"),
-		bEnableActorComponentCompatibility,
-		TEXT("Enables ActorComponent tracking in TEDS"),
-		ECVF_ReadOnly);
-
-	static bool bEnableDemoTab = false;
-	FAutoConsoleVariableRef CVarAddDemoTab(
-		TEXT("TEDS.Debug.ActorCompatibility.ActorComponents.EnableHierarchyViewer"),
-		bEnableDemoTab,
-		TEXT("Enables a debug tab to demonstrate usage of hierarchy views with ActorComponents"),
-		ECVF_ReadOnly);
-	
 	void FTedsActorCompatibilityModule::StartupModule()
 	{
 		IModuleInterface::StartupModule();
 
 		FModuleManager::Get().LoadModule(TEXT("TypedElementFramework"));
-		FModuleManager::Get().LoadModule(TEXT("EditorDataStorageHierarchy"));
-
-		if (bEnableActorComponentCompatibility)
-		{
-			using namespace UE::Editor::DataStorage;
-			using namespace UE::Editor::DataStorage::Queries;
-
-			UE::Editor::DataStorage::OnEditorDataStorageFeaturesEnabled().AddLambda([this]()
-			{
-				using namespace UE::Editor::DataStorage;
-				using namespace UE::Editor::DataStorage::Queries;
-				
-				ICoreProvider* CoreProvider = GetMutableDataStorageFeature<ICoreProvider>(StorageFeatureName);
-		
-				const FHierarchyRegistrationParams HierarchyParams
-				{
-					.Name = TEXT("ActorComponent"),
-				};
-
-				CoreProvider->RegisterHierarchy(HierarchyParams);
-		
-				TableHandle ActorComponentTable = CoreProvider->RegisterTable(
-					{
-						FTypedElementUObjectColumn::StaticStruct(),
-						FTypedElementUObjectIdColumn::StaticStruct(),
-						FTypedElementClassTypeInfoColumn::StaticStruct(),
-						FActorComponentTypeTag::StaticStruct(),
-						FTypedElementLabelColumn::StaticStruct(),
-						FTypedElementLabelHashColumn::StaticStruct(),
-						FTypedElementSyncFromWorldTag::StaticStruct()
-					},
-					FName("Editor_ActorComponentTable"));
-
-				ICompatibilityProvider* CompatibilityProvider = GetMutableDataStorageFeature<ICompatibilityProvider>(CompatibilityFeatureName);
-				CompatibilityProvider->RegisterTypeTableAssociation(UActorComponent::StaticClass(), ActorComponentTable);
-
-				// 1. For actors that need syncing, add any components they have to the compatibility layer
-				// This solves a lot of issues around trying to intercept all the places components are added and removed
-				// In a perfect world, this may not be required
-				struct FRegisterActorComponentsWithCompatibility
-				{
-					TWeakObjectPtr<AActor> WkActor;
-					RowHandle ActorRow;
-					TArray<TWeakObjectPtr<UActorComponent>> ComponentsToRegister;
-
-					void operator()()
-					{
-						AActor* Actor = WkActor.Get();
-						if (!Actor)
-						{
-							return;
-						}
-						ICompatibilityProvider* CompatibilityProvider = GetMutableDataStorageFeature<ICompatibilityProvider>(CompatibilityFeatureName);
-
-						for (TWeakObjectPtr<UActorComponent> Component : ComponentsToRegister)
-						{
-							UActorComponent* NewComponent = Component.Get();
-							if (!NewComponent)
-							{
-								CompatibilityProvider->AddCompatibleObject(NewComponent);
-							}
-						}
-					};
-				};
-				CoreProvider->RegisterQuery(
-					Select(
-						TEXT("Register ActorComponent Ownership"),
-						FProcessor(EQueryTickPhase::DuringPhysics, CoreProvider->GetQueryTickGroupName(EQueryTickGroups::Default)),
-						[](IQueryContext& Context, RowHandle Row, const FTypedElementUObjectColumn& ObjectColumn)
-						{
-							AActor* Actor = static_cast<AActor*>(ObjectColumn.Object.Get());
-
-							FRegisterActorComponentsWithCompatibility Command;
-							Command.WkActor = Actor;
-							Command.ActorRow = Row;
-							
-							constexpr bool bIncludeFromChildActors = false;
-							Actor->ForEachComponent(bIncludeFromChildActors, [&Context, &Command](UActorComponent* Component)
-							{
-								RowHandle ComponentRow = Context.LookupMappedRow(ICompatibilityProvider::ObjectMappingDomain, FMapKeyView(static_cast<const UObject*>(Component)));
-								if (ComponentRow == InvalidRowHandle)
-								{
-									Command.ComponentsToRegister.Add(Component);
-								}
-							});
-							
-							if (!Command.ComponentsToRegister.IsEmpty())
-							{
-								Context.PushCommand<FRegisterActorComponentsWithCompatibility>(MoveTemp(Command));
-							}
-						})
-					.Where()
-						.All<FTypedElementActorTag>()
-						.Any<FTypedElementSyncFromWorldTag>()
-					.Compile());
-
-				// 2. For components that aren't a child, which means they don't have an associated owner, set up that owner
-				//    If a component is not associated with an owner, then remove that association in TEDS - in practice this is usally only a
-				//    transitory state
-				//    Note: ActorComponent ownership can change due to renaming.  
-				CoreProvider->RegisterQuery(
-					Select(
-						TEXT("Set ActorComponent Hierarchy"),
-						FProcessor(EQueryTickPhase::PostPhysics, CoreProvider->GetQueryTickGroupName(EQueryTickGroups::Default)),
-						[](IQueryContext& Context, RowHandle ComponentRow, const FTypedElementUObjectColumn& ObjectColumn)
-						{
-							if (UActorComponent* ActorComponent = Cast<UActorComponent>(ObjectColumn.Object.Get()))
-							{
-								// Confusingly, an ActorComponent could indicate that it has an owner while its
-								// owner does not consider it owned.
-								// Defer to the owner's source of truth for this
-								const AActor* OwnerActor = ActorComponent->GetOwner();
-								RowHandle ActorRow;
-								if (OwnerActor && OwnerActor->GetComponents().Contains(ActorComponent))
-								{
-									ActorRow = Context.LookupMappedRow(ICompatibilityProvider::ObjectMappingDomain, FMapKeyView(static_cast<const UObject*>(OwnerActor)));
-								}
-								else
-								{
-									ActorRow = InvalidRowHandle;
-								}
-								
-								const bool bNotAssociatedWithActor = OwnerActor == nullptr && ActorRow == InvalidRowHandle;
-								const bool bActorRowAssigned = Context.IsRowAssigned(ActorRow);
-
-								RowHandle ParentRowOfComponent = Context.GetParentRow(ComponentRow);
-								
-								if (bNotAssociatedWithActor || (bActorRowAssigned && ParentRowOfComponent != ActorRow))
-								{
-									Context.SetParentRow(ComponentRow, ActorRow);
-								}
-							}
-						})
-						.AccessesHierarchy(TEXT("ActorComponent"))
-						.Where()
-							.All<FActorComponentTypeTag, FTypedElementSyncFromWorldTag>()
-						.Compile());
-
-				// 
-				CoreProvider->RegisterQuery(
-					Select(
-						TEXT("Sync actorcomponent label to column"),
-						FProcessor(EQueryTickPhase::PrePhysics, CoreProvider->GetQueryTickGroupName(EQueryTickGroups::SyncExternalToDataStorage)),
-						[](const FTypedElementUObjectColumn& Actor, FTypedElementLabelColumn& Label, FTypedElementLabelHashColumn& LabelHash)
-						{
-							if (const UActorComponent* ComponentInstance = Cast<UActorComponent>(Actor.Object); ComponentInstance != nullptr)
-							{
-								const FString Name = ComponentInstance->GetName();
-								uint64 NameHash = CityHash64(reinterpret_cast<const char*>(*Name), Name.Len() * sizeof(**Name));
-								if (LabelHash.LabelHash != NameHash)
-								{
-									Label.Label = Name;
-									LabelHash.LabelHash = NameHash;
-								}
-							}
-						}
-					)
-					.Where()
-						.All<FActorComponentTypeTag, FTypedElementSyncFromWorldTag>()
-					.Compile());
-			});
-		}
-
-		if (bEnableDemoTab)
-		{
-			UE::Editor::DataStorage::ActorCompatibility::RegisterActionComponentDebugHierarchyWidget();
-		}
 	}
 
 	void FTedsActorCompatibilityModule::ShutdownModule()
diff --git a/Engine/Plugins/Experimental/EditorDataStorageFeatures/Source/TedsActorCompatibility/Public/Columns/TedsActorComponentCompatibilityColumns.h b/Engine/Plugins/Experimental/EditorDataStorageFeatures/Source/TedsActorCompatibility/Public/Columns/TedsActorComponentCompatibilityColumns.h
deleted file mode 100644
diff --git a/Engine/Plugins/Experimental/EditorDataStorageFeatures/Source/TedsActorCompatibility/TedsActorCompatibility.Build.cs b/Engine/Plugins/Experimental/EditorDataStorageFeatures/Source/TedsActorCompatibility/TedsActorCompatibility.Build.cs
--- a/Engine/Plugins/Experimental/EditorDataStorageFeatures/Source/TedsActorCompatibility/TedsActorCompatibility.Build.cs
+++ b/Engine/Plugins/Experimental/EditorDataStorageFeatures/Source/TedsActorCompatibility/TedsActorCompatibility.Build.cs
@@ -24,12 +24,7 @@ public class TedsActorCompatibility : ModuleRules
 					"CoreUObject",
 					"Engine",
 					"TypedElementFramework",
-					"UnrealEd",
-					"TedsQueryStack",
-					"WorkspaceMenuStructure",
-					"TedsTableViewer",
-					"SlateCore",
-					"Slate",
+					"UnrealEd"
 				});
 
 			DynamicallyLoadedModuleNames.AddRange(new string[] { });
diff --git a/Engine/Plugins/Experimental/EditorDataStorageFeatures/Source/TedsQueryStack/Private/TedsHierarchyNode.cpp b/Engine/Plugins/Experimental/EditorDataStorageFeatures/Source/TedsQueryStack/Private/TedsHierarchyNode.cpp
--- a/Engine/Plugins/Experimental/EditorDataStorageFeatures/Source/TedsQueryStack/Private/TedsHierarchyNode.cpp
+++ b/Engine/Plugins/Experimental/EditorDataStorageFeatures/Source/TedsQueryStack/Private/TedsHierarchyNode.cpp
@@ -10,7 +10,7 @@ namespace UE::Editor::DataStorage
 	FHierarchyRowNode::FHierarchyRowNode(
 		ICoreProvider& Storage,
 		FHierarchyHandle InHierarchyHandle,
-		TSharedPtr<QueryStack::IRowNode>& InTopLevelRows,
+		const TSharedPtr<QueryStack::IRowNode>& InTopLevelRows,
 		ESyncFlags InSyncFlags)
 		: TopLevelRows(InTopLevelRows)
 		, Storage(Storage)
@@ -26,47 +26,34 @@ namespace UE::Editor::DataStorage
 
 	void FHierarchyRowNode::Update()
 	{
-		TopLevelRows->Update();
-		switch (SyncFlags)
+		
+		if (TopLevelRows->GetRevision() != ParentRevision)
 		{
-			case ESyncFlags::ParentRevisionDifferent:
-			{
-				if (TopLevelRows->GetRevision() != ParentRevision)
-				{
-					Rows.Empty();
-					RefreshInternal(Storage, Rows);
-					++Revision;
-					ParentRevision = TopLevelRows->GetRevision();
-				}
-				break;
-			}
-			case ESyncFlags::IncrementWhenDifferent:
-			{
-				FRowHandleArray& TempArray = GetTempRowHandleArray();
-				RefreshInternal(Storage, TempArray);
-				TempArray.Sort();
-				Rows.Sort();
-				if (!TempArray.IsSame(Rows))
-				{
-					Rows = MoveTemp(TempArray);
-					++Revision;
-					ParentRevision = TopLevelRows->GetRevision();
-				}
-				break;
-			}
-			case ESyncFlags::Always:
+			Rows.Empty();
+			RefreshInternal(Storage, Rows);
+			++Revision;
+			ParentRevision = TopLevelRows->GetRevision();
+		}
+		if (EnumHasAllFlags(SyncFlags, ESyncFlags::IncrementWhenDifferent))
+		{
+			FRowHandleArray& TempArray = GetTempRowHandleArray();
+			RefreshInternal(Storage, TempArray);
+			TempArray.Sort();
+			Rows.Sort();
+			if (!TempArray.IsSame(Rows))
 			{
-				Rows.Empty();
-				RefreshInternal(Storage, Rows);
+				Rows = MoveTemp(TempArray);
 				++Revision;
 				ParentRevision = TopLevelRows->GetRevision();
-				break;
-			}
-			default:
-			{
-				// Do nothing
 			}
 		}
+		else
+		{
+			Rows.Empty();
+			RefreshInternal(Storage, Rows);
+			++Revision;
+			ParentRevision = TopLevelRows->GetRevision();
+		}
 	}
 
 	FRowHandleArrayView FHierarchyRowNode::GetRows() const
diff --git a/Engine/Plugins/Experimental/EditorDataStorageFeatures/Source/TedsQueryStack/Public/TedsHierarchyNode.h b/Engine/Plugins/Experimental/EditorDataStorageFeatures/Source/TedsQueryStack/Public/TedsHierarchyNode.h
--- a/Engine/Plugins/Experimental/EditorDataStorageFeatures/Source/TedsQueryStack/Public/TedsHierarchyNode.h
+++ b/Engine/Plugins/Experimental/EditorDataStorageFeatures/Source/TedsQueryStack/Public/TedsHierarchyNode.h
@@ -13,25 +13,25 @@ namespace UE::Editor::DataStorage
 	public:
 		enum class ESyncFlags
 		{
+			None = 0,
 			/** 
-			 * Increment the revision only when the parent's revision changed
+			 * Update the row list whenever "Update" is called. This can take several milliseconds for large lists. It's
+			 * recommended to use a monitoring node to detect changes for large lists.
 			 */
-			ParentRevisionDifferent = 0,
+			RefreshOnUpdate = 1 << 0,
+			/** Update the row list whenever the parent changes. */
+			RefreshOnParentChange = 1 << 1,
 			/** 
 			 * Compares the previous row list with the current and only updates the revision if there are differences.
 			 */
-			IncrementWhenDifferent = 1,
-			/**
-			 * Update the row list whenever "Update" is called. 
-			 */
-			Always
+			IncrementWhenDifferent = 1 << 2
 		};
 		using EmitRowFn = TFunctionRef<void(TArrayView<const RowHandle>)>;
 		
 		TEDSQUERYSTACK_API FHierarchyRowNode(
 			ICoreProvider& Storage,
 			FHierarchyHandle InHierarchyHandle,
-			TSharedPtr<QueryStack::IRowNode>& InTopLevelRows,
+			const TSharedPtr<QueryStack::IRowNode>& InTopLevelRows,
 			ESyncFlags SyncFlags);
 		
 		TEDSQUERYSTACK_API virtual RevisionId GetRevision() const override;
@@ -45,7 +45,7 @@ namespace UE::Editor::DataStorage
 		
 		FRowHandleArray Rows;
 		FRowHandleArray NewRows;
-		TSharedPtr<QueryStack::IRowNode> TopLevelRows;
+		const TSharedPtr<QueryStack::IRowNode> TopLevelRows;
 		RevisionId ParentRevision = 0;
 		
 		ICoreProvider& Storage;
diff --git a/Engine/Plugins/Experimental/EditorDataStorageFeatures/Source/TedsQueryStack/Public/TedsRowQueryResultsNode.h b/Engine/Plugins/Experimental/EditorDataStorageFeatures/Source/TedsQueryStack/Public/TedsRowQueryResultsNode.h
--- a/Engine/Plugins/Experimental/EditorDataStorageFeatures/Source/TedsQueryStack/Public/TedsRowQueryResultsNode.h
+++ b/Engine/Plugins/Experimental/EditorDataStorageFeatures/Source/TedsQueryStack/Public/TedsRowQueryResultsNode.h
@@ -50,8 +50,8 @@ namespace UE::Editor::DataStorage::QueryStack
 		TEDSQUERYSTACK_API virtual FRowHandleArrayView GetRows() const override;
 		TEDSQUERYSTACK_API virtual FRowHandleArray& GetMutableRows() override;
 
-	protected:
-		TEDSQUERYSTACK_API virtual void RefreshInternal(FRowHandleArray& TargetRows);
+	private:
+		void RefreshInternal(FRowHandleArray& TargetRows);
 
 		FRowHandleArray Rows;
 		TSharedPtr<IQueryNode> QueryNode;
diff --git a/Engine/Source/Runtime/Engine/Private/Actor.cpp b/Engine/Source/Runtime/Engine/Private/Actor.cpp
--- a/Engine/Source/Runtime/Engine/Private/Actor.cpp
+++ b/Engine/Source/Runtime/Engine/Private/Actor.cpp
@@ -72,15 +72,11 @@
 #endif
 
 #if WITH_EDITOR
-#include "Features/IModularFeatures.h"
 #include "FoliageHelper.h"
 #include "LevelInstance/LevelInstanceSubsystem.h"
 #include "LevelInstance/LevelInstanceInterface.h"
 #endif
 
-#include "Elements/Columns/TypedElementMiscColumns.h"
-#include "DataStorage/Features.h"
-#include "Elements/Interfaces/TypedElementDataStorageCompatibilityInterface.h"
 #include "WorldPartition/WorldPartitionLog.h"
 #include "WorldPartition/DataLayer/WorldDataLayers.h"
 #include "WorldPartition/WorldPartitionActorDescUtils.h"
@@ -518,7 +514,7 @@ void AActor::ResetOwnedComponents()
 	OwnedComponents.Reset();
 	ReplicatedComponents.Reset();
 	ReplicatedComponentsInfo.Reset();
-	
+
 	ForEachObjectWithOuter(this, [this](UObject* Child)
 	{
 		UActorComponent* Component = Cast<UActorComponent>(Child);
@@ -534,19 +530,6 @@ void AActor::ResetOwnedComponents()
 			}
 		}
 	}, true, RF_NoFlags, EInternalObjectFlags::Garbage);
-
-#if WITH_EDITOR
-	{
-		using namespace UE::Editor::DataStorage;
-		if (ICompatibilityProvider* Storage = UE::Editor::DataStorage::GetMutableDataStorageFeature<ICompatibilityProvider>(CompatibilityFeatureName))
-		{
-			for (UActorComponent* Component : OwnedComponents)
-			{
-				Storage->AddCompatibleObject(const_cast<UActorComponent*>(Component));
-			}
-		}
-	}
-#endif
 }
 
 void AActor::PostInitProperties()
@@ -3860,16 +3843,6 @@ void AActor::AddOwnedComponent(UActorComponent* Component)
 		{
 			InstanceComponents.Add(Component);
 		}
-#if WITH_EDITOR
-		{
-			using namespace UE::Editor::DataStorage;
-			if (ICoreProvider* Storage = UE::Editor::DataStorage::GetMutableDataStorageFeature<ICoreProvider>(StorageFeatureName))
-			{
-				RowHandle ComponentRow = Storage->LookupMappedRow(ICompatibilityProvider::ObjectMappingDomain, FMapKeyView(Component));
-				Storage->AddColumns<FTypedElementSyncFromWorldTag>(ComponentRow);
-			}
-		}
-#endif
 	}
 }
 
@@ -3893,18 +3866,6 @@ void AActor::RemoveOwnedComponent(UActorComponent* Component)
 		{
 			InstanceComponents.RemoveSingleSwap(Component);
 		}
-
-		
-#if WITH_EDITOR
-		{
-			using namespace UE::Editor::DataStorage;
-			if (ICoreProvider* Storage = UE::Editor::DataStorage::GetMutableDataStorageFeature<ICoreProvider>(StorageFeatureName))
-			{
-				RowHandle ComponentRow = Storage->LookupMappedRow(ICompatibilityProvider::ObjectMappingDomain, FMapKeyView(Component));
-				Storage->AddColumns<FTypedElementSyncFromWorldTag>(ComponentRow);
-			}
-		}
-#endif
 	}
 }
 
diff --git a/Engine/Source/Runtime/Engine/Private/Components/ActorComponent.cpp b/Engine/Source/Runtime/Engine/Private/Components/ActorComponent.cpp
--- a/Engine/Source/Runtime/Engine/Private/Components/ActorComponent.cpp
+++ b/Engine/Source/Runtime/Engine/Private/Components/ActorComponent.cpp
@@ -48,9 +48,6 @@
 #include "Kismet2/ComponentEditorUtils.h"
 #include "ObjectCacheEventSink.h"
 #include "StaticMeshCompiler.h"
-#include "Elements/Columns/TypedElementMiscColumns.h"
-#include "DataStorage/Features.h"
-#include "Elements/Interfaces/TypedElementDataStorageCompatibilityInterface.h"
 #endif
 
 #if UE_WITH_IRIS
@@ -596,11 +593,6 @@ void UActorComponent::PostInitProperties()
 
 #if WITH_EDITOR
 	UEngineElementsLibrary::CreateEditorComponentElement(this);
-	using namespace UE::Editor::DataStorage;
-	if (ICompatibilityProvider* Storage = UE::Editor::DataStorage::GetMutableDataStorageFeature<ICompatibilityProvider>(CompatibilityFeatureName))
-	{
-		Storage->AddCompatibleObject(this);
-	}
 #endif	// WITH_EDITOR
 
 	// Instance components will be added during the owner's initialization
@@ -896,15 +888,6 @@ void UActorComponent::PostRename(UObject* OldOuter, const FName OldName)
 				OwnerPrivate->AddOwnedComponent(this);
 			}
 
-#if WITH_EDITOR
-			using namespace UE::Editor::DataStorage;
-			if (ICoreProvider* Storage = UE::Editor::DataStorage::GetMutableDataStorageFeature<ICoreProvider>(StorageFeatureName))
-			{
-				RowHandle ComponentRow = Storage->LookupMappedRow(ICompatibilityProvider::ObjectMappingDomain, FMapKeyView(this));
-				Storage->AddColumns<FTypedElementSyncFromWorldTag>(ComponentRow);
-			}
-#endif
-
 			TArray<UObject*> Children;
 			GetObjectsWithOuter(this, Children, /*bIncludeNestedObjects=*/false);
 
@@ -1130,11 +1113,6 @@ void UActorComponent::BeginDestroy()
 	}
 
 #if WITH_EDITOR
-	using namespace UE::Editor::DataStorage;
-	if (ICompatibilityProvider* Storage = UE::Editor::DataStorage::GetMutableDataStorageFeature<ICompatibilityProvider>(CompatibilityFeatureName))
-	{
-		Storage->RemoveCompatibleObject(this);
-	}
 	UEngineElementsLibrary::DestroyEditorComponentElement(this);
 #endif	// WITH_EDITOR
 
@@ -2121,16 +2099,6 @@ void UActorComponent::OnComponentDestroyed(bool bDestroyingHierarchy)
 {
 	// @TODO: Would be nice to ensure(bHasBeenCreated), but there are still many places where components are created without calling OnComponentCreated
 	bHasBeenCreated = false;
-
-#if WITH_EDITOR
-	{
-		using namespace UE::Editor::DataStorage;
-		if (ICompatibilityProvider* Storage = UE::Editor::DataStorage::GetMutableDataStorageFeature<ICompatibilityProvider>(CompatibilityFeatureName))
-		{
-			Storage->RemoveCompatibleObject(this);
-		}
-	}
-#endif
 }
 
 void UActorComponent::K2_DestroyComponent(UObject* Object)
diff --git a/Engine/Source/Runtime/Engine/Private/Elements/Framework/EngineElementsLibrary.cpp b/Engine/Source/Runtime/Engine/Private/Elements/Framework/EngineElementsLibrary.cpp
--- a/Engine/Source/Runtime/Engine/Private/Elements/Framework/EngineElementsLibrary.cpp
+++ b/Engine/Source/Runtime/Engine/Private/Elements/Framework/EngineElementsLibrary.cpp
@@ -383,15 +383,9 @@ void UEngineElementsLibrary::RegisterActorElement(const AActor* InActor)
 {
 	using namespace UE::Editor::DataStorage;
 	ICompatibilityProvider* Storage = GetMutableDataStorageFeature<ICompatibilityProvider>(CompatibilityFeatureName);
-	RowHandle AddedObjectRow = InvalidRowHandle;
 	if (Storage)
 	{
-		AddedObjectRow = Storage->AddCompatibleObject(const_cast<AActor*>(InActor));
-		constexpr bool bIncludeChildActors = false;
-		InActor->ForEachComponent(bIncludeChildActors, [&Storage](UActorComponent* Component)
-		{
-			Storage->AddCompatibleObject(const_cast<UActorComponent*>(Component));
-		});
+		Storage->AddCompatibleObject(const_cast<AActor*>(InActor));
 	}
 }
 
diff --git a/Engine/Source/Runtime/Engine/Private/Level.cpp b/Engine/Source/Runtime/Engine/Private/Level.cpp
--- a/Engine/Source/Runtime/Engine/Private/Level.cpp
+++ b/Engine/Source/Runtime/Engine/Private/Level.cpp
@@ -1094,11 +1094,6 @@ void ULevel::AddLoadedActors(const TArray<AActor*>& ActorList, const FTransform*
 			if (TedsCompat)
 			{
 				TedsCompat->AddCompatibleObject(Actor);
-				constexpr bool bIncludeChildActors = false;
-				Actor->ForEachComponent(bIncludeChildActors, [&TedsCompat](UActorComponent* Component)
-				{
-					TedsCompat->AddCompatibleObject(Component);
-				});
 			}
 		}
 
@@ -1212,12 +1207,6 @@ void ULevel::RemoveLoadedActors(const TArray<AActor*>& ActorList, const FTransfo
 			// but for our purposes it is considered dead already
 			if (TedsCompat)
 			{
-				constexpr bool bIncludeChildActors = false;
-				Actor->ForEachComponent(bIncludeChildActors, [&TedsCompat](UActorComponent* Component)
-				{
-					TedsCompat->RemoveCompatibleObject(Component);
-				});
-				
 				TedsCompat->RemoveCompatibleObject(Actor);
 			}
 
