Move TypeInfoModule.h to Public folder

diff --git a/Engine/Plugins/Experimental/EditorDataStorageFeatures/Source/TedsTypeInfo/Private/TedsTypeInfoModule.h b/Engine/Plugins/Experimental/EditorDataStorageFeatures/Source/TedsTypeInfo/Private/TedsTypeInfoModule.h
deleted file mode 100644
diff --git a/Engine/Plugins/Experimental/EditorDataStorageFeatures/Source/TedsTypeInfo/Public/TedsTypeInfoModule.h b/Engine/Plugins/Experimental/EditorDataStorageFeatures/Source/TedsTypeInfo/Public/TedsTypeInfoModule.h
new file mode 100644
--- /dev/null
+++ b/Engine/Plugins/Experimental/EditorDataStorageFeatures/Source/TedsTypeInfo/Public/TedsTypeInfoModule.h
@@ -0,0 +1,29 @@
+// Copyright Epic Games, Inc. All Rights Reserved.
+
+#pragma once
+
+#include "Modules/ModuleInterface.h"
+
+namespace UE::Editor::DataStorage::TypeInfo
+{
+	class FTedsTypeInfoModule : public IModuleInterface
+	{
+	public:
+		virtual void StartupModule() override;
+		virtual void ShutdownModule() override;
+
+		TEDSTYPEINFO_API static FTedsTypeInfoModule* Get();
+		TEDSTYPEINFO_API static FTedsTypeInfoModule& GetChecked();
+
+		TEDSTYPEINFO_API void EnableTedsTypeInfoIntegration();
+		TEDSTYPEINFO_API void DisableTedsTypeInfoIntegration();
+		TEDSTYPEINFO_API bool IsTedsTypeInfoIntegrationEnabled();
+
+		TEDSTYPEINFO_API void RefreshTypeInfo();
+
+	private:
+		void FlushTypeInfo();
+
+		bool bDataStorageFeaturesEnabled = false;
+	};
+}
\ No newline at end of file
