[TEDS] Create a base for an STedsSearchBox

diff --git a/Engine/Plugins/Experimental/EditorDataStorageFeatures/Source/TedsEverythingPicker/Private/Context/TedsPickerContextUtil.cpp b/Engine/Plugins/Experimental/EditorDataStorageFeatures/Source/TedsEverythingPicker/Private/Context/TedsPickerContextUtil.cpp
--- a/Engine/Plugins/Experimental/EditorDataStorageFeatures/Source/TedsEverythingPicker/Private/Context/TedsPickerContextUtil.cpp
+++ b/Engine/Plugins/Experimental/EditorDataStorageFeatures/Source/TedsEverythingPicker/Private/Context/TedsPickerContextUtil.cpp
@@ -12,12 +12,86 @@
 
 #include "Widgets/STedsHierarchyViewer.h"
 #include "Widgets/SEverythingPicker.h"
+#include "Widgets/STedsSearchBox.h"
 
 #include "TedsQueryNode.h"
 #include "TedsRowQueryResultsNode.h"
 
 namespace UE::Editor::DataStorage::Picker
 {
+	// A custom query stack node to get top level rows of the given rows to display in the hierarchy viewer
+	class FGatherTopLevelRowsNode : public UE::Editor::DataStorage::QueryStack::IRowNode
+	{
+	public:
+		
+		FGatherTopLevelRowsNode(const UE::Editor::DataStorage::ICoreProvider* InStorage, TSharedPtr<IRowNode> InParent)
+			: Storage(InStorage)
+			, Parent(MoveTemp(InParent))
+		{
+			UpdateRows();
+		}
+		
+		virtual ~FGatherTopLevelRowsNode() override = default;
+		
+		/** Retrieve access to the rows used by this node. */
+		virtual UE::Editor::DataStorage::FRowHandleArrayView GetRows() const override
+		{
+			return Rows.GetRows();
+		}
+		
+		virtual UE::Editor::DataStorage::FRowHandleArray& GetMutableRows() override
+		{
+			return Rows;
+		}
+		
+		virtual RevisionId GetRevision() const override
+		{
+			return Revision;
+		}
+		
+		virtual void Update() override
+		{
+			Parent->Update();
+			if (Parent->GetRevision() != CachedParentRevision)
+			{
+				UpdateRows();
+				CachedParentRevision = Parent->GetRevision();
+				Revision++;
+			}
+		}
+
+		void UpdateRows()
+		{
+			Rows.Empty();
+
+			UE::Editor::DataStorage::FRowHandleArray RowsToCheck(Parent->GetRows());
+			RowsToCheck.Sort();
+			for (UE::Editor::DataStorage::RowHandle Row : Parent->GetRows())
+			{
+				// If the row doesn't have a parent, it is a top level node and can be added
+				if (const FTableRowParentColumn* ParentColumn = Storage->GetColumn<FTableRowParentColumn>(Row))
+				{
+					// If the filtered list contains the parent of the row, it has a visible parent and thus is not a top-level node
+					if(!RowsToCheck.Contains(ParentColumn->Parent))
+					{
+						Rows.Add(Row);
+					}
+				}
+				else
+				{
+					Rows.Add(Row);
+				}
+			}
+		}
+
+	protected:
+		RevisionId Revision = 0;
+		const UE::Editor::DataStorage::ICoreProvider* Storage;
+		TSharedPtr<IRowNode> Parent;
+		UE::Editor::DataStorage::FRowHandleArray Rows;
+		RevisionId CachedParentRevision = TNumericLimits<RevisionId>::Max();
+	};
+
 	void SObjectReferenceContextView::Construct(const FArguments& InArgs)
 	{
 		using namespace UE::Editor::DataStorage::Queries;
@@ -33,13 +107,25 @@ namespace UE::Editor::DataStorage::Picker
 			MakeShared<QueryStack::FRowQueryResultsNode>(*Storage,
 				ReferenceQueryNode, QueryStack::FRowQueryResultsNode::ESyncFlags::RefreshOnUpdate);
 
-		// Node that filters out the ReferenceQueryResultsNode to generate all the top level nodes
-		TSharedRef<QueryStack::IRowNode> TopLevelNode =
-			MakeShared<QueryStack::FRowFilterNode>(Storage, ReferenceQueryResultsNode.ToSharedPtr(),
-				[](TQueryContext<SingleRowInfo>& Context)
-				{
-					return !Context.CurrentRowHasColumns<FTableRowParentColumn>();
-				});
+		TSharedPtr<QueryStack::IRowNode> SearchNode;
+		TSharedPtr<SBox> SearchBoxContainer = SNew(SBox);
+		if(InArgs._SearchingEnabled)
+		{
+			SearchBoxContainer->SetContent(SNew(STedsSearchBox)
+				.OutSearchNode(&SearchNode)
+				.QueryResultsNode(ReferenceQueryResultsNode.ToSharedPtr()));			
+		}
+
+		// If there is no SearchNode available, either disabled or there is an internal error, assign it to be the 
+		// ReferenceQueryResultsNode so that won't be null
+		if(!SearchNode)
+		{
+			SearchNode = ReferenceQueryResultsNode.ToSharedPtr();
+		}
+
+		// Node that filters out the ReferenceQueryResultsNode or SearchNode (if enabled) to generate all the top level nodes
+		TSharedRef<FGatherTopLevelRowsNode> TopLevelNode =
+			MakeShared<FGatherTopLevelRowsNode>(Storage, SearchNode);
 
 		// Using the Teds Outliner Purpose for matching customization for now
 		IUiProvider::FPurposeID PurposeId = IUiProvider::FPurposeInfo("SceneOutliner", "RowLabel", NAME_None).GeneratePurposeID();
@@ -51,7 +137,7 @@ namespace UE::Editor::DataStorage::Picker
 		};
 
 		TSharedPtr<SHierarchyViewer> ReferenceView = SNew(SHierarchyViewer, FTableRowParentColumn::StaticStruct(), ParentExtractionFunction)
-			.AllNodeProvider(ReferenceQueryResultsNode)
+			.AllNodeProvider(SearchNode)
 			.TopLevelNodeProvider(TopLevelNode)
 			.Columns({ FTypedElementLabelColumn::StaticStruct(), FTypedElementClassTypeInfoColumn::StaticStruct() })
 			.CellWidgetPurpose(PurposeId)
@@ -60,7 +146,17 @@ namespace UE::Editor::DataStorage::Picker
 
 		ChildSlot
 			[
-				ReferenceView.ToSharedRef()
+				SNew(SVerticalBox)
+				+ SVerticalBox::Slot()
+				.AutoHeight()
+				.Padding(4.0f, 4.0f)
+				[
+					SearchBoxContainer.ToSharedRef()
+				]
+				+ SVerticalBox::Slot()
+				[
+					ReferenceView.ToSharedRef()
+				]
 			];
 	}
 }
\ No newline at end of file
diff --git a/Engine/Plugins/Experimental/EditorDataStorageFeatures/Source/TedsEverythingPicker/Public/Context/TedsPickerContextUtil.h b/Engine/Plugins/Experimental/EditorDataStorageFeatures/Source/TedsEverythingPicker/Public/Context/TedsPickerContextUtil.h
--- a/Engine/Plugins/Experimental/EditorDataStorageFeatures/Source/TedsEverythingPicker/Public/Context/TedsPickerContextUtil.h
+++ b/Engine/Plugins/Experimental/EditorDataStorageFeatures/Source/TedsEverythingPicker/Public/Context/TedsPickerContextUtil.h
@@ -17,9 +17,11 @@ namespace UE::Editor::DataStorage::Picker
 	{
 	public:
 		SLATE_BEGIN_ARGS(SObjectReferenceContextView)
+			: _SearchingEnabled(false)
 			{}
 			SLATE_ARGUMENT(FQueryDescription, Query)
 			SLATE_ARGUMENT(FOnSelectionChanged, OnSelectionChanged)
+			SLATE_ARGUMENT(bool, SearchingEnabled);
 		SLATE_END_ARGS()
 
 		UE_API void Construct(const FArguments& InArgs);
diff --git a/Engine/Plugins/Experimental/EditorDataStorageFeatures/Source/TedsTableViewer/Private/Widgets/STedsSearchBox.cpp b/Engine/Plugins/Experimental/EditorDataStorageFeatures/Source/TedsTableViewer/Private/Widgets/STedsSearchBox.cpp
new file mode 100644
--- /dev/null
+++ b/Engine/Plugins/Experimental/EditorDataStorageFeatures/Source/TedsTableViewer/Private/Widgets/STedsSearchBox.cpp
@@ -0,0 +1,78 @@
+// Copyright Epic Games, Inc. All Rights Reserved.
+
+#include "Widgets/STedsSearchBox.h"
+
+#include "DataStorage/Features.h"
+#include "Elements/Framework/TypedElementQueryBuilder.h"
+#include "Elements/Framework/TypedElementQueryContext.h"
+
+#include "Elements/Columns/TypedElementLabelColumns.h"
+#include "Elements/Columns/TypedElementTypeInfoColumns.h"
+#include "TedsQueryNode.h"
+#include "TedsRowFilterNode.h"
+#include "Widgets/Input/SSearchBox.h"
+#include "Widgets/Text/STextBlock.h"
+
+#define LOCTEXT_NAMESPACE "STedsSearchBox"
+
+namespace UE::Editor::DataStorage
+{
+	void STedsSearchBox::Construct(const FArguments& InArgs)
+	{
+		using namespace UE::Editor::DataStorage::Queries;
+
+		ICoreProvider* Storage = GetMutableDataStorageFeature<ICoreProvider>(UE::Editor::DataStorage::StorageFeatureName);
+
+		if(!ensureMsgf(Storage, TEXT("Cannot create a TEDS search widget before TEDS is initialized")))
+		{
+			ChildSlot
+			[
+				SNew(STextBlock)
+					.Text(LOCTEXT("TedsSearchErrorText", "No valid editor data storage available."))
+			];
+			return;
+		}
+
+		if(InArgs._OutSearchNode)
+		{
+			SearchNode = MakeShared<QueryStack::FRowFilterNode>(Storage, InArgs._QueryResultsNode,
+				[this](TQueryContext<SingleRowInfo>& Context, const FTypedElementLabelColumn& LabelColumn)
+				{
+					if(SearchText.IsEmpty())
+					{
+						return true;
+					}
+
+					TArray<FString> SearchTokens;
+					SearchText.ToString().ParseIntoArrayWS(SearchTokens);
+					const FString Label = LabelColumn.Label;
+
+					for (const FString& SearchToken : SearchTokens)
+					{
+						if (Label.IsEmpty() || !Label.Contains(SearchToken, ESearchCase::Type::IgnoreCase))
+						{
+							return false;
+						}
+					}
+					return true;
+				});
+				
+			*InArgs._OutSearchNode = SearchNode;
+		}
+		
+		ChildSlot
+		[
+			SNew(SSearchBox)
+				.OnTextChanged(this, &STedsSearchBox::OnTextChanged)
+		];
+	}
+
+	void STedsSearchBox::OnTextChanged(const FText& InSearchText)
+	{
+		SearchText = InSearchText;
+		SearchNode->ForceRefresh();
+	}
+	
+} // namespace UE::Editor::DataStorage
+
+#undef LOCTEXT_NAMESPACE //"STedsSearchBox"
\ No newline at end of file
diff --git a/Engine/Plugins/Experimental/EditorDataStorageFeatures/Source/TedsTableViewer/Public/Widgets/STedsSearchBox.h b/Engine/Plugins/Experimental/EditorDataStorageFeatures/Source/TedsTableViewer/Public/Widgets/STedsSearchBox.h
new file mode 100644
--- /dev/null
+++ b/Engine/Plugins/Experimental/EditorDataStorageFeatures/Source/TedsTableViewer/Public/Widgets/STedsSearchBox.h
@@ -0,0 +1,44 @@
+// Copyright Epic Games, Inc. All Rights Reserved.
+#pragma once
+
+#include "TedsQueryStackInterfaces.h"
+#include "TedsRowFilterNode.h"
+#include "TedsRowQueryResultsNode.h"
+#include "Widgets/SCompoundWidget.h"
+
+namespace UE::Editor::DataStorage
+{
+	namespace QueryStack
+	{
+		class FRowQueryResultsNode;
+	}
+
+	/*
+	 * A search bar widget that can be used to search across a TEDS table viewer widget using a QueryStack Search node.
+	 * Example future usage:
+	 * 
+	 *	SNew(STedsSearchBar)
+     *		.OutSearchNode(&OutSearchNode) // TSharedPtr<QueryStack::IRowNode> OutSearchNode
+     *		.QueryResultsNode(ReferenceQueryResultsNode)
+	 */
+	class STedsSearchBox : public SCompoundWidget
+	{
+	public:
+		
+		SLATE_BEGIN_ARGS(STedsSearchBox) {}
+			SLATE_ARGUMENT(TSharedPtr<QueryStack::IRowNode>*, OutSearchNode)
+			SLATE_ARGUMENT(TSharedPtr<QueryStack::FRowQueryResultsNode>, QueryResultsNode)
+		SLATE_END_ARGS()
+
+	public:
+		
+		TEDSTABLEVIEWER_API void Construct(const FArguments& InArgs);
+
+	private:
+		void OnTextChanged(const FText& InSearchText);
+
+		/** Internal SharedPtr of the OutSearchNode so that when text is updated, it can ForceRefresh the node */
+		TSharedPtr<QueryStack::FRowFilterNode> SearchNode;
+		FText SearchText;
+	};
+}
\ No newline at end of file
