Teds: Move actor hierarchies to the new interface and fix a bug where hierarchies were not updating properly on detachment

diff --git a/Engine/Plugins/Experimental/EditorDataStorageFeatures/Source/TedsActorCompatibility/Private/Compatibility/TedsActorParentQueries.cpp b/Engine/Plugins/Experimental/EditorDataStorageFeatures/Source/TedsActorCompatibility/Private/Compatibility/TedsActorParentQueries.cpp
--- a/Engine/Plugins/Experimental/EditorDataStorageFeatures/Source/TedsActorCompatibility/Private/Compatibility/TedsActorParentQueries.cpp
+++ b/Engine/Plugins/Experimental/EditorDataStorageFeatures/Source/TedsActorCompatibility/Private/Compatibility/TedsActorParentQueries.cpp
@@ -7,39 +7,57 @@
 #include "Elements/Columns/TypedElementHiearchyColumns.h"
 #include "Elements/Columns/TypedElementMiscColumns.h"
 #include "Elements/Framework/TypedElementQueryBuilder.h"
+#include "Engine/Engine.h"
 #include "GameFramework/Actor.h"
 
 #include UE_INLINE_GENERATED_CPP_BY_NAME(TedsActorParentQueries)
 
-namespace UE::Editor::DataStorage::Queries::Private
+namespace UE::Editor::DataStorage::ActorParentQueries::Private
 {
-	static bool bAddParentColumnToActors = false;
+	static FName ActorEditorHierarchyName(TEXT("ActorEditorHierarchy"));
+	static bool bAddParentColumnToActors = true;
 
-	// Cvar to allow the TEDS-Outliner to automatically take over the level editor's 4 Outliners instead of appearing as a separate tab
-	static FAutoConsoleVariableRef CVarUseTEDSOutliner(
+	// Cvar to add parenting info for actors to TEDS
+	static FAutoConsoleVariableRef CvarAddParentColumnToActors(
 		TEXT("TEDS.AddParentColumnToActors"),
 		bAddParentColumnToActors,
 		TEXT("Mirror parent information for actors to TEDS (only works when set on startup)"));
 };
 
-void UActorParentDataStorageFactory::RegisterQueries(UE::Editor::DataStorage::ICoreProvider& DataStorage)
+void UActorParentDataStorageFactory::RegisterQueries(UE::Editor::DataStorage::ICoreProvider& InDataStorage)
 {
-	if (UE::Editor::DataStorage::Queries::Private::bAddParentColumnToActors)
+	if (UE::Editor::DataStorage::ActorParentQueries::Private::bAddParentColumnToActors)
 	{
-		RegisterAddParentColumn(DataStorage);
-		RegisterUpdateOrRemoveParentColumn(DataStorage);
+		RegisterAddParentColumn(InDataStorage);
+		RegisterUpdateOrRemoveParentColumn(InDataStorage);
 	}
 }
 
-void UActorParentDataStorageFactory::RegisterAddParentColumn(UE::Editor::DataStorage::ICoreProvider& DataStorage) const
+void UActorParentDataStorageFactory::PreRegister(UE::Editor::DataStorage::ICoreProvider& InDataStorage)
+{
+	UE::Editor::DataStorage::FHierarchyRegistrationParams ActorHierarchyParams{.Name = UE::Editor::DataStorage::ActorParentQueries::Private::ActorEditorHierarchyName};
+	EditorActorHierarchyHandle = InDataStorage.RegisterHierarchy(ActorHierarchyParams);
+
+	GEngine->OnLevelActorDetached().AddUObject(this, &UActorParentDataStorageFactory::OnLevelActorDetached);
+
+	// Store a pointer to data storage so we don't have to grab it from the global instance of the registry every time
+	DataStorage = &InDataStorage;
+}
+
+void UActorParentDataStorageFactory::PreShutdown(UE::Editor::DataStorage::ICoreProvider& InDataStorage)
+{
+	GEngine->OnLevelActorDetached().RemoveAll(this);
+}
+
+void UActorParentDataStorageFactory::RegisterAddParentColumn(UE::Editor::DataStorage::ICoreProvider& InDataStorage) const
 {
 	using namespace UE::Editor::DataStorage;
 	using namespace UE::Editor::DataStorage::Queries;
 
-	DataStorage.RegisterQuery(
+	InDataStorage.RegisterQuery(
 		Select(
 			TEXT("Add parent column to actor"),
-			FProcessor(EQueryTickPhase::PrePhysics, DataStorage.GetQueryTickGroupName(EQueryTickGroups::SyncExternalToDataStorage))
+			FProcessor(EQueryTickPhase::PrePhysics, InDataStorage.GetQueryTickGroupName(EQueryTickGroups::SyncExternalToDataStorage))
 			.SetExecutionMode(EExecutionMode::GameThread),
 			[](IQueryContext& Context, RowHandle Row, const FTypedElementUObjectColumn& Actor)
 			{
@@ -51,57 +69,72 @@ void UActorParentDataStorageFactory::RegisterAddParentColumn(UE::Editor::DataSto
 						RowHandle ParentRow = Context.LookupMappedRow(ICompatibilityProvider::ObjectMappingDomain, IdKey);
 						if (Context.IsRowAvailable(ParentRow))
 						{
-							Context.AddColumn(Row, FTableRowParentColumn{ .Parent = ParentRow });
+							Context.SetParentRow(Row, ParentRow);
 						}
 						else
 						{
-							Context.AddColumn(Row, FUnresolvedTableRowParentColumn{ .ParentIdKey = FMapKey(Parent) });
+							Context.SetUnresolvedParent(Row, FMapKey(Parent), ICompatibilityProvider::ObjectMappingDomain);
 						}
 					}
 				}
 			})
+		.AccessesHierarchy(ActorParentQueries::Private::ActorEditorHierarchyName)
 		.Where()
-		.All<FTypedElementSyncFromWorldTag, FTypedElementActorTag>()
-		.None<FTableRowParentColumn, FUnresolvedTableRowParentColumn>()
+			.All<FTypedElementSyncFromWorldTag, FTypedElementActorTag>()
+			.None(InDataStorage.GetChildTagType(EditorActorHierarchyHandle))
 		.Compile());
 }
 
-void UActorParentDataStorageFactory::RegisterUpdateOrRemoveParentColumn(UE::Editor::DataStorage::ICoreProvider& DataStorage) const
+void UActorParentDataStorageFactory::RegisterUpdateOrRemoveParentColumn(UE::Editor::DataStorage::ICoreProvider& InDataStorage) const
 {
+	using namespace UE::Editor::DataStorage;
 	using namespace UE::Editor::DataStorage::Queries;
 
-	DataStorage.RegisterQuery(
+	InDataStorage.RegisterQuery(
 		Select(
 			TEXT("Sync actor's parent to column"),
-			FProcessor(EQueryTickPhase::PrePhysics, DataStorage.GetQueryTickGroupName(EQueryTickGroups::SyncExternalToDataStorage))
+			FProcessor(EQueryTickPhase::PrePhysics, InDataStorage.GetQueryTickGroupName(EQueryTickGroups::SyncExternalToDataStorage))
 			.SetExecutionMode(EExecutionMode::GameThread),
-			[](IQueryContext& Context, RowHandle Row, const FTypedElementUObjectColumn& Actor, FTableRowParentColumn& Parent)
+			[](IQueryContext& Context, RowHandle TargetRow, const FTypedElementUObjectColumn& Actor)
 			{
 				if (const AActor* ActorInstance = Cast<AActor>(Actor.Object))
 				{
 					if (AActor* ParentActor = ActorInstance->GetAttachParentActor())
 					{
 						FMapKeyView IdKey = FMapKeyView(ParentActor);
-						RowHandle ParentRow = Context.LookupMappedRow(ICompatibilityProvider::ObjectMappingDomain, IdKey);
+						RowHandle NewParentRow = Context.LookupMappedRow(ICompatibilityProvider::ObjectMappingDomain, IdKey);
+						RowHandle CurrentParentRow = Context.GetParentRow(TargetRow);
 
-						if (Parent.Parent != ParentRow)
+						if (CurrentParentRow != NewParentRow)
 						{
-							if (Context.IsRowAvailable(ParentRow))
+							if (Context.IsRowAvailable(NewParentRow))
 							{
-								Parent.Parent = ParentRow;
+								Context.SetParentRow(TargetRow, NewParentRow);
 							}
 							else
 							{
-								Context.RemoveColumns<FTableRowParentColumn>(Row);
-								Context.AddColumn(Row, FUnresolvedTableRowParentColumn{ .ParentIdKey = FMapKey(ParentActor) });
+								Context.SetUnresolvedParent(TargetRow, FMapKey(ParentActor), ICompatibilityProvider::ObjectMappingDomain);
 							}
 						}
 						return;
 					}
 				}
-				Context.RemoveColumns<FTableRowParentColumn>(Row);
+
+				// If we have reached here, the actor does not have a parent
+				Context.SetParentRow(TargetRow, InvalidRowHandle);
 			})
+		.AccessesHierarchy(ActorParentQueries::Private::ActorEditorHierarchyName)
 		.Where()
-		.All<FTypedElementActorTag, FTypedElementSyncFromWorldTag>()
+			.All<FTypedElementActorTag, FTypedElementSyncFromWorldTag>()
 		.Compile());
 }
+
+void UActorParentDataStorageFactory::OnLevelActorDetached(AActor* Actor, const AActor* OldParent) const
+{
+	// Detaching an actor from its parent only calls Modify() on the parent and not on the child, so we explicitly add the sync tag to the child so the
+	// hierarchy in TEDS is updated
+	using namespace UE::Editor::DataStorage;
+	
+	RowHandle Row = DataStorage->LookupMappedRow(ICompatibilityProvider::ObjectMappingDomain, FMapKey(Actor));
+	DataStorage->AddColumn<FTypedElementSyncFromWorldTag>(Row);
+}
diff --git a/Engine/Plugins/Experimental/EditorDataStorageFeatures/Source/TedsActorCompatibility/Private/Compatibility/TedsActorParentQueries.h b/Engine/Plugins/Experimental/EditorDataStorageFeatures/Source/TedsActorCompatibility/Private/Compatibility/TedsActorParentQueries.h
--- a/Engine/Plugins/Experimental/EditorDataStorageFeatures/Source/TedsActorCompatibility/Private/Compatibility/TedsActorParentQueries.h
+++ b/Engine/Plugins/Experimental/EditorDataStorageFeatures/Source/TedsActorCompatibility/Private/Compatibility/TedsActorParentQueries.h
@@ -2,6 +2,7 @@
 
 #pragma once
 
+#include "DataStorage/Handles.h"
 #include "Elements/Interfaces/TypedElementDataStorageFactory.h"
 #include "UObject/ObjectMacros.h"
 
@@ -15,7 +16,9 @@ class UActorParentDataStorageFactory : public UEditorDataStorageFactory
 public:
 	~UActorParentDataStorageFactory() override = default;
 
-	void RegisterQueries(UE::Editor::DataStorage::ICoreProvider& DataStorage) override;
+	virtual void RegisterQueries(UE::Editor::DataStorage::ICoreProvider& DataStorage) override;
+	virtual void PreRegister(UE::Editor::DataStorage::ICoreProvider& DataStorage) override;
+	virtual void PreShutdown(UE::Editor::DataStorage::ICoreProvider& DataStorage) override;
 
 private:
 	/**
@@ -28,4 +31,13 @@ private:
 	 * with the actor anymore.
 	 */
 	void RegisterUpdateOrRemoveParentColumn(UE::Editor::DataStorage::ICoreProvider& DataStorage) const;
+
+	/**
+	 * Function called when an actor is detached from its parent in the world
+	 */
+	void OnLevelActorDetached(AActor* Actor, const AActor* OldParent) const;
+
+private:
+	UE::Editor::DataStorage::FHierarchyHandle EditorActorHierarchyHandle;
+	UE::Editor::DataStorage::ICoreProvider* DataStorage;
 };
