SEverythingPicker TypeListView

diff --git a/Engine/Plugins/Experimental/EditorDataStorageFeatures/Source/TedsEverythingPicker/Private/Context/TedsPickerContextUtil.cpp b/Engine/Plugins/Experimental/EditorDataStorageFeatures/Source/TedsEverythingPicker/Private/Context/TedsPickerContextUtil.cpp
--- a/Engine/Plugins/Experimental/EditorDataStorageFeatures/Source/TedsEverythingPicker/Private/Context/TedsPickerContextUtil.cpp
+++ b/Engine/Plugins/Experimental/EditorDataStorageFeatures/Source/TedsEverythingPicker/Private/Context/TedsPickerContextUtil.cpp
@@ -3,19 +3,27 @@
 #include "Context/TedsPickerContextUtil.h"
 
 #include "DataStorage/Features.h"
+#include "DataStorage/Handles.h"
 #include "Elements/Framework/TypedElementQueryBuilder.h"
 
 #include "Elements/Columns/TypedElementLabelColumns.h"
 #include "Elements/Columns/TypedElementTypeInfoColumns.h"
-
+#include "Elements/Columns/TedsTypeInfoColumns.h"
 #include "Elements/Framework/TypedElementQueryContext.h"
 
 #include "Widgets/STedsHierarchyViewer.h"
+#include "Widgets/STedsTableViewer.h"
+
 #include "Widgets/SEverythingPicker.h"
 #include "Widgets/STedsSearchBox.h"
 
+#include "TedsHierarchyNode.h"
+#include "TedsRowArrayNode.h"
 #include "TedsQueryNode.h"
 #include "TedsRowQueryResultsNode.h"
+#include "TedsRowMergeNode.h"
+
+#define LOCTEXT_NAMESPACE "TedsPickerContextUtil"
 
 namespace UE::Editor::DataStorage::Picker
 {
@@ -98,7 +106,15 @@ namespace UE::Editor::DataStorage::Picker
 
 		ICoreProvider* Storage = GetMutableDataStorageFeature<ICoreProvider>(UE::Editor::DataStorage::StorageFeatureName);
 
-		ensure(Storage);
+		if (!ensureMsgf(Storage, TEXT("Cannot create a TEDS ObjectReferenceContext widget before TEDS is initialized")))
+		{
+			ChildSlot
+				[
+					SNew(STextBlock)
+						.Text(LOCTEXT("ObjectReferenceErrorCheck", "No valid editor data storage available."))
+				];
+			return;
+		}
 
 		TSharedRef<QueryStack::FQueryNode> ReferenceQueryNode =
 			MakeShared<QueryStack::FQueryNode>(*Storage, InArgs._Query);
@@ -112,12 +128,12 @@ namespace UE::Editor::DataStorage::Picker
 		if(InArgs._SearchingEnabled)
 		{
 			SearchBoxContainer->SetContent(SNew(STedsSearchBox)
-				.OutSearchNode(&SearchNode)
-				.QueryResultsNode(ReferenceQueryResultsNode.ToSharedPtr()));			
+				.InSearchableRowNode(ReferenceQueryResultsNode.ToSharedPtr())
+				.OutSearchNode(&SearchNode));
 		}
 
 		// If there is no SearchNode available, either disabled or there is an internal error, assign it to be the 
-		// ReferenceQueryResultsNode so that won't be null
+		// ReferenceQueryResultsNode so that it won't be null
 		if(!SearchNode)
 		{
 			SearchNode = ReferenceQueryResultsNode.ToSharedPtr();
@@ -159,4 +175,112 @@ namespace UE::Editor::DataStorage::Picker
 				]
 			];
 	}
-}
\ No newline at end of file
+
+	void STypeListContextView::Construct(const FArguments& InArgs)
+	{
+		using namespace UE::Editor::DataStorage::Queries;
+
+		ICoreProvider* Storage = GetMutableDataStorageFeature<ICoreProvider>(UE::Editor::DataStorage::StorageFeatureName);
+
+		if (!ensureMsgf(Storage, TEXT("Cannot create a TEDS TypeListContext widget before TEDS is initialized")))
+		{
+			ChildSlot
+				[
+					SNew(STextBlock)
+						.Text(LOCTEXT("TypeListErrorText", "No valid editor data storage available."))
+				];
+			return;
+		}
+
+		// Default FinalNode to an empty row list
+		TSharedPtr<QueryStack::IRowNode> FinalNode = nullptr;
+
+		TSharedRef<QueryStack::FQueryNode> TypeQueryNode =
+			MakeShared<QueryStack::FQueryNode>(*Storage, InArgs._Query);
+
+		bool bHasValidQuery = TypeQueryNode->GetQuery() != InvalidQueryHandle;
+		if (bHasValidQuery)
+		{
+			TSharedRef<QueryStack::FRowQueryResultsNode> TypeResultsNode =
+				MakeShared<QueryStack::FRowQueryResultsNode>(*Storage,
+					TypeQueryNode, QueryStack::FRowQueryResultsNode::ESyncFlags::RefreshOnUpdate);
+
+			FinalNode = TypeResultsNode.ToSharedPtr();
+		}
+		else
+		{
+			FinalNode = MakeShared<QueryStack::FRowArrayNode>(FRowHandleArray());
+		}
+
+		if (InArgs._BaseType)
+		{
+			RowHandle BaseClassRowHandle = Storage->LookupMappedRow(TypeInfo::TypeMappingDomain, FMapKey(InArgs._BaseType));
+
+			if (BaseClassRowHandle != InvalidRowHandle)
+			{
+				FRowHandleArray BaseClassHandleArray;
+				BaseClassHandleArray.Add(BaseClassRowHandle);
+
+				TSharedPtr<QueryStack::IRowNode> BaseClassRowNode =
+					MakeShared<QueryStack::FRowArrayNode>(BaseClassHandleArray);
+
+				// Acquire the row list of BaseClass + SubClasses
+				FHierarchyHandle ClassHierarchyHandle = Storage->FindHierarchyByName(TEXT("ClassHierarchy"));
+				TSharedRef<FHierarchyRowNode> HierarchyNode =
+					MakeShared<FHierarchyRowNode>(*Storage, ClassHierarchyHandle, BaseClassRowNode, FHierarchyRowNode::ESyncFlags::IncrementWhenDifferent);
+
+				// Do an initial update to populate the node
+				HierarchyNode->Update();
+
+				// If we had a valid query and a valid base class we will take the intersection of the two lists as our final
+				if (bHasValidQuery)
+				{
+					TSharedPtr<QueryStack::IRowNode> MergeNodes[] = { FinalNode, HierarchyNode.ToSharedPtr() };
+					FinalNode =
+						MakeShared<QueryStack::FRowMergeNode>(MergeNodes, QueryStack::FRowMergeNode::EMergeApproach::Repeating);
+				}
+				else
+				{
+					// No query but valid base class. We'll display the hierarchy list of BaseClass + SubClasses
+					FinalNode = HierarchyNode;
+				}
+			}
+		}
+
+		TSharedPtr<QueryStack::IRowNode> SearchNode;
+		TSharedPtr<SBox> SearchBoxContainer = SNew(SBox);
+		if (InArgs._SearchingEnabled)
+		{
+			SearchBoxContainer->SetContent(SNew(STedsSearchBox)
+				.InSearchableRowNode(FinalNode)
+				.OutSearchNode(&SearchNode));
+		}
+
+		if (SearchNode)
+		{
+			FinalNode = SearchNode;
+		}
+
+		TSharedPtr<STedsTableViewer> TypeListView = SNew(STedsTableViewer)
+			.QueryStack(FinalNode)
+			.Columns({ FTypedElementClassTypeInfoColumn::StaticStruct() })
+			.OnSelectionChanged(InArgs._OnSelectionChanged);
+
+		ChildSlot
+			[
+				SNew(SVerticalBox)
+					+ SVerticalBox::Slot()
+					.AutoHeight()
+					.Padding(4.0f, 4.0f)
+					[
+						SearchBoxContainer.ToSharedRef()
+					]
+					+ SVerticalBox::Slot()
+					[
+						TypeListView.ToSharedRef()
+					]
+			];
+	}
+}
+
+#undef LOCTEXT_NAMESPACE // TedsPickerContextUtil
\ No newline at end of file
diff --git a/Engine/Plugins/Experimental/EditorDataStorageFeatures/Source/TedsEverythingPicker/Public/Context/TedsPickerContextUtil.h b/Engine/Plugins/Experimental/EditorDataStorageFeatures/Source/TedsEverythingPicker/Public/Context/TedsPickerContextUtil.h
--- a/Engine/Plugins/Experimental/EditorDataStorageFeatures/Source/TedsEverythingPicker/Public/Context/TedsPickerContextUtil.h
+++ b/Engine/Plugins/Experimental/EditorDataStorageFeatures/Source/TedsEverythingPicker/Public/Context/TedsPickerContextUtil.h
@@ -9,6 +9,12 @@
 #include "Widgets/SCompoundWidget.h"
 
 #define UE_API TEDSEVERYTHINGPICKER_API
+
+namespace UE::Editor::DataStorage::QueryStack
+{
+	class FRowFilterNode;
+}
+
 namespace UE::Editor::DataStorage::Picker
 {
 	DECLARE_DELEGATE_OneParam(FOnSelectionChanged, RowHandle);
@@ -21,7 +27,22 @@ namespace UE::Editor::DataStorage::Picker
 			{}
 			SLATE_ARGUMENT(FQueryDescription, Query)
 			SLATE_ARGUMENT(FOnSelectionChanged, OnSelectionChanged)
-			SLATE_ARGUMENT(bool, SearchingEnabled);
+			SLATE_ARGUMENT(bool, SearchingEnabled)
+		SLATE_END_ARGS()
+
+		UE_API void Construct(const FArguments& InArgs);
+	};
+
+	class STypeListContextView : public SCompoundWidget
+	{
+	public:
+		SLATE_BEGIN_ARGS(STypeListContextView)
+			: _SearchingEnabled(false)
+			{}
+			SLATE_ARGUMENT(FQueryDescription, Query) // If set will perform the requested query. Ex: Select().Where().All<FClassTypeInfoTag>().Compile(). Results Intersect with BaseType
+			SLATE_ARGUMENT(UStruct*, BaseType) // If set, will collect a list of BaseType + SubTypes. Results intersect with Query
+			SLATE_ARGUMENT(FOnSelectionChanged, OnSelectionChanged)
+			SLATE_ARGUMENT(bool, SearchingEnabled)
 		SLATE_END_ARGS()
 
 		UE_API void Construct(const FArguments& InArgs);
diff --git a/Engine/Plugins/Experimental/EditorDataStorageFeatures/Source/TedsEverythingPicker/TedsEverythingPicker.Build.cs b/Engine/Plugins/Experimental/EditorDataStorageFeatures/Source/TedsEverythingPicker/TedsEverythingPicker.Build.cs
--- a/Engine/Plugins/Experimental/EditorDataStorageFeatures/Source/TedsEverythingPicker/TedsEverythingPicker.Build.cs
+++ b/Engine/Plugins/Experimental/EditorDataStorageFeatures/Source/TedsEverythingPicker/TedsEverythingPicker.Build.cs
@@ -30,6 +30,7 @@ public class TedsEverythingPicker : ModuleRules
 					"SlateCore",
 					"TedsQueryStack",
 					"TedsTableViewer",
+					"TedsTypeInfo",
 				});
 		}
 	}
diff --git a/Engine/Plugins/Experimental/EditorDataStorageFeatures/Source/TedsTableViewer/Private/Widgets/STedsSearchBox.cpp b/Engine/Plugins/Experimental/EditorDataStorageFeatures/Source/TedsTableViewer/Private/Widgets/STedsSearchBox.cpp
--- a/Engine/Plugins/Experimental/EditorDataStorageFeatures/Source/TedsTableViewer/Private/Widgets/STedsSearchBox.cpp
+++ b/Engine/Plugins/Experimental/EditorDataStorageFeatures/Source/TedsTableViewer/Private/Widgets/STedsSearchBox.cpp
@@ -35,7 +35,7 @@ namespace UE::Editor::DataStorage
 
 		if(InArgs._OutSearchNode)
 		{
-			SearchNode = MakeShared<QueryStack::FRowFilterNode>(Storage, InArgs._QueryResultsNode,
+			SearchNode = MakeShared<QueryStack::FRowFilterNode>(Storage, InArgs._InSearchableRowNode,
 				[this](TQueryContext<SingleRowInfo>& Context, const FTypedElementLabelColumn& LabelColumn)
 				{
 					if(SearchText.IsEmpty())
diff --git a/Engine/Plugins/Experimental/EditorDataStorageFeatures/Source/TedsTableViewer/Public/Widgets/STedsSearchBox.h b/Engine/Plugins/Experimental/EditorDataStorageFeatures/Source/TedsTableViewer/Public/Widgets/STedsSearchBox.h
--- a/Engine/Plugins/Experimental/EditorDataStorageFeatures/Source/TedsTableViewer/Public/Widgets/STedsSearchBox.h
+++ b/Engine/Plugins/Experimental/EditorDataStorageFeatures/Source/TedsTableViewer/Public/Widgets/STedsSearchBox.h
@@ -19,15 +19,15 @@ namespace UE::Editor::DataStorage
 	 * 
 	 *	SNew(STedsSearchBar)
      *		.OutSearchNode(&OutSearchNode) // TSharedPtr<QueryStack::IRowNode> OutSearchNode
-     *		.QueryResultsNode(ReferenceQueryResultsNode)
+     *		.InSearchableRowNode(ReferenceQueryResultsNode)
 	 */
 	class STedsSearchBox : public SCompoundWidget
 	{
 	public:
 		
 		SLATE_BEGIN_ARGS(STedsSearchBox) {}
-			SLATE_ARGUMENT(TSharedPtr<QueryStack::IRowNode>*, OutSearchNode)
-			SLATE_ARGUMENT(TSharedPtr<QueryStack::FRowQueryResultsNode>, QueryResultsNode)
+			SLATE_ARGUMENT(TSharedPtr<QueryStack::IRowNode>, InSearchableRowNode) // The node whose rows we'll perform a search on
+			SLATE_ARGUMENT(TSharedPtr<QueryStack::IRowNode>*, OutSearchNode)      // The search results node we receive back
 		SLATE_END_ARGS()
 
 	public:
