TedsActorCompatibility: Sanity check on the AActor* in the query registering actor components in TEDS
Also mark the queries as game thread only since they access actor/actor component internals

diff --git a/Engine/Plugins/Experimental/EditorDataStorageFeatures/Source/TedsActorCompatibility/Private/TedsActorCompatibilityModule.cpp b/Engine/Plugins/Experimental/EditorDataStorageFeatures/Source/TedsActorCompatibility/Private/TedsActorCompatibilityModule.cpp
--- a/Engine/Plugins/Experimental/EditorDataStorageFeatures/Source/TedsActorCompatibility/Private/TedsActorCompatibilityModule.cpp
+++ b/Engine/Plugins/Experimental/EditorDataStorageFeatures/Source/TedsActorCompatibility/Private/TedsActorCompatibilityModule.cpp
@@ -105,28 +105,30 @@ namespace UE::Editor::DataStorage
 				CoreProvider->RegisterQuery(
 					Select(
 						TEXT("Register ActorComponent Ownership"),
-						FProcessor(EQueryTickPhase::DuringPhysics, CoreProvider->GetQueryTickGroupName(EQueryTickGroups::Default)),
+						FProcessor(EQueryTickPhase::DuringPhysics, CoreProvider->GetQueryTickGroupName(EQueryTickGroups::Default))
+						.SetExecutionMode(EExecutionMode::GameThread),
 						[](IQueryContext& Context, RowHandle Row, const FTypedElementUObjectColumn& ObjectColumn)
 						{
-							AActor* Actor = static_cast<AActor*>(ObjectColumn.Object.Get());
-
-							FRegisterActorComponentsWithCompatibility Command;
-							Command.WkActor = Actor;
-							Command.ActorRow = Row;
-							
-							constexpr bool bIncludeFromChildActors = false;
-							Actor->ForEachComponent(bIncludeFromChildActors, [&Context, &Command](UActorComponent* Component)
+							if (AActor* Actor = static_cast<AActor*>(ObjectColumn.Object.Get()))
 							{
-								RowHandle ComponentRow = Context.LookupMappedRow(ICompatibilityProvider::ObjectMappingDomain, FMapKeyView(static_cast<const UObject*>(Component)));
-								if (ComponentRow == InvalidRowHandle)
+								FRegisterActorComponentsWithCompatibility Command;
+								Command.WkActor = Actor;
+								Command.ActorRow = Row;
+								
+								constexpr bool bIncludeFromChildActors = false;
+								Actor->ForEachComponent(bIncludeFromChildActors, [&Context, &Command](UActorComponent* Component)
+								{
+									RowHandle ComponentRow = Context.LookupMappedRow(ICompatibilityProvider::ObjectMappingDomain, FMapKeyView(static_cast<const UObject*>(Component)));
+									if (ComponentRow == InvalidRowHandle)
+									{
+										Command.ComponentsToRegister.Add(Component);
+									}
+								});
+								
+								if (!Command.ComponentsToRegister.IsEmpty())
 								{
-									Command.ComponentsToRegister.Add(Component);
+									Context.PushCommand<FRegisterActorComponentsWithCompatibility>(MoveTemp(Command));
 								}
-							});
-							
-							if (!Command.ComponentsToRegister.IsEmpty())
-							{
-								Context.PushCommand<FRegisterActorComponentsWithCompatibility>(MoveTemp(Command));
 							}
 						})
 					.Where()
@@ -141,7 +143,8 @@ namespace UE::Editor::DataStorage
 				CoreProvider->RegisterQuery(
 					Select(
 						TEXT("Set ActorComponent Hierarchy"),
-						FProcessor(EQueryTickPhase::PostPhysics, CoreProvider->GetQueryTickGroupName(EQueryTickGroups::Default)),
+						FProcessor(EQueryTickPhase::PostPhysics, CoreProvider->GetQueryTickGroupName(EQueryTickGroups::Default))
+						.SetExecutionMode(EExecutionMode::GameThread),
 						[](IQueryContext& Context, RowHandle ComponentRow, const FTypedElementUObjectColumn& ObjectColumn)
 						{
 							if (UActorComponent* ActorComponent = Cast<UActorComponent>(ObjectColumn.Object.Get()))
@@ -180,7 +183,8 @@ namespace UE::Editor::DataStorage
 				CoreProvider->RegisterQuery(
 					Select(
 						TEXT("Sync actorcomponent label to column"),
-						FProcessor(EQueryTickPhase::PrePhysics, CoreProvider->GetQueryTickGroupName(EQueryTickGroups::SyncExternalToDataStorage)),
+						FProcessor(EQueryTickPhase::PrePhysics, CoreProvider->GetQueryTickGroupName(EQueryTickGroups::SyncExternalToDataStorage))
+						.SetExecutionMode(EExecutionMode::GameThread),
 						[](const FTypedElementUObjectColumn& Actor, FTypedElementLabelColumn& Label, FTypedElementLabelHashColumn& LabelHash)
 						{
 							if (const UActorComponent* ComponentInstance = Cast<UActorComponent>(Actor.Object); ComponentInstance != nullptr)
