Teds Outliner: Fix a bug where external selection changes were not respected when a new row was added
Also add an event on the Outliner row that external systems can listen to for selection changes

diff --git a/Engine/Plugins/Experimental/EditorDataStorageFeatures/Source/TedsOutliner/Private/TedsOutlinerImpl.cpp b/Engine/Plugins/Experimental/EditorDataStorageFeatures/Source/TedsOutliner/Private/TedsOutlinerImpl.cpp
--- a/Engine/Plugins/Experimental/EditorDataStorageFeatures/Source/TedsOutliner/Private/TedsOutlinerImpl.cpp
+++ b/Engine/Plugins/Experimental/EditorDataStorageFeatures/Source/TedsOutliner/Private/TedsOutlinerImpl.cpp
@@ -298,6 +298,11 @@ void FTedsOutlinerImpl::OnMonitoredRowsAdded(DataStorage::FRowHandleArrayView In
 {
 	for (DataStorage::RowHandle Row : InRows)
 	{
+		// If the row we are adding is selected, we are probably synced to some external selection and the Outliner needs to update selection
+		if (Storage->HasColumns<FTypedElementSelectionColumn>(Row))
+		{
+			bSelectionDirty = true;
+		}
 		FSceneOutlinerHierarchyChangedData EventData;
 		EventData.Type = FSceneOutlinerHierarchyChangedData::Added;
 		EventData.Items.Add(SceneOutlinerMode->CreateItemFor<FTedsOutlinerTreeItem>(FTedsOutlinerTreeItem(Row, AsShared())));
@@ -570,7 +575,12 @@ void FTedsOutlinerImpl::RegenerateColumns()
 		MetaDataView, ColumnConstructor);
 }
 
-void FTedsOutlinerImpl::CreateItemsFromQuery(TArray<FSceneOutlinerTreeItemPtr>& OutItems, ISceneOutlinerMode* InMode) const
+DataStorage::RowHandle FTedsOutlinerImpl::GetOutlinerRowHandle() const
+{
+	return OutlinerRowHandle;
+}
+
+void FTedsOutlinerImpl::CreateItemsFromQuery(TArray<FSceneOutlinerTreeItemPtr>& OutItems, ISceneOutlinerMode* InMode)
 {
 	using namespace UE::Editor::DataStorage::Queries;
 
@@ -594,6 +604,9 @@ void FTedsOutlinerImpl::CreateItemsFromQuery(TArray<FSceneOutlinerTreeItemPtr>&
 		}
 	}
 
+	// We need to update selection since the whole outliner was re-populated
+	bSelectionDirty = true;
+
 }
 
 void FTedsOutlinerImpl::CreateChildren(const FSceneOutlinerTreeItemPtr& Item, TArray<FSceneOutlinerTreeItemPtr>& OutChildren) const
@@ -1120,7 +1133,7 @@ TOptional<FName> FTedsOutlinerImpl::GetSelectionSetName() const
 	return SelectionSetName;
 }
 
-FTedsOutlinerImpl::FOnTedsOutlinerSelectionChanged& FTedsOutlinerImpl::OnSelectionChanged()
+FOnTedsOutlinerSelectionChanged& FTedsOutlinerImpl::OnSelectionChanged()
 {
 	return OnTedsOutlinerSelectionChanged;
 }
diff --git a/Engine/Plugins/Experimental/EditorDataStorageFeatures/Source/TedsOutliner/Private/TedsOutlinerMode.cpp b/Engine/Plugins/Experimental/EditorDataStorageFeatures/Source/TedsOutliner/Private/TedsOutlinerMode.cpp
--- a/Engine/Plugins/Experimental/EditorDataStorageFeatures/Source/TedsOutliner/Private/TedsOutlinerMode.cpp
+++ b/Engine/Plugins/Experimental/EditorDataStorageFeatures/Source/TedsOutliner/Private/TedsOutlinerMode.cpp
@@ -75,6 +75,11 @@ void FTedsOutlinerMode::OnSelectionChanged()
 		}
 		return false;
 	});
+
+	if (FTedsOutlinerSelectionChangeColumn* SelectionChangeColumn = Storage->GetColumn<FTedsOutlinerSelectionChangeColumn>(TedsOutlinerImpl->GetOutlinerRowHandle()))
+	{
+		SelectionChangeColumn->OnSelectionChanged.Broadcast();
+	}
 }
 
 void FTedsOutlinerMode::SynchronizeSelection()
diff --git a/Engine/Plugins/Experimental/EditorDataStorageFeatures/Source/TedsOutliner/Public/Columns/TedsOutlinerColumns.h b/Engine/Plugins/Experimental/EditorDataStorageFeatures/Source/TedsOutliner/Public/Columns/TedsOutlinerColumns.h
--- a/Engine/Plugins/Experimental/EditorDataStorageFeatures/Source/TedsOutliner/Public/Columns/TedsOutlinerColumns.h
+++ b/Engine/Plugins/Experimental/EditorDataStorageFeatures/Source/TedsOutliner/Public/Columns/TedsOutlinerColumns.h
@@ -17,6 +17,9 @@ namespace UE::Editor::Outliner
 	// Delegate that refreshes columns on the Teds Outliner it is bound to when executed
 	DECLARE_DELEGATE(FOnRefreshColumns)
 
+	// Delegate called when the selection in the Outliner changes
+	DECLARE_MULTICAST_DELEGATE(FOnTedsOutlinerSelectionChanged)
+
 	// Delegate to get the TreeItemID for a row
 	DECLARE_DELEGATE_RetVal_OneParam(FSceneOutlinerTreeItemID, FTreeItemIDDealiaser, DataStorage::RowHandle);
 
@@ -66,3 +69,12 @@ struct FTedsOutlinerDealiaserColumn final : public FEditorDataStorageColumn
 
 	UE::Editor::Outliner::FTreeItemIDDealiaser Dealiaser;
 };
+
+/** Column with a delegate that is broadcast when the outliner's selection changes */
+USTRUCT()
+struct FTedsOutlinerSelectionChangeColumn final : public FEditorDataStorageColumn
+{
+	GENERATED_BODY()
+
+	UE::Editor::Outliner::FOnTedsOutlinerSelectionChanged OnSelectionChanged;
+};
\ No newline at end of file
diff --git a/Engine/Plugins/Experimental/EditorDataStorageFeatures/Source/TedsOutliner/Public/TedsOutlinerImpl.h b/Engine/Plugins/Experimental/EditorDataStorageFeatures/Source/TedsOutliner/Public/TedsOutlinerImpl.h
--- a/Engine/Plugins/Experimental/EditorDataStorageFeatures/Source/TedsOutliner/Public/TedsOutlinerImpl.h
+++ b/Engine/Plugins/Experimental/EditorDataStorageFeatures/Source/TedsOutliner/Public/TedsOutlinerImpl.h
@@ -5,6 +5,7 @@
 #include "Elements/Interfaces/TypedElementDataStorageInterface.h"
 #include "ISceneOutlinerHierarchy.h"
 #include "ISceneOutlinerMode.h"
+#include "Columns/TedsOutlinerColumns.h"
 #include "Elements/Columns/TypedElementHiearchyColumns.h"
 #include "Elements/Framework/TypedElementRowHandleArrayView.h"
 #include "Elements/Interfaces/TypedElementDataStorageUiInterface.h"
@@ -176,7 +177,6 @@ public:
 	UE_API TOptional<FName> GetSelectionSetName() const;
 
 	// Delegate fired when the selection in TEDS changes, only if SelectionSetName is set
-	DECLARE_MULTICAST_DELEGATE(FOnTedsOutlinerSelectionChanged)
 	UE_API FOnTedsOutlinerSelectionChanged& OnSelectionChanged();
 
 	// Delegate fired when the hierarchy changes due to item addition/removal/move
@@ -203,7 +203,7 @@ public:
 	UE_API void AppendExternalQueries(DataStorage::FQueryDescription& OutQuery);
 	
 	// Outliner specific functionality
-	UE_API void CreateItemsFromQuery(TArray<FSceneOutlinerTreeItemPtr>& OutItems, ISceneOutlinerMode* InMode) const;
+	UE_API void CreateItemsFromQuery(TArray<FSceneOutlinerTreeItemPtr>& OutItems, ISceneOutlinerMode* InMode);
 	UE_API void CreateChildren(const FSceneOutlinerTreeItemPtr& Item, TArray<FSceneOutlinerTreeItemPtr>& OutChildren) const;
 
 	// Get the parent row for a given row
@@ -225,6 +225,8 @@ public:
 		EColumnSortPriority::Type Priority);
 	void SortItems(TArray<FSceneOutlinerTreeItemPtr>& Items, EColumnSortMode::Type Direction) const;
 
+	UE_API DataStorage::RowHandle GetOutlinerRowHandle() const;
+
 protected:
 	
 	UE_API void ClearSelection() const;
