Teds Outliner: Add selection column to the outliner table by default

diff --git a/Engine/Plugins/Experimental/EditorDataStorageFeatures/Source/TedsOutliner/Private/Compatibility/SceneOutlinerTedsBridge.cpp b/Engine/Plugins/Experimental/EditorDataStorageFeatures/Source/TedsOutliner/Private/Compatibility/SceneOutlinerTedsBridge.cpp
--- a/Engine/Plugins/Experimental/EditorDataStorageFeatures/Source/TedsOutliner/Private/Compatibility/SceneOutlinerTedsBridge.cpp
+++ b/Engine/Plugins/Experimental/EditorDataStorageFeatures/Source/TedsOutliner/Private/Compatibility/SceneOutlinerTedsBridge.cpp
@@ -40,7 +40,7 @@ void USceneOutlinerTedsBridgeFactory::PreRegister(UE::Editor::DataStorage::ICore
 
 void USceneOutlinerTedsBridgeFactory::RegisterTables(UE::Editor::DataStorage::ICoreProvider& DataStorage)
 {
-	DataStorage.RegisterTable<FTedsOutlinerColumn, FTedsOutlinerColumnQueryColumn>(UE::Editor::Outliner::Helpers::GetTedsOutlinerTableName());
+	DataStorage.RegisterTable<FTedsOutlinerColumn, FTedsOutlinerColumnQueryColumn, FTedsOutlinerSelectionChangeColumn>(UE::Editor::Outliner::Helpers::GetTedsOutlinerTableName());
 }
 
 void USceneOutlinerTedsBridgeFactory::RegisterWidgetPurposes(UE::Editor::DataStorage::IUiProvider& DataStorageUi) const
