Fix startup crash bug in TypeInfoModule

diff --git a/Engine/Plugins/Experimental/EditorDataStorageFeatures/Source/TedsTypeInfo/Private/TedsTypeInfoFactory.cpp b/Engine/Plugins/Experimental/EditorDataStorageFeatures/Source/TedsTypeInfo/Private/TedsTypeInfoFactory.cpp
--- a/Engine/Plugins/Experimental/EditorDataStorageFeatures/Source/TedsTypeInfo/Private/TedsTypeInfoFactory.cpp
+++ b/Engine/Plugins/Experimental/EditorDataStorageFeatures/Source/TedsTypeInfo/Private/TedsTypeInfoFactory.cpp
@@ -8,6 +8,8 @@
 #include "ClassViewerFilter.h"
 #include "ClassViewerModule.h"
 
+#include "TedsTypeInfoModule.h"
+
 #include "DataStorage/Features.h"
 #include "Elements/Interfaces/TypedElementDataStorageInterface.h"
 #include "Elements/Columns/TypedElementLabelColumns.h"
@@ -68,6 +70,15 @@ void UTypeInfoFactory::RegisterQueries(UE::Editor::DataStorage::ICoreProvider& D
 			.Where()
 				.All<FTypeInfoRequiresHierarchyUpdateTag>()
 			.Compile());
+
+	// Populate our type info on register if our CVar is enabled
+	static const auto TedsTypeInfoEnabledCVar = IConsoleManager::Get().FindConsoleVariable(TEXT("TEDS.Feature.TypeInfoIntegration.Enable"));
+	if (TedsTypeInfoEnabledCVar && TedsTypeInfoEnabledCVar->GetBool())
+	{
+		RefreshAllTypeInfo();
+	}
+
+	TypeInfo::FTedsTypeInfoModule::GetChecked().SetTypeInfoFactoryEnabled();
 }
 
 UE::Editor::DataStorage::TableHandle UTypeInfoFactory::GetTypeTable() const
diff --git a/Engine/Plugins/Experimental/EditorDataStorageFeatures/Source/TedsTypeInfo/Private/TedsTypeInfoModule.cpp b/Engine/Plugins/Experimental/EditorDataStorageFeatures/Source/TedsTypeInfo/Private/TedsTypeInfoModule.cpp
--- a/Engine/Plugins/Experimental/EditorDataStorageFeatures/Source/TedsTypeInfo/Private/TedsTypeInfoModule.cpp
+++ b/Engine/Plugins/Experimental/EditorDataStorageFeatures/Source/TedsTypeInfo/Private/TedsTypeInfoModule.cpp
@@ -23,11 +23,11 @@ namespace UE::Editor::DataStorage::TypeInfo
 
 					if (bIsEnabled)
 					{
-						Module.EnableTedsTypeInfoIntegration();
+						Module.RefreshTypeInfo();
 					}
 					else
 					{
-						Module.DisableTedsTypeInfoIntegration();
+						Module.FlushTypeInfo();
 					}
 				}));
 	}
@@ -39,10 +39,6 @@ namespace UE::Editor::DataStorage::TypeInfo
 		UE::Editor::DataStorage::OnEditorDataStorageFeaturesEnabled().AddLambda([this]()
 			{
 				bDataStorageFeaturesEnabled = true;
-				if (Private::CVarTEDSTypeInfoEnabled.GetValueOnGameThread())
-				{
-					EnableTedsTypeInfoIntegration();
-				}
 			});
 	}
 
@@ -63,22 +59,12 @@ namespace UE::Editor::DataStorage::TypeInfo
 
 	void FTedsTypeInfoModule::EnableTedsTypeInfoIntegration()
 	{
-		if (!bDataStorageFeaturesEnabled)
-		{
-			return;
-		}
-
 		Private::CVarTEDSTypeInfoEnabled.AsVariable()->Set(true);
 		RefreshTypeInfo();
 	}
 
 	void FTedsTypeInfoModule::DisableTedsTypeInfoIntegration()
 	{
-		if (!bDataStorageFeaturesEnabled)
-		{
-			return;
-		}
-
 		FlushTypeInfo();
 		Private::CVarTEDSTypeInfoEnabled.AsVariable()->Set(false);
 	}
@@ -97,6 +83,11 @@ namespace UE::Editor::DataStorage::TypeInfo
 			return;
 		}
 
+		if (!bDataStorageFeaturesEnabled || bTypeInfoFactoryEnabled)
+		{
+			return;
+		}
+
 		ICoreProvider* DataStorage = GetMutableDataStorageFeature<ICoreProvider>(StorageFeatureName);
 
 		ensure(DataStorage);
@@ -116,6 +107,11 @@ namespace UE::Editor::DataStorage::TypeInfo
 			return;
 		}
 
+		if (!bDataStorageFeaturesEnabled || !bTypeInfoFactoryEnabled)
+		{
+			return;
+		}
+
 		ICoreProvider* DataStorage = GetMutableDataStorageFeature<ICoreProvider>(StorageFeatureName);
 
 		ensure(DataStorage);
diff --git a/Engine/Plugins/Experimental/EditorDataStorageFeatures/Source/TedsTypeInfo/Public/TedsTypeInfoModule.h b/Engine/Plugins/Experimental/EditorDataStorageFeatures/Source/TedsTypeInfo/Public/TedsTypeInfoModule.h
--- a/Engine/Plugins/Experimental/EditorDataStorageFeatures/Source/TedsTypeInfo/Public/TedsTypeInfoModule.h
+++ b/Engine/Plugins/Experimental/EditorDataStorageFeatures/Source/TedsTypeInfo/Public/TedsTypeInfoModule.h
@@ -20,10 +20,13 @@ namespace UE::Editor::DataStorage::TypeInfo
 		TEDSTYPEINFO_API bool IsTedsTypeInfoIntegrationEnabled();
 
 		TEDSTYPEINFO_API void RefreshTypeInfo();
+		TEDSTYPEINFO_API void FlushTypeInfo();
+
+		void SetTypeInfoFactoryEnabled() { bTypeInfoFactoryEnabled = true; }
 
 	private:
-		void FlushTypeInfo();
 
+		bool bTypeInfoFactoryEnabled = false;
 		bool bDataStorageFeaturesEnabled = false;
 	};
 }
\ No newline at end of file
