Editor Data Storage: Add support for disabling manual sorting of colums when creating Teds tableviewer columns by applying custom metadata flags.

diff --git a/Engine/Plugins/Experimental/EditorDataStorageFeatures/Source/TedsOutliner/Private/TedsOutlinerColumn.cpp b/Engine/Plugins/Experimental/EditorDataStorageFeatures/Source/TedsOutliner/Private/TedsOutlinerColumn.cpp
--- a/Engine/Plugins/Experimental/EditorDataStorageFeatures/Source/TedsOutliner/Private/TedsOutlinerColumn.cpp
+++ b/Engine/Plugins/Experimental/EditorDataStorageFeatures/Source/TedsOutliner/Private/TedsOutlinerColumn.cpp
@@ -37,7 +37,7 @@ namespace UE::Editor::Outliner
 		using namespace UE::Editor::DataStorage;
 
 		TableViewerColumnImpl = MakeUnique<FTedsTableViewerColumn>(NameId, InitParams.CellWidgetConstructor, InitParams.ColumnTypes,
-			InitParams.HeaderWidgetConstructor, FComboMetaDataView(FGenericMetaDataView(MetaData)).Next(FQueryMetaDataView(QueryDescription)));
+			InitParams.HeaderWidgetConstructor, QueryDescription.MetaData);
 
 		TableViewerColumnImpl->SetIsRowVisibleDelegate(
 			FTedsTableViewerColumn::FIsRowVisible::CreateRaw(this, &FTedsOutlinerUiColumn::IsRowVisible)
diff --git a/Engine/Plugins/Experimental/EditorDataStorageFeatures/Source/TedsTableViewer/Private/TedsTableViewerColumn.cpp b/Engine/Plugins/Experimental/EditorDataStorageFeatures/Source/TedsTableViewer/Private/TedsTableViewerColumn.cpp
--- a/Engine/Plugins/Experimental/EditorDataStorageFeatures/Source/TedsTableViewer/Private/TedsTableViewerColumn.cpp
+++ b/Engine/Plugins/Experimental/EditorDataStorageFeatures/Source/TedsTableViewer/Private/TedsTableViewerColumn.cpp
@@ -22,7 +22,7 @@ namespace UE::Editor::DataStorage
 {
 	FTedsTableViewerColumn::FTedsTableViewerColumn(const FName& InColumnName, const TSharedPtr<FTypedElementWidgetConstructor>& InCellWidgetConstructor,
 		const TArray<TWeakObjectPtr<const UScriptStruct>>& InMatchedColumns, const TSharedPtr<FTypedElementWidgetConstructor>& InHeaderWidgetConstructor,
-		const FMetaDataView& InWidgetMetaData)
+		const FMetaData& InWidgetMetaData)
 		: ColumnName(InColumnName)
 		, CellWidgetConstructor(InCellWidgetConstructor)
 		, HeaderWidgetConstructor(InHeaderWidgetConstructor)
@@ -43,12 +43,12 @@ namespace UE::Editor::DataStorage
 		// Retrieve sort functors to allow the content of the column to be used for sorting
 		if (HeaderWidgetConstructor)
 		{
-			ColumnSorters = HeaderWidgetConstructor->ConstructColumnSorters(Storage, StorageUi, InWidgetMetaData);
+			ColumnSorters = HeaderWidgetConstructor->ConstructColumnSorters(Storage, StorageUi, FGenericMetaDataView(WidgetMetaData));
 		}
 		// Prefer to use the sort function in the header, but if it doesn't exist, use the one from the cell.
 		if (CellWidgetConstructor && ColumnSorters.IsEmpty())
 		{
-			ColumnSorters = CellWidgetConstructor->ConstructColumnSorters(Storage, StorageUi, InWidgetMetaData);
+			ColumnSorters = CellWidgetConstructor->ConstructColumnSorters(Storage, StorageUi, FGenericMetaDataView(WidgetMetaData));
 		}
 
 		RegisterQueries();
@@ -88,7 +88,7 @@ namespace UE::Editor::DataStorage
 				WidgetRowSetupDelegate(*Storage, UiRowHandle);
 			}
 
-			RowWidget = StorageUi->ConstructWidget(UiRowHandle, *CellWidgetConstructor, WidgetMetaData);
+			RowWidget = StorageUi->ConstructWidget(UiRowHandle, *CellWidgetConstructor, FGenericMetaDataView(WidgetMetaData));
 		}
 
 		return RowWidget;
@@ -123,7 +123,7 @@ namespace UE::Editor::DataStorage
 			}
 
 			Widget = StorageUi->ConstructWidget(UiRowHandle, *HeaderWidgetConstructor,
-				WidgetMetaData);
+				FGenericMetaDataView(WidgetMetaData));
 		}
 		if (!Widget.IsValid())
 		{
@@ -138,6 +138,24 @@ namespace UE::Editor::DataStorage
 				.Text(DefaultHeaderText);
 		}
 
+		const bool* bAllowSorting;
+		const bool* bColumnInitialSortIsAscending;
+
+		{
+			UE::Editor::DataStorage::FMetaDataEntryView MetaDataEntryView = FGenericMetaDataView(WidgetMetaData).FindGeneric("bAllowSorting");
+		    bAllowSorting = MetaDataEntryView.TryGetExact<bool>();
+		}
+		{
+			UE::Editor::DataStorage::FMetaDataEntryView MetaDataEntryView = FGenericMetaDataView(WidgetMetaData).FindGeneric("bColumnInitialSortIsAscending");
+			bColumnInitialSortIsAscending = MetaDataEntryView.TryGetExact<bool>();
+		}
+
+		EColumnSortMode::Type SortMode = EColumnSortMode::Ascending;
+		if (bColumnInitialSortIsAscending)
+		{
+			SortMode = *bColumnInitialSortIsAscending ? EColumnSortMode::Ascending : EColumnSortMode::Descending;
+		}
+
 		SHeaderRow::FColumn::FArguments Column = SHeaderRow::Column(ColumnName)
 			.FillWidth(1)
 			.HeaderComboVisibility(EHeaderComboVisibility::OnHover)
@@ -147,6 +165,8 @@ namespace UE::Editor::DataStorage
 			.SortPriority_Raw(this, &FTedsTableViewerColumn::GetSortPriority)
 			.IsSorting_Raw(this, &FTedsTableViewerColumn::IsSorting)
 			.OnSort_Raw(this, &FTedsTableViewerColumn::OnSortCallback)
+			.InitialSortMode(SortMode)
+			.CanManuallySort(bAllowSorting ? *bAllowSorting : true)
 			.HeaderContent()
 			[
 				SNew(SHorizontalBox)
@@ -378,7 +398,7 @@ namespace UE::Editor::DataStorage
 			// If a column was added and we are here, we need to re-create the widget
 			if(*bColumnAddedPtr)
 			{
-				if(const TSharedPtr<SWidget> RowWidget = StorageUi->ConstructInternalWidget(Row, *CellWidgetConstructor, WidgetMetaData))
+				if(const TSharedPtr<SWidget> RowWidget = StorageUi->ConstructInternalWidget(Row, *CellWidgetConstructor, FGenericMetaDataView(WidgetMetaData)))
 				{
 					TedsWidget->SetContent(RowWidget.ToSharedRef());
 				}
diff --git a/Engine/Plugins/Experimental/EditorDataStorageFeatures/Source/TedsTableViewer/Public/TedsTableViewerColumn.h b/Engine/Plugins/Experimental/EditorDataStorageFeatures/Source/TedsTableViewer/Public/TedsTableViewerColumn.h
--- a/Engine/Plugins/Experimental/EditorDataStorageFeatures/Source/TedsTableViewer/Public/TedsTableViewerColumn.h
+++ b/Engine/Plugins/Experimental/EditorDataStorageFeatures/Source/TedsTableViewer/Public/TedsTableViewerColumn.h
@@ -34,7 +34,7 @@ namespace UE::Editor::DataStorage
 			const TSharedPtr<FTypedElementWidgetConstructor>& InCellWidgetConstructor, // The widget constructor to use for this column
 			const TArray<TWeakObjectPtr<const UScriptStruct>>& InMatchedColumns = {}, // An optional list of matched Teds columns 
 			const TSharedPtr<FTypedElementWidgetConstructor>& InHeaderWidgetConstructor = nullptr, // Optional constructor to use for the header widget
-			const FMetaDataView& InWidgetMetaData = FMetaDataView()); // Optional metadata to use when constructing widgets
+			const FMetaData& MetaData = FMetaData()); // Optional metadata to use when constructing widgets
 
 		TEDSTABLEVIEWER_API ~FTedsTableViewerColumn();
 		
@@ -100,7 +100,7 @@ namespace UE::Editor::DataStorage
 		Queries::FConditions MatchedColumnConditions;
 		
 		// The Metadata used to create widgets
-		FMetaDataView WidgetMetaData;
+		const FMetaData WidgetMetaData;
 
 		// TEDS Constructs
 		ICoreProvider* Storage;
