Teds:
Add a new TedsEditorCompatibility module
Register UWorld and ULevels with TEDS Compat

diff --git a/Engine/Plugins/Experimental/EditorDataStorageFeatures/EditorDataStorageFeatures.uplugin b/Engine/Plugins/Experimental/EditorDataStorageFeatures/EditorDataStorageFeatures.uplugin
--- a/Engine/Plugins/Experimental/EditorDataStorageFeatures/EditorDataStorageFeatures.uplugin
+++ b/Engine/Plugins/Experimental/EditorDataStorageFeatures/EditorDataStorageFeatures.uplugin
@@ -71,6 +71,11 @@
 			"Name": "TedsTypedElementBridge",
 			"Type": "Editor"
 		}
+		,
+		{
+			"Name": "TedsEditorCompatibility",
+			"Type": "Editor"
+		}
 	],
 	"Plugins": [
 	],
diff --git a/Engine/Plugins/Experimental/EditorDataStorageFeatures/Source/TedsEditorCompatibility/Private/Level/TedsLevelFactory.cpp b/Engine/Plugins/Experimental/EditorDataStorageFeatures/Source/TedsEditorCompatibility/Private/Level/TedsLevelFactory.cpp
new file mode 100644
--- /dev/null
+++ b/Engine/Plugins/Experimental/EditorDataStorageFeatures/Source/TedsEditorCompatibility/Private/Level/TedsLevelFactory.cpp
@@ -0,0 +1,153 @@
+// Copyright Epic Games, Inc. All Rights Reserved.
+
+#include "TedsLevelFactory.h"
+
+#include "DataStorage/Features.h"
+#include "Elements/Columns/TypedElementCompatibilityColumns.h"
+#include "Elements/Columns/TypedElementIconOverrideColumns.h"
+#include "Elements/Columns/TypedElementLabelColumns.h"
+#include "Elements/Columns/TypedElementMiscColumns.h"
+#include "Elements/Framework/TypedElementQueryBuilder.h"
+#include "Elements/Interfaces/TypedElementDataStorageCompatibilityInterface.h"
+#include "Elements/Interfaces/TypedElementDataStorageInterface.h"
+#include "Engine/Level.h"
+#include "Engine/World.h"
+#include "Misc/PackageName.h"
+#include "Styling/SlateIconFinder.h"
+#include "UObject/Package.h"
+
+namespace UE::Editor::DataStorage::Levels
+{
+	static const FName LevelTableName("Editor_Level");
+}
+
+void UTedsLevelFactory::RegisterTables(UE::Editor::DataStorage::ICoreProvider& DataStorage)
+{
+	using namespace UE::Editor::DataStorage;
+	
+	TableHandle LevelTable = DataStorage.RegisterTable<FLevelTag, FTypedElementLabelColumn, FTypedElementIconOverrideColumn>(Levels::LevelTableName);
+
+	DataStorageCompat = GetMutableDataStorageFeature<ICompatibilityProvider>(CompatibilityFeatureName);
+	checkf(DataStorageCompat, TEXT("TEDS Factory cannot be init before the data storage compatibility feature is available"));
+	
+	DataStorageCompat->RegisterTypeTableAssociation(ULevel::StaticClass(), LevelTable);
+}
+
+void UTedsLevelFactory::PreRegister(UE::Editor::DataStorage::ICoreProvider& DataStorage)
+{
+	FWorldDelegates::LevelAddedToWorld.AddUObject(this, &UTedsLevelFactory::OnLevelAddedToWorld);
+	FWorldDelegates::LevelRemovedFromWorld.AddUObject(this, &UTedsLevelFactory::OnLevelRemovedFromWorld);
+}
+
+void UTedsLevelFactory::PreShutdown(UE::Editor::DataStorage::ICoreProvider& DataStorage)
+{
+	FWorldDelegates::LevelAddedToWorld.RemoveAll(this);
+	FWorldDelegates::LevelRemovedFromWorld.RemoveAll(this);
+}
+
+void UTedsLevelFactory::RegisterQueries(UE::Editor::DataStorage::ICoreProvider& DataStorage)
+{
+	using namespace UE::Editor::DataStorage;
+	using namespace UE::Editor::DataStorage::Queries;
+	
+	DataStorage.RegisterQuery(
+		Select(
+			TEXT("Sync level label to column"),
+			FProcessor(EQueryTickPhase::PrePhysics, DataStorage.GetQueryTickGroupName(EQueryTickGroups::SyncExternalToDataStorage))
+				.SetExecutionMode(EExecutionMode::GameThread),
+			[](const FTypedElementUObjectColumn& LevelColumn, FTypedElementLabelColumn& Label)
+			{
+				if (const ULevel* Level = Cast<ULevel>(LevelColumn.Object); Level != nullptr)
+				{
+					Label.Label = FPackageName::GetShortName(Level->GetOutermost()->GetName());
+				}
+			}
+		)
+		.Where()
+			.All<FLevelTag, FTypedElementSyncFromWorldTag>()
+		.Compile());
+
+	struct FRegisterLevelCommand
+	{
+		void operator()() const
+		{
+			if (ULevel* Level = LevelPtr.Get())
+			{
+				TedsCompat->AddCompatibleObject(Level);
+			}
+		}
+
+		ICompatibilityProvider* TedsCompat;
+		TWeakObjectPtr<ULevel> LevelPtr;
+	};
+
+	// The Persistent world for each world needs to be tracked separately, since it doesn't go through LevelAddedToWorld/LevelRemovedFromWorld
+	DataStorage.RegisterQuery(
+		Select(TEXT("Register Persistent Level on World creation"),
+			FObserver::OnAdd<FWorldTag>(),
+			[this](IQueryContext& Context, RowHandle Row, const FTypedElementUObjectColumn& ObjectColumn)
+			{
+				if (UWorld* World = Cast<UWorld>(ObjectColumn.Object.Get()))
+				{
+					FRegisterLevelCommand Command{.TedsCompat = DataStorageCompat, .LevelPtr = World->PersistentLevel};
+					Context.PushCommand(Command);
+				}
+			})
+		.Where()
+			.All<FWorldTag>()
+		.Compile());
+
+	struct FUnregisterLevelCommand
+	{
+		void operator()() const
+		{
+			if (ULevel* Level = LevelPtr.Get())
+			{
+				TedsCompat->RemoveCompatibleObject(Level);
+			}
+		}
+
+		ICompatibilityProvider* TedsCompat;
+		TWeakObjectPtr<ULevel> LevelPtr;
+	};
+
+	DataStorage.RegisterQuery(
+		Select(TEXT("Unregister Persistent Level on World removal"),
+			FObserver::OnRemove<FTypedElementUObjectColumn>(),
+			[this](IQueryContext& Context, RowHandle Row, const FTypedElementUObjectColumn& ObjectColumn)
+			{
+				// We need to make sure the persistent level is removed properly if the world is marked for GC
+				constexpr bool bEvenIfPendingKill = true;
+				if (UWorld* World = Cast<UWorld>(ObjectColumn.Object.Get(bEvenIfPendingKill)))
+				{
+					FUnregisterLevelCommand Command{.TedsCompat = DataStorageCompat, .LevelPtr = World->PersistentLevel};
+					Context.PushCommand(Command);
+				}
+			})
+		.Where()
+			.All<FWorldTag>()
+		.Compile());
+
+	DataStorage.RegisterQuery(
+		Select(TEXT("Add Icon override to level"),
+			FObserver::OnAdd<FTypedElementUObjectColumn>(),
+			[this](IQueryContext& Context, RowHandle Row, FTypedElementIconOverrideColumn& IconOverrideColumn)
+			{
+				// ULevels show up with the icon for UWorld in the editor (e.g Outliner)
+				FSlateIcon WorldIcon = FSlateIconFinder::FindIconForClass(UWorld::StaticClass());
+				IconOverrideColumn.IconName = WorldIcon.GetStyleName();
+			})
+		.Where()
+			.All<FLevelTag>()
+		.Compile());
+}
+
+void UTedsLevelFactory::OnLevelAddedToWorld(ULevel* InLevel, UWorld* InWorld)
+{
+	DataStorageCompat->AddCompatibleObject(InLevel);
+}
+
+void UTedsLevelFactory::OnLevelRemovedFromWorld(ULevel* InLevel, UWorld* InWorld)
+{
+	DataStorageCompat->RemoveCompatibleObject(InLevel);
+}
diff --git a/Engine/Plugins/Experimental/EditorDataStorageFeatures/Source/TedsEditorCompatibility/Private/Level/TedsLevelFactory.h b/Engine/Plugins/Experimental/EditorDataStorageFeatures/Source/TedsEditorCompatibility/Private/Level/TedsLevelFactory.h
new file mode 100644
--- /dev/null
+++ b/Engine/Plugins/Experimental/EditorDataStorageFeatures/Source/TedsEditorCompatibility/Private/Level/TedsLevelFactory.h
@@ -0,0 +1,29 @@
+// Copyright Epic Games, Inc. All Rights Reserved.
+
+#pragma once
+
+#include "DataStorage/Handles.h"
+#include "Elements/Interfaces/TypedElementDataStorageFactory.h"
+
+#include "TedsLevelFactory.generated.h"
+
+UCLASS()
+class UTedsLevelFactory : public UEditorDataStorageFactory
+{
+	GENERATED_BODY()
+
+public:
+	~UTedsLevelFactory() override = default;
+
+	virtual void RegisterTables(UE::Editor::DataStorage::ICoreProvider& DataStorage) override;
+	virtual void PreRegister(UE::Editor::DataStorage::ICoreProvider& DataStorage) override;
+	virtual void PreShutdown(UE::Editor::DataStorage::ICoreProvider& DataStorage) override;
+	virtual void RegisterQueries(UE::Editor::DataStorage::ICoreProvider& DataStorage) override;
+
+private:
+	void OnLevelAddedToWorld(ULevel* InLevel, UWorld* InWorld);
+	void OnLevelRemovedFromWorld(ULevel* InLevel, UWorld* InWorld);
+
+private:
+	UE::Editor::DataStorage::ICompatibilityProvider* DataStorageCompat;
+};
\ No newline at end of file
diff --git a/Engine/Plugins/Experimental/EditorDataStorageFeatures/Source/TedsEditorCompatibility/Private/TedsEditorCompatibilityModule.cpp b/Engine/Plugins/Experimental/EditorDataStorageFeatures/Source/TedsEditorCompatibility/Private/TedsEditorCompatibilityModule.cpp
new file mode 100644
--- /dev/null
+++ b/Engine/Plugins/Experimental/EditorDataStorageFeatures/Source/TedsEditorCompatibility/Private/TedsEditorCompatibilityModule.cpp
@@ -0,0 +1,20 @@
+// Copyright Epic Games, Inc. All Rights Reserved.
+
+#include "Modules/ModuleManager.h"
+#include "Modules/ModuleInterface.h"
+
+class FTedsEditorCompatibilityModule : public IModuleInterface
+{
+public:
+
+	virtual void StartupModule() override
+	{
+	}
+
+	virtual void ShutdownModule() override
+	{
+	}
+};
+
+IMPLEMENT_MODULE(FTedsEditorCompatibilityModule, TedsEditorCompatibility);
+
diff --git a/Engine/Plugins/Experimental/EditorDataStorageFeatures/Source/TedsEditorCompatibility/Private/World/TedsWorldFactory.cpp b/Engine/Plugins/Experimental/EditorDataStorageFeatures/Source/TedsEditorCompatibility/Private/World/TedsWorldFactory.cpp
new file mode 100644
--- /dev/null
+++ b/Engine/Plugins/Experimental/EditorDataStorageFeatures/Source/TedsEditorCompatibility/Private/World/TedsWorldFactory.cpp
@@ -0,0 +1,77 @@
+// Copyright Epic Games, Inc. All Rights Reserved.
+
+#include "TedsWorldFactory.h"
+
+#include "DataStorage/Features.h"
+#include "Elements/Columns/TypedElementCompatibilityColumns.h"
+#include "Elements/Interfaces/TypedElementDataStorageCompatibilityInterface.h"
+#include "Elements/Interfaces/TypedElementDataStorageInterface.h"
+#include "Elements/Columns/TypedElementLabelColumns.h"
+#include "Elements/Columns/TypedElementMiscColumns.h"
+#include "Elements/Framework/TypedElementQueryBuilder.h"
+#include "Engine/Engine.h"
+#include "Engine/World.h"
+
+namespace UE::Editor::DataStorage::World
+{
+	static const FName WorldTableName("Editor_World");
+}
+
+void UTedsWorldFactory::RegisterTables(UE::Editor::DataStorage::ICoreProvider& DataStorage)
+{
+	using namespace UE::Editor::DataStorage;
+	
+	TableHandle WorldTable = DataStorage.RegisterTable<FWorldTag, FTypedElementLabelColumn>(World::WorldTableName);
+
+	DataStorageCompat = GetMutableDataStorageFeature<ICompatibilityProvider>(CompatibilityFeatureName);
+	checkf(DataStorageCompat, TEXT("TEDS Factory cannot be init before the data storage compatibility feature is available"));
+	
+	DataStorageCompat->RegisterTypeTableAssociation(UWorld::StaticClass(), WorldTable);
+}
+
+void UTedsWorldFactory::PreRegister(UE::Editor::DataStorage::ICoreProvider& DataStorage)
+{
+	FWorldDelegates::OnPostWorldInitialization.AddUObject(this, &UTedsWorldFactory::PostWorldInitialized);
+	FWorldDelegates::OnPreWorldFinishDestroy.AddUObject(this, &UTedsWorldFactory::OnPreWorldFinishDestroy);
+}
+
+void UTedsWorldFactory::PreShutdown(UE::Editor::DataStorage::ICoreProvider& DataStorage)
+{
+	FWorldDelegates::OnPostWorldInitialization.RemoveAll(this);
+	FWorldDelegates::OnPreWorldFinishDestroy.RemoveAll(this);
+}
+
+void UTedsWorldFactory::RegisterQueries(UE::Editor::DataStorage::ICoreProvider& DataStorage)
+{
+	using namespace UE::Editor::DataStorage;
+	using namespace UE::Editor::DataStorage::Queries;
+	
+	DataStorage.RegisterQuery(
+		Select(
+			TEXT("Sync world label to column"),
+			FProcessor(EQueryTickPhase::PrePhysics, DataStorage.GetQueryTickGroupName(EQueryTickGroups::SyncExternalToDataStorage))
+				.SetExecutionMode(EExecutionMode::GameThread),
+			[](const FTypedElementUObjectColumn& WorldColumn, FTypedElementLabelColumn& Label)
+			{
+				if (const UWorld* World = Cast<UWorld>(WorldColumn.Object); World != nullptr)
+				{
+					// Show the type of world in the name (e.g "MyWorld (Editor)")
+					FString WorldTypeString = LexToString(World->WorldType);
+					Label.Label = FString::Printf(TEXT("%s (%s)"), *World->GetName(), *WorldTypeString);
+				}
+			}
+		)
+		.Where()
+			.All<FWorldTag, FTypedElementSyncFromWorldTag>()
+		.Compile());
+}
+
+void UTedsWorldFactory::PostWorldInitialized(UWorld* InWorld, const UWorld::InitializationValues IVS) const
+{
+	DataStorageCompat->AddCompatibleObject(InWorld);
+}
+
+void UTedsWorldFactory::OnPreWorldFinishDestroy(UWorld* World) const
+{
+	DataStorageCompat->RemoveCompatibleObject(World);
+}
diff --git a/Engine/Plugins/Experimental/EditorDataStorageFeatures/Source/TedsEditorCompatibility/Private/World/TedsWorldFactory.h b/Engine/Plugins/Experimental/EditorDataStorageFeatures/Source/TedsEditorCompatibility/Private/World/TedsWorldFactory.h
new file mode 100644
--- /dev/null
+++ b/Engine/Plugins/Experimental/EditorDataStorageFeatures/Source/TedsEditorCompatibility/Private/World/TedsWorldFactory.h
@@ -0,0 +1,31 @@
+// Copyright Epic Games, Inc. All Rights Reserved.
+
+#pragma once
+
+#include "DataStorage/Handles.h"
+#include "Elements/Interfaces/TypedElementDataStorageFactory.h"
+#include "Engine/World.h"
+
+#include "TedsWorldFactory.generated.h"
+
+UCLASS()
+class UTedsWorldFactory : public UEditorDataStorageFactory
+{
+	GENERATED_BODY()
+
+public:
+	~UTedsWorldFactory() override = default;
+
+	virtual void RegisterTables(UE::Editor::DataStorage::ICoreProvider& DataStorage) override;
+	virtual void PreRegister(UE::Editor::DataStorage::ICoreProvider& DataStorage) override;
+	virtual void PreShutdown(UE::Editor::DataStorage::ICoreProvider& DataStorage) override;
+	virtual void RegisterQueries(UE::Editor::DataStorage::ICoreProvider& DataStorage) override;
+
+private:
+
+	void PostWorldInitialized(UWorld* InWorld, const UWorld::InitializationValues IVS) const;
+	void OnPreWorldFinishDestroy(UWorld* World) const;
+
+private:
+	UE::Editor::DataStorage::ICompatibilityProvider* DataStorageCompat;
+};
\ No newline at end of file
diff --git a/Engine/Plugins/Experimental/EditorDataStorageFeatures/Source/TedsEditorCompatibility/TedsEditorCompatibility.Build.cs b/Engine/Plugins/Experimental/EditorDataStorageFeatures/Source/TedsEditorCompatibility/TedsEditorCompatibility.Build.cs
new file mode 100644
--- /dev/null
+++ b/Engine/Plugins/Experimental/EditorDataStorageFeatures/Source/TedsEditorCompatibility/TedsEditorCompatibility.Build.cs
@@ -0,0 +1,34 @@
+// Copyright Epic Games, Inc. All Rights Reserved.
+
+using UnrealBuildTool;
+
+public class TedsEditorCompatibility : ModuleRules
+{
+	public TedsEditorCompatibility(ReadOnlyTargetRules Target) : base(Target)
+	{
+		// Enable truncation warnings in this module
+		CppCompileWarningSettings.UnsafeTypeCastWarningLevel = WarningLevel.Warning;
+
+		PCHUsage = ModuleRules.PCHUsageMode.UseExplicitOrSharedPCHs;
+
+		if (Target.bBuildEditor)
+		{
+			PublicIncludePaths.AddRange(new string[] { });
+			PrivateIncludePaths.AddRange(new string[] { });
+
+			PrivateDependencyModuleNames.AddRange(
+				new string[] 
+				{
+					"Core",
+					"CoreUObject",
+					"Engine",
+					"UnrealEd",
+					"TypedElementFramework",
+					"SlateCore",
+					"Slate",
+				});
+
+			DynamicallyLoadedModuleNames.AddRange(new string[] { });
+		}
+	}
+}
\ No newline at end of file
diff --git a/Engine/Source/Runtime/TypedElementFramework/Public/Elements/Columns/TypedElementCompatibilityColumns.h b/Engine/Source/Runtime/TypedElementFramework/Public/Elements/Columns/TypedElementCompatibilityColumns.h
--- a/Engine/Source/Runtime/TypedElementFramework/Public/Elements/Columns/TypedElementCompatibilityColumns.h
+++ b/Engine/Source/Runtime/TypedElementFramework/Public/Elements/Columns/TypedElementCompatibilityColumns.h
@@ -84,6 +84,16 @@ struct FTypedElementWorldColumn final : public FEditorDataStorageColumn
 	TWeakObjectPtr<UWorld> World;
 };
 
+/**
+* Tag to identify a row with a UWorld. If there's a FTypedElementUObjectColumn, the stored object will be
+ * a UWorld if this tag is present.
+ */
+USTRUCT(meta = (DisplayName = "World Tag"))
+struct FEditorDataStorageWorldTag final : public FEditorDataStorageTag
+{
+	GENERATED_BODY()
+};
+
 /**
  * Column that stores a reference to a level.
  */
@@ -95,6 +105,16 @@ struct FTypedElementLevelColumn final : public FEditorDataStorageColumn
 	TWeakObjectPtr<ULevel> Level;
 };
 
+/**
+* Tag to identify a row with a ULevel. If there's a FTypedElementUObjectColumn, the stored object will be
+ * a ULevel if this tag is present.
+ */
+USTRUCT(meta = (DisplayName = "World Tag"))
+struct FEditorDataStorageLevelTag final : public FEditorDataStorageTag
+{
+	GENERATED_BODY()
+};
+
 /**
  * Tag to signal that a row represents a property bag "placeholder-typed" object reference.
  * This object has an unknown base type and will generally be associated with a "property
@@ -145,4 +165,6 @@ namespace UE::Editor::DataStorage
 {
 	using FUObjectIdNameColumn = FEditorDataStorageUObjectIdNameColumn;
 	using FLevelColumn = FTypedElementLevelColumn;
+	using FLevelTag = FEditorDataStorageLevelTag;
+	using FWorldTag = FEditorDataStorageWorldTag;
 }  // namespace UE::Editor::DataStorage
\ No newline at end of file
