Teds: Add sanity checks to datastorage compat in the world and level factory

diff --git a/Engine/Plugins/Experimental/EditorDataStorageFeatures/Source/TedsEditorCompatibility/Private/Level/TedsLevelFactory.cpp b/Engine/Plugins/Experimental/EditorDataStorageFeatures/Source/TedsEditorCompatibility/Private/Level/TedsLevelFactory.cpp
--- a/Engine/Plugins/Experimental/EditorDataStorageFeatures/Source/TedsEditorCompatibility/Private/Level/TedsLevelFactory.cpp
+++ b/Engine/Plugins/Experimental/EditorDataStorageFeatures/Source/TedsEditorCompatibility/Private/Level/TedsLevelFactory.cpp
@@ -144,10 +144,16 @@ void UTedsLevelFactory::RegisterQueries(UE::Editor::DataStorage::ICoreProvider&
 
 void UTedsLevelFactory::OnLevelAddedToWorld(ULevel* InLevel, UWorld* InWorld)
 {
-	DataStorageCompat->AddCompatibleObject(InLevel);
+	if (DataStorageCompat)
+	{
+		DataStorageCompat->AddCompatibleObject(InLevel);
+	}
 }
 
 void UTedsLevelFactory::OnLevelRemovedFromWorld(ULevel* InLevel, UWorld* InWorld)
 {
-	DataStorageCompat->RemoveCompatibleObject(InLevel);
+	if (DataStorageCompat)
+	{
+		DataStorageCompat->RemoveCompatibleObject(InLevel);
+	}
 }
diff --git a/Engine/Plugins/Experimental/EditorDataStorageFeatures/Source/TedsEditorCompatibility/Private/World/TedsWorldFactory.cpp b/Engine/Plugins/Experimental/EditorDataStorageFeatures/Source/TedsEditorCompatibility/Private/World/TedsWorldFactory.cpp
--- a/Engine/Plugins/Experimental/EditorDataStorageFeatures/Source/TedsEditorCompatibility/Private/World/TedsWorldFactory.cpp
+++ b/Engine/Plugins/Experimental/EditorDataStorageFeatures/Source/TedsEditorCompatibility/Private/World/TedsWorldFactory.cpp
@@ -68,10 +68,16 @@ void UTedsWorldFactory::RegisterQueries(UE::Editor::DataStorage::ICoreProvider&
 
 void UTedsWorldFactory::PostWorldInitialized(UWorld* InWorld, const UWorld::InitializationValues IVS) const
 {
-	DataStorageCompat->AddCompatibleObject(InWorld);
+	if (DataStorageCompat)
+	{
+		DataStorageCompat->AddCompatibleObject(InWorld);
+	}
 }
 
 void UTedsWorldFactory::OnPreWorldFinishDestroy(UWorld* World) const
 {
-	DataStorageCompat->RemoveCompatibleObject(World);
+	if (DataStorageCompat)
+	{
+		DataStorageCompat->RemoveCompatibleObject(World);
+	}
 }
