Teds Outliner:
Only attempt to update label for rows that exist in the query stack. Otherwise the Outliner force creates items for them

Query Stack:

RowMonitorNode:
Only clear the AddedRows/RemovedRows when they are actually merged
When the OnRemove observer is fired, only remove rows that don't match the parent query anymore.

QueryMergeNode
Fix a bug where parent nodes would not get updated as soon as the first parent node that changed was found

diff --git a/Engine/Plugins/Experimental/EditorDataStorageFeatures/Source/TedsOutliner/Private/TedsOutlinerImpl.cpp b/Engine/Plugins/Experimental/EditorDataStorageFeatures/Source/TedsOutliner/Private/TedsOutlinerImpl.cpp
--- a/Engine/Plugins/Experimental/EditorDataStorageFeatures/Source/TedsOutliner/Private/TedsOutlinerImpl.cpp
+++ b/Engine/Plugins/Experimental/EditorDataStorageFeatures/Source/TedsOutliner/Private/TedsOutlinerImpl.cpp
@@ -1034,25 +1034,31 @@ void FTedsOutlinerImpl::PostDataStorageUpdate()
 	// Update the label for any rows that might need it
 	for (DataStorage::RowHandle Row : RowsPendingLabelUpdate)
 	{
-		// If the item already exists, it only needs an update if it passed a filter previously and does not now (or vice versa)
-		if (FSceneOutlinerTreeItemPtr ExistingItem = SceneOutliner->GetTreeItem(Row))
+		// Search is currently handled externally and not through the query stack, so if a node doesn't match the search filters it should still
+		// exist in the query stack. If it does not exist in the query stack, there's no need to check it against the search filter or add it to the
+		// Outliner
+		if (RowSortMapNode->GetRowIndex(Row) != INDEX_NONE)
 		{
-			bool bCachedFilteredFlag = ExistingItem->Flags.bIsFilteredOut;
+			// If the item already exists, it only needs an update if it passed a filter previously and does not now (or vice versa)
+			if (FSceneOutlinerTreeItemPtr ExistingItem = SceneOutliner->GetTreeItem(Row))
+			{
+				bool bCachedFilteredFlag = ExistingItem->Flags.bIsFilteredOut;
 
-			// This implicitly calls into the data storage to get the label of the row and check against the search query
-			ExistingItem->Flags.bIsFilteredOut = !SceneOutliner->PassesAllFilters(ExistingItem);
+				// This implicitly calls into the data storage to get the label of the row and check against the search query
+				ExistingItem->Flags.bIsFilteredOut = !SceneOutliner->PassesAllFilters(ExistingItem);
 
-			if (bCachedFilteredFlag != ExistingItem->Flags.bIsFilteredOut)
+				if (bCachedFilteredFlag != ExistingItem->Flags.bIsFilteredOut)
+				{
+					SceneOutliner->OnItemLabelChanged(ExistingItem, false);
+				}
+			}
+			// If the item doesn't exist, create a dummy item to see if it would match the current search/filter queries and should be actually added
+			else if (FSceneOutlinerTreeItemPtr PotentialItem = SceneOutlinerMode->CreateItemFor<FTedsOutlinerTreeItem>(
+					FTedsOutlinerTreeItem(Row, AsShared()), true))
 			{
-				SceneOutliner->OnItemLabelChanged(ExistingItem, false);
+				SceneOutliner->OnItemLabelChanged(PotentialItem, false);
 			}
 		}
-		// If the item doesn't exist, create a dummy item to see if it would match the current search/filter queries and should be actually added
-		else if (FSceneOutlinerTreeItemPtr PotentialItem = SceneOutlinerMode->CreateItemFor<FTedsOutlinerTreeItem>(
-				FTedsOutlinerTreeItem(Row, AsShared()), true))
-		{
-			SceneOutliner->OnItemLabelChanged(PotentialItem, false);
-		}
 	}
 	
 	RowsPendingLabelUpdate.Empty();
diff --git a/Engine/Plugins/Experimental/EditorDataStorageFeatures/Source/TedsQueryStack/Private/TedsQueryMergeNode.cpp b/Engine/Plugins/Experimental/EditorDataStorageFeatures/Source/TedsQueryStack/Private/TedsQueryMergeNode.cpp
--- a/Engine/Plugins/Experimental/EditorDataStorageFeatures/Source/TedsQueryStack/Private/TedsQueryMergeNode.cpp
+++ b/Engine/Plugins/Experimental/EditorDataStorageFeatures/Source/TedsQueryStack/Private/TedsQueryMergeNode.cpp
@@ -42,7 +42,6 @@ namespace UE::Editor::DataStorage::QueryStack
 			{
 				bRebuild = true;
 				ParentInfo.Revision = ParentRevision;
-				break;
 			}
 		}
 
diff --git a/Engine/Plugins/Experimental/EditorDataStorageFeatures/Source/TedsQueryStack/Private/TedsRowMonitorNode.cpp b/Engine/Plugins/Experimental/EditorDataStorageFeatures/Source/TedsQueryStack/Private/TedsRowMonitorNode.cpp
--- a/Engine/Plugins/Experimental/EditorDataStorageFeatures/Source/TedsQueryStack/Private/TedsRowMonitorNode.cpp
+++ b/Engine/Plugins/Experimental/EditorDataStorageFeatures/Source/TedsQueryStack/Private/TedsRowMonitorNode.cpp
@@ -111,6 +111,8 @@ namespace UE::Editor::DataStorage::QueryStack
 			RemovedRows.Sort();
 			RemovedRows.MakeUnique();
 
+			ResolveRemovedRows();
+
 			FRowHandleArray& TargetRows = ParentNode->GetMutableRows();
 			TargetRows.SortedMerge(AddedRows);
 			TargetRows.Remove(RemovedRows.GetRows());
@@ -124,11 +126,12 @@ namespace UE::Editor::DataStorage::QueryStack
 			{
 				RowsRemovedEvent.Broadcast(RemovedRows.GetRows());
 			}
+
+			AddedRows.Empty();
+			RemovedRows.Empty();
 			
 			Revision++;
 		}
-		AddedRows.Empty();
-		RemovedRows.Empty();
 	}
 
 	FRowHandleArrayView FRowMonitorNode::GetRows() const
@@ -301,4 +304,39 @@ namespace UE::Editor::DataStorage::QueryStack
 			}
 		}
 	}
+
+	void FRowMonitorNode::ResolveRemovedRows()
+	{
+		// OnRemove observers are fired when any column from our conditions are removed from the row, however this doesn't necessarily mean
+		// that the row should be removed as it could still be matching the query. For instance, if the Conditions are (ColumnA || ColumnB), the row
+		// contains both columns and only ColumnB is removed the row still matches the original query and shouldn't be removed from the parent node.
+		const FQueryDescription& Query = Storage.GetQueryDescription(QueryNode->GetQuery());
+
+		// RemovedRows is guaranteed to be sorted when this function is called, but we'll add a sanity check to sort it anyways in case that changes
+		if (!RemovedRows.IsSorted())
+		{
+			RemovedRows.Sort();
+		}
+
+		// We don't currently have a way to detect if a row matches a query using the old syntax, so this logic only applies to the FConditions
+		// based syntax. Which means that you could get an incorrect output if you use this node with a complex .Where() syntax.
+		if (const Queries::FConditions* Conditions = Query.Conditions.GetPtrOrNull())
+		{
+			FRowHandleArray RowsToNotRemove;
+			RowsToNotRemove.Reserve(RemovedRows.Num());
+			
+			for (RowHandle Row : RemovedRows.GetRows())
+			{
+				// If the row still matches the conditions, don't remove it.
+				if (Storage.MatchesColumns(Row, *Conditions))
+				{
+					// Since RemovedRows is sorted, iterating through it and adding rows keep that property
+					static constexpr FRowHandleArrayView::EFlags Flags(FRowHandleArrayView::EFlags::IsSorted | FRowHandleArrayView::EFlags::IsUnique );
+					RowsToNotRemove.Add(Row, Flags);
+				}
+			}
+			
+			RemovedRows.Remove(RowsToNotRemove.GetRows());
+		}
+	}
 } // namespace UE::Editor::DataStorage::QueryStack
diff --git a/Engine/Plugins/Experimental/EditorDataStorageFeatures/Source/TedsQueryStack/Public/TedsRowMonitorNode.h b/Engine/Plugins/Experimental/EditorDataStorageFeatures/Source/TedsQueryStack/Public/TedsRowMonitorNode.h
--- a/Engine/Plugins/Experimental/EditorDataStorageFeatures/Source/TedsQueryStack/Public/TedsRowMonitorNode.h
+++ b/Engine/Plugins/Experimental/EditorDataStorageFeatures/Source/TedsQueryStack/Public/TedsRowMonitorNode.h
@@ -56,6 +56,7 @@ namespace UE::Editor::DataStorage::QueryStack
 	private:
 		void UpdateColumnsFromQuery();
 		void UpdateMonitoredColumns();
+		void ResolveRemovedRows();
 
 		FRowHandleArray AddedRows;
 		FRowHandleArray RemovedRows;
