Teds Hierarchy Viewer:
Add a new hierarchy data interface to allow the widget to operate on legacy hierarchies or the new TEDS hierarchy handles
Add a new generic node to collect top level rows for the hierarchy viewer to operate on (that can be re-used for the most common case)

diff --git a/Engine/Plugins/Experimental/EditorDataStorageFeatures/Source/TedsActorCompatibility/Private/ActorComponentDebugHierarchyWidget/SActorComponentTreeWidget.cpp b/Engine/Plugins/Experimental/EditorDataStorageFeatures/Source/TedsActorCompatibility/Private/ActorComponentDebugHierarchyWidget/SActorComponentTreeWidget.cpp
--- a/Engine/Plugins/Experimental/EditorDataStorageFeatures/Source/TedsActorCompatibility/Private/ActorComponentDebugHierarchyWidget/SActorComponentTreeWidget.cpp
+++ b/Engine/Plugins/Experimental/EditorDataStorageFeatures/Source/TedsActorCompatibility/Private/ActorComponentDebugHierarchyWidget/SActorComponentTreeWidget.cpp
@@ -39,11 +39,12 @@ void SActorComponentTreeWidget::Construct(
 		FHierarchyRowNode::ESyncFlags::IncrementWhenDifferent);
 	
 	TSharedPtr<QueryStack::IRowNode> HierarchyRowNodePtr = HierarchyRowNode.ToSharedPtr();
+
+	TSharedPtr<FHierarchyViewerLegacyData> HierarchyData = MakeShared<FHierarchyViewerLegacyData>(HierarchyStruct, MoveTemp(GetParentRowFromColumnPredicate));
 	
 	TSharedRef<SHierarchyViewer> HierarchyViewer =
 		SNew(SHierarchyViewer,
-			HierarchyStruct,
-			MoveTemp(GetParentRowFromColumnPredicate))
+			MoveTemp(HierarchyData))
 			.AllNodeProvider(HierarchyRowNode)
 			.TopLevelNodeProvider(TopLevelRowProvider)
 			.EmptyRowsMessage(LOCTEXT("EmptyRowsMessage", "The provided query has no results."))
diff --git a/Engine/Plugins/Experimental/EditorDataStorageFeatures/Source/TedsDebugger/Private/Widgets/QueryEditor/TedsQueryEditorResultsView.cpp b/Engine/Plugins/Experimental/EditorDataStorageFeatures/Source/TedsDebugger/Private/Widgets/QueryEditor/TedsQueryEditorResultsView.cpp
--- a/Engine/Plugins/Experimental/EditorDataStorageFeatures/Source/TedsDebugger/Private/Widgets/QueryEditor/TedsQueryEditorResultsView.cpp
+++ b/Engine/Plugins/Experimental/EditorDataStorageFeatures/Source/TedsDebugger/Private/Widgets/QueryEditor/TedsQueryEditorResultsView.cpp
@@ -108,7 +108,7 @@ namespace UE::Editor::DataStorage::Debug::QueryEditor
 
 		if(RowHandleColumn)
 		{
-			TableViewer->AddCustomColumn(RowHandleColumn.ToSharedRef());
+			TableViewer->AddCustomRowWidget(RowHandleColumn.ToSharedRef());
 		}
 	}
 
@@ -150,7 +150,7 @@ namespace UE::Editor::DataStorage::Debug::QueryEditor
 				// Since setting the columns clears all columns, we are going to re-add the custom column back
 				if(RowHandleColumn)
 				{
-					TableViewer->AddCustomColumn(RowHandleColumn.ToSharedRef());
+					TableViewer->AddCustomRowWidget(RowHandleColumn.ToSharedRef());
 				}
 		
 				// Mass doesn't like empty queries, so we only set it if there are actual conditions
diff --git a/Engine/Plugins/Experimental/EditorDataStorageFeatures/Source/TedsEverythingPicker/Private/Context/TedsPickerContextUtil.cpp b/Engine/Plugins/Experimental/EditorDataStorageFeatures/Source/TedsEverythingPicker/Private/Context/TedsPickerContextUtil.cpp
--- a/Engine/Plugins/Experimental/EditorDataStorageFeatures/Source/TedsEverythingPicker/Private/Context/TedsPickerContextUtil.cpp
+++ b/Engine/Plugins/Experimental/EditorDataStorageFeatures/Source/TedsEverythingPicker/Private/Context/TedsPickerContextUtil.cpp
@@ -11,6 +11,8 @@
 #include "Elements/Columns/TedsTypeInfoColumns.h"
 #include "Elements/Framework/TypedElementQueryContext.h"
 
+#include "QueryStack/TopLevelRowsNode.h"
+
 #include "Widgets/STedsHierarchyViewer.h"
 #include "Widgets/STedsTableViewer.h"
 
@@ -27,79 +29,6 @@
 
 namespace UE::Editor::DataStorage::Picker
 {
-	// A custom query stack node to get top level rows of the given rows to display in the hierarchy viewer
-	class FGatherTopLevelRowsNode : public UE::Editor::DataStorage::QueryStack::IRowNode
-	{
-	public:
-		
-		FGatherTopLevelRowsNode(const UE::Editor::DataStorage::ICoreProvider* InStorage, TSharedPtr<IRowNode> InParent)
-			: Storage(InStorage)
-			, Parent(MoveTemp(InParent))
-		{
-			UpdateRows();
-		}
-		
-		virtual ~FGatherTopLevelRowsNode() override = default;
-		
-		/** Retrieve access to the rows used by this node. */
-		virtual UE::Editor::DataStorage::FRowHandleArrayView GetRows() const override
-		{
-			return Rows.GetRows();
-		}
-		
-		virtual UE::Editor::DataStorage::FRowHandleArray& GetMutableRows() override
-		{
-			return Rows;
-		}
-		
-		virtual RevisionId GetRevision() const override
-		{
-			return Revision;
-		}
-		
-		virtual void Update() override
-		{
-			Parent->Update();
-			if (Parent->GetRevision() != CachedParentRevision)
-			{
-				UpdateRows();
-				CachedParentRevision = Parent->GetRevision();
-				Revision++;
-			}
-		}
-
-		void UpdateRows()
-		{
-			Rows.Empty();
-
-			UE::Editor::DataStorage::FRowHandleArray RowsToCheck(Parent->GetRows());
-			RowsToCheck.Sort();
-			for (UE::Editor::DataStorage::RowHandle Row : Parent->GetRows())
-			{
-				// If the row doesn't have a parent, it is a top level node and can be added
-				if (const FTableRowParentColumn* ParentColumn = Storage->GetColumn<FTableRowParentColumn>(Row))
-				{
-					// If the filtered list contains the parent of the row, it has a visible parent and thus is not a top-level node
-					if(!RowsToCheck.Contains(ParentColumn->Parent))
-					{
-						Rows.Add(Row);
-					}
-				}
-				else
-				{
-					Rows.Add(Row);
-				}
-			}
-		}
-
-	protected:
-		RevisionId Revision = 0;
-		const UE::Editor::DataStorage::ICoreProvider* Storage;
-		TSharedPtr<IRowNode> Parent;
-		UE::Editor::DataStorage::FRowHandleArray Rows;
-		RevisionId CachedParentRevision = TNumericLimits<RevisionId>::Max();
-	};
-
 	void SObjectReferenceContextView::Construct(const FArguments& InArgs)
 	{
 		using namespace UE::Editor::DataStorage::Queries;
@@ -138,11 +67,7 @@ namespace UE::Editor::DataStorage::Picker
 		{
 			SearchNode = ReferenceQueryResultsNode.ToSharedPtr();
 		}
-
-		// Node that filters out the ReferenceQueryResultsNode or SearchNode (if enabled) to generate all the top level nodes
-		TSharedRef<FGatherTopLevelRowsNode> TopLevelNode =
-			MakeShared<FGatherTopLevelRowsNode>(Storage, SearchNode);
-
+		
 		// Using the Teds Outliner Purpose for matching customization for now
 		IUiProvider::FPurposeID PurposeId = IUiProvider::FPurposeInfo("SceneOutliner", "RowLabel", NAME_None).GeneratePurposeID();
 
@@ -150,9 +75,15 @@ namespace UE::Editor::DataStorage::Picker
 		{
 			const FTableRowParentColumn* ParentColumn = static_cast<const FTableRowParentColumn*>(Column);
 			return ParentColumn->Parent;
+			
 		};
 
-		TSharedPtr<SHierarchyViewer> ReferenceView = SNew(SHierarchyViewer, FTableRowParentColumn::StaticStruct(), ParentExtractionFunction)
+		TSharedPtr<FHierarchyViewerLegacyData> HierarchyData = MakeShared<FHierarchyViewerLegacyData>(FTableRowParentColumn::StaticStruct(), MoveTemp(ParentExtractionFunction));
+
+		// Node that filters out the ReferenceQueryResultsNode or SearchNode (if enabled) to generate all the top level nodes
+		TSharedRef<QueryStack::FTopLevelRowsNode> TopLevelNode = MakeShared<QueryStack::FTopLevelRowsNode>(Storage, HierarchyData, SearchNode);
+
+		TSharedPtr<SHierarchyViewer> ReferenceView = SNew(SHierarchyViewer, HierarchyData)
 			.AllNodeProvider(SearchNode)
 			.TopLevelNodeProvider(TopLevelNode)
 			.Columns({ FTypedElementLabelColumn::StaticStruct(), FTypedElementClassTypeInfoColumn::StaticStruct() })
@@ -282,5 +213,4 @@ namespace UE::Editor::DataStorage::Picker
 			];
 	}
 }
-
 #undef LOCTEXT_NAMESPACE // TedsPickerContextUtil
\ No newline at end of file
diff --git a/Engine/Plugins/Experimental/EditorDataStorageFeatures/Source/TedsTableViewer/Private/HierarchyViewerIntefaces.cpp b/Engine/Plugins/Experimental/EditorDataStorageFeatures/Source/TedsTableViewer/Private/HierarchyViewerIntefaces.cpp
new file mode 100644
--- /dev/null
+++ b/Engine/Plugins/Experimental/EditorDataStorageFeatures/Source/TedsTableViewer/Private/HierarchyViewerIntefaces.cpp
@@ -0,0 +1,35 @@
+// Copyright Epic Games, Inc. All Rights Reserved.
+
+#include "HierarchyViewerIntefaces.h"
+
+namespace UE::Editor::DataStorage
+{
+	FHierarchyViewerData::FHierarchyViewerData(FHierarchyHandle InHierarchyHandle)
+		: HierarchyHandle(InHierarchyHandle)
+	{
+	}
+
+	RowHandle FHierarchyViewerData::GetParent(const ICoreProvider& Storage, RowHandle InRow) const
+	{
+		return Storage.GetParentRow(HierarchyHandle, InRow);
+	}
+
+	FHierarchyViewerLegacyData::FHierarchyViewerLegacyData(const UScriptStruct* InHierarchyDataColumnType,
+	                                                       TFunction<RowHandle(const void*, const UScriptStruct*)> InGetRowHandleFn)
+		: HierarchyDataColumnType(InHierarchyDataColumnType)
+		, GetRowHandleFn(InGetRowHandleFn)
+	{
+		checkf(HierarchyDataColumnType != nullptr, TEXT("Nullptr ParentColumnType"));
+	}
+
+	RowHandle FHierarchyViewerLegacyData::GetParent(const ICoreProvider& Storage, RowHandle InRow) const
+	{
+		if (const void* ParentColumn = Storage.GetColumnData(InRow, HierarchyDataColumnType))
+		{
+			RowHandle ParentRow = GetRowHandleFn(ParentColumn, HierarchyDataColumnType);
+			return ParentRow;
+		}
+
+		return InvalidRowHandle;
+	}
+}
\ No newline at end of file
diff --git a/Engine/Plugins/Experimental/EditorDataStorageFeatures/Source/TedsTableViewer/Private/QueryStack/TopLevelRowsNode.cpp b/Engine/Plugins/Experimental/EditorDataStorageFeatures/Source/TedsTableViewer/Private/QueryStack/TopLevelRowsNode.cpp
new file mode 100644
--- /dev/null
+++ b/Engine/Plugins/Experimental/EditorDataStorageFeatures/Source/TedsTableViewer/Private/QueryStack/TopLevelRowsNode.cpp
@@ -0,0 +1,69 @@
+// Copyright Epic Games, Inc. All Rights Reserved.
+
+#include "QueryStack/TopLevelRowsNode.h"
+
+#include "HierarchyViewerIntefaces.h"
+
+namespace UE::Editor::DataStorage::QueryStack
+{
+	FTopLevelRowsNode::FTopLevelRowsNode(const ICoreProvider* InStorage, TSharedPtr<IHierarchyViewerDataInterface> InHierarchyData, TSharedPtr<IRowNode> InParent)
+			: Storage(InStorage)
+			, HierarchyData(MoveTemp(InHierarchyData))
+			, Parent(MoveTemp(InParent))
+	{
+		checkf(HierarchyData, TEXT("This node requires a valid hierarchy to view"))
+		UpdateRows();
+	}
+
+	FRowHandleArrayView FTopLevelRowsNode::GetRows() const
+	{
+		return Rows.GetRows();
+	}
+
+	FRowHandleArray& FTopLevelRowsNode::GetMutableRows()
+	{
+		return Rows;
+	}
+
+	INode::RevisionId FTopLevelRowsNode::GetRevision() const
+	{
+		return Revision;
+	}
+
+	void FTopLevelRowsNode::Update()
+	{
+		Parent->Update();
+		if (Parent->GetRevision() != CachedParentRevision)
+		{
+			UpdateRows();
+			CachedParentRevision = Parent->GetRevision();
+			Revision++;
+		}
+	}
+
+	void FTopLevelRowsNode::UpdateRows()
+	{
+		Rows.Empty();
+		FRowHandleArray RowsToCheck(Parent->GetRows());
+		RowsToCheck.Sort();
+		
+		for (RowHandle Row : Parent->GetRows())
+		{
+			RowHandle ParentRow = HierarchyData->GetParent(*Storage, Row);
+			
+			// If the row doesn't have a parent, it is a top level node and can be added
+			if (Storage->IsRowAvailable(ParentRow))
+			{
+				// If the filtered list contains the parent of the row, it has a visible parent and thus is not a top-level node
+				if(!RowsToCheck.Contains(ParentRow))
+				{
+					Rows.Add(Row);
+				}
+			}
+			else
+			{
+				Rows.Add(Row);
+			}
+		}
+	}
+}
diff --git a/Engine/Plugins/Experimental/EditorDataStorageFeatures/Source/TedsTableViewer/Private/TedsTableViewerModel.cpp b/Engine/Plugins/Experimental/EditorDataStorageFeatures/Source/TedsTableViewer/Private/TedsTableViewerModel.cpp
--- a/Engine/Plugins/Experimental/EditorDataStorageFeatures/Source/TedsTableViewer/Private/TedsTableViewerModel.cpp
+++ b/Engine/Plugins/Experimental/EditorDataStorageFeatures/Source/TedsTableViewer/Private/TedsTableViewerModel.cpp
@@ -162,7 +162,7 @@ namespace UE::Editor::DataStorage
 		GenerateColumns();
 	}
 
-	void FTedsTableViewerModel::AddCustomColumn(const TSharedRef<FTedsTableViewerColumn>& InColumn)
+	void FTedsTableViewerModel::AddCustomRowWidget(const TSharedRef<FTedsTableViewerColumn>& InColumn)
 	{
 		// Table Viewer TODO: We should allow users to specify sort order using a TEDS column on the UI row, but for now we put any custom
 		// columns on the front
diff --git a/Engine/Plugins/Experimental/EditorDataStorageFeatures/Source/TedsTableViewer/Private/Widgets/STedsHierarchyViewer.cpp b/Engine/Plugins/Experimental/EditorDataStorageFeatures/Source/TedsTableViewer/Private/Widgets/STedsHierarchyViewer.cpp
--- a/Engine/Plugins/Experimental/EditorDataStorageFeatures/Source/TedsTableViewer/Private/Widgets/STedsHierarchyViewer.cpp
+++ b/Engine/Plugins/Experimental/EditorDataStorageFeatures/Source/TedsTableViewer/Private/Widgets/STedsHierarchyViewer.cpp
@@ -20,12 +20,9 @@ namespace UE::Editor::DataStorage
 {
 	void SHierarchyViewer::Construct(
 		const FArguments& InArgs,
-		const UScriptStruct* InHierarchyDataColumnType,
-		TFunction<RowHandle(const void*, const UScriptStruct*)> InGetRowHandleFn)
+		TSharedPtr<IHierarchyViewerDataInterface> InHierarchyInterface)
 	{
-		checkf(InHierarchyDataColumnType != nullptr, TEXT("Nullptr ParentColumnType"));
-		HierarchyColumnType = InHierarchyDataColumnType;
-		GetParentFromColumnFn = MoveTemp(InGetRowHandleFn);
+		HierarchyInterface = MoveTemp(InHierarchyInterface);
 		OnSelectionChanged = InArgs._OnSelectionChanged;
 		EmptyRowsMessage = InArgs._EmptyRowsMessage;
 		DefaultExpansionState = InArgs._DefaultExpansionState;
@@ -125,13 +122,14 @@ namespace UE::Editor::DataStorage
 
 	TableViewerItemPtr SHierarchyViewer::GetParentRow(TableViewerItemPtr InItem) const
 	{
-		if (const void* ParentColumn = Model->GetDataStorageInterface()->GetColumnData(InItem, HierarchyColumnType))
+		// To work around unreliable query stack update timings, in case the widget queries parents before the top level rows node has updated
+		// but after the hierarchies themselves have updated since HierarchyInterface directly queries TEDS. 
+		if (TopLevelRows.Contains(InItem) || !HierarchyInterface)
 		{
-			RowHandle ParentRow = GetParentFromColumnFn(ParentColumn, HierarchyColumnType);
-			return FTedsRowHandle(ParentRow);
+			return FTedsRowHandle(InvalidRowHandle);
 		}
-
-		return FTedsRowHandle(InvalidRowHandle);
+		
+		return FTedsRowHandle(HierarchyInterface->GetParent(*Model->GetDataStorageInterface(), InItem));
 	}
 
 	void SHierarchyViewer::CreateInternalWidget()
@@ -194,9 +192,9 @@ namespace UE::Editor::DataStorage
 		RefreshColumnWidgets();
 	}
 
-	void SHierarchyViewer::AddCustomColumn(const TSharedRef<FTedsTableViewerColumn>& InColumn)
+	void SHierarchyViewer::AddCustomRowWidget(const TSharedRef<FTedsTableViewerColumn>& InColumn)
 	{
-		Model->AddCustomColumn(InColumn);
+		Model->AddCustomRowWidget(InColumn);
 		RefreshColumnWidgets();
 	}
 
@@ -259,6 +257,7 @@ namespace UE::Editor::DataStorage
 		TreeView->SetItemExpansion(InItem, DefaultExpansionState.Get() == EExpansionState::Expanded);
 
 		return SNew(SHierarchyViewerRow, OwnerTable, Model.ToSharedRef())
+				.Padding(FMargin(0))
 				.Item(InItem);
 	}
 } // namespace UE::Editor::DataStorage
diff --git a/Engine/Plugins/Experimental/EditorDataStorageFeatures/Source/TedsTableViewer/Private/Widgets/STedsTableViewer.cpp b/Engine/Plugins/Experimental/EditorDataStorageFeatures/Source/TedsTableViewer/Private/Widgets/STedsTableViewer.cpp
--- a/Engine/Plugins/Experimental/EditorDataStorageFeatures/Source/TedsTableViewer/Private/Widgets/STedsTableViewer.cpp
+++ b/Engine/Plugins/Experimental/EditorDataStorageFeatures/Source/TedsTableViewer/Private/Widgets/STedsTableViewer.cpp
@@ -166,9 +166,9 @@ namespace UE::Editor::DataStorage
 		RefreshColumnWidgets();
 	}
 
-	void STedsTableViewer::AddCustomColumn(const TSharedRef<FTedsTableViewerColumn>& InColumn)
+	void STedsTableViewer::AddCustomRowWidget(const TSharedRef<FTedsTableViewerColumn>& InColumn)
 	{
-		Model->AddCustomColumn(InColumn);
+		Model->AddCustomRowWidget(InColumn);
 		RefreshColumnWidgets();
 	}
 
diff --git a/Engine/Plugins/Experimental/EditorDataStorageFeatures/Source/TedsTableViewer/Private/Widgets/STedsTileViewer.cpp b/Engine/Plugins/Experimental/EditorDataStorageFeatures/Source/TedsTableViewer/Private/Widgets/STedsTileViewer.cpp
--- a/Engine/Plugins/Experimental/EditorDataStorageFeatures/Source/TedsTableViewer/Private/Widgets/STedsTileViewer.cpp
+++ b/Engine/Plugins/Experimental/EditorDataStorageFeatures/Source/TedsTableViewer/Private/Widgets/STedsTileViewer.cpp
@@ -83,6 +83,9 @@ namespace UE::Editor::DataStorage
 
 	void STedsTileViewer::CreateTileWidgetConstructor()
 	{
+		// Empty the Model's columns
+		Model->SetColumns(TArray<TWeakObjectPtr<const UScriptStruct>>());
+
 		// We're going to use a custom column to represent our tile so we can use FTedsTableViewerModel with STileView which doesn't have the concept
 		// of columns
 		TSharedPtr<FTedsTableViewerColumn> Column;
@@ -111,7 +114,7 @@ namespace UE::Editor::DataStorage
 
 		if (Column)
 		{
-			Model->AddCustomColumn(Column.ToSharedRef());
+			Model->AddCustomRowWidget(Column.ToSharedRef());
 		}
 	}
 
@@ -165,6 +168,19 @@ namespace UE::Editor::DataStorage
 		return IsSelected(InRow) && TileView->GetNumItemsSelected() == 1;
 	}
 
+	void STedsTileViewer::SetColumns(const TArray<TWeakObjectPtr<const UScriptStruct>>& InColumns)
+	{
+		// We aren't using columns in the traditional way, so all we have to do is update our columns list and re-create the tile constructor
+		Columns = InColumns;
+		CreateTileWidgetConstructor();
+	}
+
+	void STedsTileViewer::AddCustomRowWidget(const TSharedRef<FTedsTableViewerColumn>& InColumn)
+	{
+		// The Tile Viewer doesn't have traditional columns, so adding a custom column is ambiguous (where should the column be placed?)
+		ensureMsgf(false, TEXT("Adding custom columns is currently not supported in the Tile Viewer!"));
+	}
+
 	TSharedRef<ITableRow> STedsTileViewer::MakeTileWidget(TableViewerItemPtr InItem,
 	                                                      const TSharedRef<STableViewBase>& OwnerTable) const
 	{
diff --git a/Engine/Plugins/Experimental/EditorDataStorageFeatures/Source/TedsTableViewer/Public/HierarchyViewerIntefaces.h b/Engine/Plugins/Experimental/EditorDataStorageFeatures/Source/TedsTableViewer/Public/HierarchyViewerIntefaces.h
new file mode 100644
--- /dev/null
+++ b/Engine/Plugins/Experimental/EditorDataStorageFeatures/Source/TedsTableViewer/Public/HierarchyViewerIntefaces.h
@@ -0,0 +1,45 @@
+// Copyright Epic Games, Inc. All Rights Reserved.
+
+#pragma once
+
+#include "DataStorage/Handles.h"
+#include "Elements/Interfaces/TypedElementDataStorageInterface.h"
+
+#define UE_API TEDSTABLEVIEWER_API
+
+namespace UE::Editor::DataStorage
+{
+	class SHierarchyViewer;
+	
+	// Data interface for the hierarchy viewer to extract hierarchy information from a row
+	class IHierarchyViewerDataInterface
+	{
+	public:
+		UE_API virtual ~IHierarchyViewerDataInterface() = default;
+		UE_API virtual RowHandle GetParent(const ICoreProvider& Storage, RowHandle InRow) const = 0;
+	};
+
+	// Data interface for the hierarchy viewer that operates on a single FHierarchyHandle 
+	class FHierarchyViewerData : public IHierarchyViewerDataInterface
+	{
+	public:
+		UE_API explicit FHierarchyViewerData(FHierarchyHandle InHierarchyHandle);
+		UE_API virtual RowHandle GetParent(const ICoreProvider& Storage, RowHandle InRow) const override;
+	private:
+		FHierarchyHandle HierarchyHandle;
+	};
+
+	// Legacy data interface for the hierarchy viewer that supports arbitrary hierarchy column types alongside a function to extract the row handle from the column data
+	class FHierarchyViewerLegacyData : public IHierarchyViewerDataInterface
+	{
+	public:
+		UE_API FHierarchyViewerLegacyData(const UScriptStruct* InHierarchyDataColumnType, TFunction<RowHandle(const void*, const UScriptStruct*)> InGetRowHandleFn);
+		UE_API virtual RowHandle GetParent(const ICoreProvider& Storage, RowHandle InRow) const override;
+		
+	protected:
+		const UScriptStruct* HierarchyDataColumnType;
+		TFunction<RowHandle(const void*, const UScriptStruct*)> GetRowHandleFn;
+	};
+}
+
+#undef UE_API
\ No newline at end of file
diff --git a/Engine/Plugins/Experimental/EditorDataStorageFeatures/Source/TedsTableViewer/Public/ITedsTableViewer.h b/Engine/Plugins/Experimental/EditorDataStorageFeatures/Source/TedsTableViewer/Public/ITedsTableViewer.h
--- a/Engine/Plugins/Experimental/EditorDataStorageFeatures/Source/TedsTableViewer/Public/ITedsTableViewer.h
+++ b/Engine/Plugins/Experimental/EditorDataStorageFeatures/Source/TedsTableViewer/Public/ITedsTableViewer.h
@@ -8,6 +8,8 @@ namespace ESelectInfo { enum Type : int; }
 
 namespace UE::Editor::DataStorage
 {
+	class FTedsTableViewerColumn;
+
 	// Interface for all types of TEDS table viewers
     class ITableViewer
     {
@@ -36,6 +38,13 @@ namespace UE::Editor::DataStorage
 
     	// Get if a specific row is selected exclusively
     	virtual bool IsSelectedExclusively(RowHandle InRow) const = 0;
+
+    	// Add a custom per-row widget to the table viewer (that doesn't necessarily map to a TEDS column)
+    	// This usually means a new column for column-based views
+    	virtual void AddCustomRowWidget(const TSharedRef<FTedsTableViewerColumn>& InRowWidgetFactory) = 0;
+
+    	// Clear the current list of columns being displayed and set it to the given list
+    	virtual void SetColumns(const TArray<TWeakObjectPtr<const UScriptStruct>>& Columns) = 0;
    
     	virtual ~ITableViewer() = default;
     };
diff --git a/Engine/Plugins/Experimental/EditorDataStorageFeatures/Source/TedsTableViewer/Public/QueryStack/TopLevelRowsNode.h b/Engine/Plugins/Experimental/EditorDataStorageFeatures/Source/TedsTableViewer/Public/QueryStack/TopLevelRowsNode.h
new file mode 100644
--- /dev/null
+++ b/Engine/Plugins/Experimental/EditorDataStorageFeatures/Source/TedsTableViewer/Public/QueryStack/TopLevelRowsNode.h
@@ -0,0 +1,46 @@
+// Copyright Epic Games, Inc. All Rights Reserved.
+
+#pragma once
+
+#include "TedsQueryStackInterfaces.h"
+#include "Elements/Interfaces/TypedElementDataStorageInterface.h"
+
+#define UE_API TEDSTABLEVIEWER_API
+
+namespace UE::Editor::DataStorage
+{
+	class IHierarchyViewerDataInterface;
+}
+
+namespace UE::Editor::DataStorage::QueryStack
+{
+	// A custom query stack node to get top level rows of the given rows to display in the hierarchy viewer
+	// This node considers any rows that don't have a parent or that have a parent that is not present in the parent node as "top level"
+	// This node can be a bit expensive on large lists since it does a sort + binary search on the parent rows, consider using a custom node to collect
+	// top level rows specific to the use case if necessary
+	class FTopLevelRowsNode : public IRowNode
+	{
+	public:
+		
+		UE_API FTopLevelRowsNode(const ICoreProvider* InStorage, TSharedPtr<IHierarchyViewerDataInterface> InHierarchyData, TSharedPtr<IRowNode> InParent);
+		UE_API virtual ~FTopLevelRowsNode() override = default;
+		
+		/** Retrieve access to the rows used by this node. */
+		virtual FRowHandleArrayView GetRows() const override;
+		virtual FRowHandleArray& GetMutableRows() override;
+		virtual RevisionId GetRevision() const override;
+		virtual void Update() override;
+
+	protected:
+		void UpdateRows();
+	protected:
+		RevisionId Revision = 0;
+		const ICoreProvider* Storage = nullptr;
+		TSharedPtr<IHierarchyViewerDataInterface> HierarchyData;
+		TSharedPtr<IRowNode> Parent;
+		FRowHandleArray Rows;
+		RevisionId CachedParentRevision = TNumericLimits<RevisionId>::Max();
+	};
+}
+
+#undef UE_API
\ No newline at end of file
diff --git a/Engine/Plugins/Experimental/EditorDataStorageFeatures/Source/TedsTableViewer/Public/TedsTableViewerModel.h b/Engine/Plugins/Experimental/EditorDataStorageFeatures/Source/TedsTableViewer/Public/TedsTableViewerModel.h
--- a/Engine/Plugins/Experimental/EditorDataStorageFeatures/Source/TedsTableViewer/Public/TedsTableViewerModel.h
+++ b/Engine/Plugins/Experimental/EditorDataStorageFeatures/Source/TedsTableViewer/Public/TedsTableViewerModel.h
@@ -70,7 +70,7 @@ namespace UE::Editor::DataStorage
 		TEDSTABLEVIEWER_API void SetColumns(const TArray<TWeakObjectPtr<const UScriptStruct>>& InColumns);
 
 		// Add a custom column to display in the table viewer, that doesn't necessarily map to a Teds column
-		TEDSTABLEVIEWER_API void AddCustomColumn(const TSharedRef<FTedsTableViewerColumn>& InColumn);
+		TEDSTABLEVIEWER_API void AddCustomRowWidget(const TSharedRef<FTedsTableViewerColumn>& InColumn);
 
 		TEDSTABLEVIEWER_API ICoreProvider* GetDataStorageInterface() const;
 
diff --git a/Engine/Plugins/Experimental/EditorDataStorageFeatures/Source/TedsTableViewer/Public/Widgets/STedsHierarchyViewer.h b/Engine/Plugins/Experimental/EditorDataStorageFeatures/Source/TedsTableViewer/Public/Widgets/STedsHierarchyViewer.h
--- a/Engine/Plugins/Experimental/EditorDataStorageFeatures/Source/TedsTableViewer/Public/Widgets/STedsHierarchyViewer.h
+++ b/Engine/Plugins/Experimental/EditorDataStorageFeatures/Source/TedsTableViewer/Public/Widgets/STedsHierarchyViewer.h
@@ -1,6 +1,7 @@
 // Copyright Epic Games, Inc. All Rights Reserved.
 #pragma once
 
+#include "HierarchyViewerIntefaces.h"
 #include "ITedsTableViewer.h"
 #include "DataStorage/Handles.h"
 #include "TedsRowFilterNode.h"
@@ -90,14 +91,13 @@ namespace UE::Editor::DataStorage
 		
 		TEDSTABLEVIEWER_API void Construct(
 			const FArguments& InArgs,
-			const UScriptStruct* InHierarchyDataColumnType,
-			TFunction<RowHandle(const void*, const UScriptStruct*)> InGetRowHandleFn);
+			TSharedPtr<IHierarchyViewerDataInterface> InHierarchyInterface);
 
 		// Clear the current list of columns being displayed and set it to the given list
-		TEDSTABLEVIEWER_API void SetColumns(const TArray<TWeakObjectPtr<const UScriptStruct>>& Columns);
+		TEDSTABLEVIEWER_API virtual void SetColumns(const TArray<TWeakObjectPtr<const UScriptStruct>>& Columns) override;
 
 		// Add a custom column to display in the table viewer, that doesn't necessarily map to a Teds column
-		TEDSTABLEVIEWER_API void AddCustomColumn(const TSharedRef<FTedsTableViewerColumn>& InColumn);
+		TEDSTABLEVIEWER_API virtual void AddCustomRowWidget(const TSharedRef<FTedsTableViewerColumn>& InColumn) override;
 
 		// Execute the given callback for each row that is selected in the table viewer
 		TEDSTABLEVIEWER_API virtual void ForEachSelectedRow(TFunctionRef<void(RowHandle)> InCallback) const override;
@@ -168,11 +168,8 @@ namespace UE::Editor::DataStorage
 		// The cached list of top level rows
 		TArray<TableViewerItemPtr> TopLevelRows;
 
-		// Column type that holds hierarchy information for the row
-		const UScriptStruct* HierarchyColumnType = nullptr;
-
-		// Function to get the parent from a the HierarchyColumnType of the row
-		TFunction<RowHandle(const void*, const UScriptStruct*)> GetParentFromColumnFn;
+		// Hierarchy data that is used to extract information about the parent of a row
+		TSharedPtr<IHierarchyViewerDataInterface> HierarchyInterface;
 
 	};
 }
\ No newline at end of file
diff --git a/Engine/Plugins/Experimental/EditorDataStorageFeatures/Source/TedsTableViewer/Public/Widgets/STedsTableViewer.h b/Engine/Plugins/Experimental/EditorDataStorageFeatures/Source/TedsTableViewer/Public/Widgets/STedsTableViewer.h
--- a/Engine/Plugins/Experimental/EditorDataStorageFeatures/Source/TedsTableViewer/Public/Widgets/STedsTableViewer.h
+++ b/Engine/Plugins/Experimental/EditorDataStorageFeatures/Source/TedsTableViewer/Public/Widgets/STedsTableViewer.h
@@ -83,10 +83,10 @@ namespace UE::Editor::DataStorage
 		TEDSTABLEVIEWER_API void Construct(const FArguments& InArgs);
 
 		// Clear the current list of columns being displayed and set it to the given list
-		TEDSTABLEVIEWER_API void SetColumns(const TArray<TWeakObjectPtr<const UScriptStruct>>& Columns);
+		TEDSTABLEVIEWER_API virtual void SetColumns(const TArray<TWeakObjectPtr<const UScriptStruct>>& Columns) override;
 
 		// Add a custom column to display in the table viewer, that doesn't necessarily map to a Teds column
-		TEDSTABLEVIEWER_API void AddCustomColumn(const TSharedRef<FTedsTableViewerColumn>& InColumn);
+		TEDSTABLEVIEWER_API virtual void AddCustomRowWidget(const TSharedRef<FTedsTableViewerColumn>& InColumn) override;
 
 		// Execute the given callback for each row that is selected in the table viewer
 		TEDSTABLEVIEWER_API virtual void ForEachSelectedRow(TFunctionRef<void(RowHandle)> InCallback) const override;
diff --git a/Engine/Plugins/Experimental/EditorDataStorageFeatures/Source/TedsTableViewer/Public/Widgets/STedsTileViewer.h b/Engine/Plugins/Experimental/EditorDataStorageFeatures/Source/TedsTableViewer/Public/Widgets/STedsTileViewer.h
--- a/Engine/Plugins/Experimental/EditorDataStorageFeatures/Source/TedsTableViewer/Public/Widgets/STedsTileViewer.h
+++ b/Engine/Plugins/Experimental/EditorDataStorageFeatures/Source/TedsTableViewer/Public/Widgets/STedsTileViewer.h
@@ -101,6 +101,12 @@ namespace UE::Editor::DataStorage
 
 		TEDSTABLEVIEWER_API virtual bool IsSelectedExclusively(RowHandle InRow) const override;
 
+		// Clear the current list of columns being displayed and set it to the given list
+		TEDSTABLEVIEWER_API virtual void SetColumns(const TArray<TWeakObjectPtr<const UScriptStruct>>& Columns) override;
+
+		// Add a custom column to display in the table viewer, that doesn't necessarily map to a Teds column
+		TEDSTABLEVIEWER_API virtual void AddCustomRowWidget(const TSharedRef<FTedsTableViewerColumn>& InColumn) override;
+
 	protected:
 		
 		TSharedRef<ITableRow> MakeTileWidget(TableViewerItemPtr InItem, const TSharedRef<STableViewBase>& OwnerTable) const;
