Add bidirectional link between TEDS Rows and TypedElementHandles
* Aid in interoperability between legacy TypedElement based systems and TEDS based systems
* Enables TEDS to be queried to get the appropriate TypedElementHandle for the representative object
* Path towards to allowing TypedElementHandle lifetime to slave from TEDS RowHandle lifetimes
* verse::component given a TypedElement handle for interoperability with the SubobjectInstanceEditor

diff --git a/Engine/Plugins/Experimental/EditorDataStorageFeatures/Source/TedsTypedElementBridge/Private/TedsTypedElementBridge/TedsTypedElementActorHandleFactory.cpp b/Engine/Plugins/Experimental/EditorDataStorageFeatures/Source/TedsTypedElementBridge/Private/TedsTypedElementBridge/TedsTypedElementActorHandleFactory.cpp
--- a/Engine/Plugins/Experimental/EditorDataStorageFeatures/Source/TedsTypedElementBridge/Private/TedsTypedElementBridge/TedsTypedElementActorHandleFactory.cpp
+++ b/Engine/Plugins/Experimental/EditorDataStorageFeatures/Source/TedsTypedElementBridge/Private/TedsTypedElementBridge/TedsTypedElementActorHandleFactory.cpp
@@ -3,7 +3,7 @@
 #include "TedsTypedElementActorHandleFactory.h"
 
 #include "TedsTypedElementBridge/TedsTypedElementBridgeCapabilities.h"
-#include "TedsTypedElementBridge/TedsTypedElementBridgeColumns.h"
+#include "Elements/Columns/TypedElementHandleColumn.h"
 #include "DataStorage/Features.h"
 #include "Elements/Columns/TypedElementCompatibilityColumns.h"
 #include "Elements/Framework/EngineElementsLibrary.h"
diff --git a/Engine/Plugins/Experimental/EditorDataStorageFeatures/Source/TedsTypedElementBridge/Private/TedsTypedElementBridge/TedsTypedElementBridgeQueries.cpp b/Engine/Plugins/Experimental/EditorDataStorageFeatures/Source/TedsTypedElementBridge/Private/TedsTypedElementBridge/TedsTypedElementBridgeQueries.cpp
--- a/Engine/Plugins/Experimental/EditorDataStorageFeatures/Source/TedsTypedElementBridge/Private/TedsTypedElementBridge/TedsTypedElementBridgeQueries.cpp
+++ b/Engine/Plugins/Experimental/EditorDataStorageFeatures/Source/TedsTypedElementBridge/Private/TedsTypedElementBridge/TedsTypedElementBridgeQueries.cpp
@@ -3,7 +3,7 @@
 #include "TedsTypedElementBridge/TedsTypedElementBridgeQueries.h"
 
 #include "TedsTypedElementBridge/TedsTypedElementBridgeCapabilities.h"
-#include "TedsTypedElementBridge/TedsTypedElementBridgeColumns.h"
+#include "Elements/Columns/TypedElementHandleColumn.h"
 #include "DataStorage/Features.h"
 #include "Elements/Columns/TypedElementCompatibilityColumns.h"
 #include "Elements/Framework/EngineElementsLibrary.h"
@@ -16,7 +16,7 @@
 
 namespace UE::Editor::DataStorage::Compatibility::Private
 {
-	bool bBridgeEnabled = false;
+	bool bBridgeEnabled = true;
 	FAutoConsoleVariableRef CVarBridgeEnabled(
 		TEXT("TEDS.TypedElementBridge.Enable"),
 		bBridgeEnabled,
@@ -106,7 +106,7 @@ void UTypedElementBridgeDataStorageFactory::HandleOnEnabled(IConsoleVariable* CV
 	else
 	{
 		Compatibility::OnTypedElementBridgeEnabled().Broadcast(bIsEnabled);
-		UnregisterQuery_NewUObject(*DataStorage);
 		CleanupTypedElementColumns(*DataStorage);
+		UnregisterQuery_NewUObject(*DataStorage);
 	}
 }
diff --git a/Engine/Plugins/Experimental/EditorDataStorageFeatures/Source/TedsTypedElementBridge/Public/TedsTypedElementBridge/TedsTypedElementBridgeColumns.h b/Engine/Plugins/Experimental/EditorDataStorageFeatures/Source/TedsTypedElementBridge/Public/TedsTypedElementBridge/TedsTypedElementBridgeColumns.h
deleted file mode 100644
diff --git a/Engine/Source/Editor/UnrealEd/Private/Elements/Framework/EditorElements.cpp b/Engine/Source/Editor/UnrealEd/Private/Elements/Framework/EditorElements.cpp
--- a/Engine/Source/Editor/UnrealEd/Private/Elements/Framework/EditorElements.cpp
+++ b/Engine/Source/Editor/UnrealEd/Private/Elements/Framework/EditorElements.cpp
@@ -31,6 +31,8 @@
 #include "Logging/LogMacros.h"
 #include "Modules/ModuleManager.h"
 #include "ObjectNameEditSinkRegistry.h"
+#include "Elements/Actor/ActorElementTedsTypedElementBridgeInterface.h"
+#include "Elements/Interfaces/TedsTypedElementBridgeInterface.h"
 #include "Templates/SharedPointer.h"
 #include "Trace/Detail/Channel.h"
 #include "UObject/UObjectGlobals.h"
@@ -54,6 +56,7 @@ void RegisterEditorActorElements()
 	Registry->RegisterElementInterface<ITypedElementWorldInterface>(NAME_Actor, NewObject<UActorElementEditorWorldInterface>(), /*bAllowOverride*/true);
 	Registry->RegisterElementInterface<ITypedElementSelectionInterface>(NAME_Actor, NewObject<UActorElementEditorSelectionInterface>(), /*bAllowOverride*/true);
 	Registry->RegisterElementInterface<ITypedElementAssetDataInterface>(NAME_Actor, NewObject<UActorElementEditorAssetDataInterface>(), /*bAllowOverride*/true);
+	Registry->RegisterElementInterface<ITedsTypedElementBridgeInterface>(NAME_Actor, NewObject<UActorElementTedsTypedElementBridgeInterface>());
 }
 
 void RegisterEditorComponentElements()
diff --git a/Engine/Source/Runtime/Engine/Private/Elements/Actor/ActorElementTedsTypedElementBridgeInterface.cpp b/Engine/Source/Runtime/Engine/Private/Elements/Actor/ActorElementTedsTypedElementBridgeInterface.cpp
new file mode 100644
--- /dev/null
+++ b/Engine/Source/Runtime/Engine/Private/Elements/Actor/ActorElementTedsTypedElementBridgeInterface.cpp
@@ -0,0 +1,18 @@
+﻿// Copyright Epic Games, Inc. All Rights Reserved.
+
+#include "Elements/Actor/ActorElementTedsTypedElementBridgeInterface.h"
+
+#include "DataStorage/Features.h"
+#include "Elements/Actor/ActorElementData.h"
+#include "Elements/Interfaces/TypedElementDataStorageCompatibilityInterface.h"
+
+FTedsRowHandle UActorElementTedsTypedElementBridgeInterface::GetRowHandle(const FTypedElementHandle& InElementHandle) const
+{
+	const AActor* Actor = ActorElementDataUtil::GetActorFromHandle(InElementHandle);
+
+	using namespace UE::Editor::DataStorage;
+	ICompatibilityProvider* CompatibilityProvider = GetMutableDataStorageFeature<ICompatibilityProvider>(CompatibilityFeatureName);
+	RowHandle Handle = CompatibilityProvider->FindRowWithCompatibleObject(Actor);
+	
+	return FTedsRowHandle(Handle);
+}
diff --git a/Engine/Source/Runtime/Engine/Public/Elements/Actor/ActorElementTedsTypedElementBridgeInterface.h b/Engine/Source/Runtime/Engine/Public/Elements/Actor/ActorElementTedsTypedElementBridgeInterface.h
new file mode 100644
--- /dev/null
+++ b/Engine/Source/Runtime/Engine/Public/Elements/Actor/ActorElementTedsTypedElementBridgeInterface.h
@@ -0,0 +1,15 @@
+﻿// Copyright Epic Games, Inc. All Rights Reserved.
+
+#pragma once
+#include "Elements/Interfaces/TedsTypedElementBridgeInterface.h"
+
+#include "ActorElementTedsTypedElementBridgeInterface.generated.h"
+
+UCLASS(MinimalAPI)
+class UActorElementTedsTypedElementBridgeInterface : public UObject, public ITedsTypedElementBridgeInterface
+{
+	GENERATED_BODY()
+
+public:
+	virtual FTedsRowHandle GetRowHandle(const FTypedElementHandle& InElementHandle) const override;
+};
\ No newline at end of file
diff --git a/Engine/Source/Runtime/Engine/Public/Elements/Interfaces/TedsTypedElementBridgeInterface.h b/Engine/Source/Runtime/Engine/Public/Elements/Interfaces/TedsTypedElementBridgeInterface.h
new file mode 100644
--- /dev/null
+++ b/Engine/Source/Runtime/Engine/Public/Elements/Interfaces/TedsTypedElementBridgeInterface.h
@@ -0,0 +1,29 @@
+﻿// Copyright Epic Games, Inc. All Rights Reserved.
+
+#pragma once
+
+#include "DataStorage/Handles.h"
+#include "Elements/Framework/TypedElementHandle.h"
+#include "UObject/Interface.h"
+
+#include "TedsTypedElementBridgeInterface.generated.h"
+
+UINTERFACE(MinimalAPI, meta = (CannotImplementInterfaceInBlueprint))
+class UTedsTypedElementBridgeInterface : public UInterface
+{
+	GENERATED_BODY()
+};
+
+class ITedsTypedElementBridgeInterface
+{
+	GENERATED_BODY()
+
+public:
+	virtual FTedsRowHandle GetRowHandle(const FTypedElementHandle& InElementHandle) const { return FTedsRowHandle(); }
+};
+
+template <>
+struct TTypedElement<ITedsTypedElementBridgeInterface> : public TTypedElementBase<ITedsTypedElementBridgeInterface>
+{
+	FTedsRowHandle GetRowHandle() const { return InterfacePtr->GetRowHandle(*this); }
+};
diff --git a/Engine/Source/Runtime/TypedElementFramework/Public/DataStorage/Handles.h b/Engine/Source/Runtime/TypedElementFramework/Public/DataStorage/Handles.h
--- a/Engine/Source/Runtime/TypedElementFramework/Public/DataStorage/Handles.h
+++ b/Engine/Source/Runtime/TypedElementFramework/Public/DataStorage/Handles.h
@@ -47,6 +47,11 @@ struct FTedsRowHandle
 		RowHandle = InRowHandle;
 		return *this;
 	}
+
+	bool IsValid() const
+	{
+		return RowHandle != UE::Editor::DataStorage::InvalidRowHandle;
+	}
 		
 	friend uint32 GetTypeHash(const FTedsRowHandle& Key)
 	{
diff --git a/Engine/Source/Runtime/TypedElementFramework/Public/Elements/Columns/TypedElementHandleColumn.h b/Engine/Source/Runtime/TypedElementFramework/Public/Elements/Columns/TypedElementHandleColumn.h
new file mode 100644
--- /dev/null
+++ b/Engine/Source/Runtime/TypedElementFramework/Public/Elements/Columns/TypedElementHandleColumn.h
@@ -0,0 +1,23 @@
+﻿// Copyright Epic Games, Inc. All Rights Reserved.
+
+#pragma once
+
+#include "DataStorage/CommonTypes.h"
+#include "Elements/Framework/TypedElementHandle.h"
+
+#include "TypedElementHandleColumn.generated.h"
+
+// A column which contains a TypedElementHandle from the TypedElementFramework
+// Prefer to use the namespaced alias when referencing in code - UE::Editor::DataStorage::Compatibility::FTypedElementColumn
+USTRUCT()
+struct FTedsTypedElementColumn : public FEditorDataStorageColumn
+{
+	GENERATED_BODY()
+
+	FTypedElementHandle Handle;
+};
+
+namespace UE::Editor::DataStorage::Compatibility
+{
+	using FTypedElementColumn = FTedsTypedElementColumn;
+}  // namespace UE::Editor::DataStorage::Compatibility
\ No newline at end of file
