Teds Debugger: Add support for a hierarchy view alongside a combo box to select which hierarchy to view
TEDS: Add a function to list all known hierarchies by name

diff --git a/Engine/Plugins/Experimental/EditorDataStorage/Source/TedsCore/Private/Hierarchy/EditorDataStorageHierarchyImplementation.cpp b/Engine/Plugins/Experimental/EditorDataStorage/Source/TedsCore/Private/Hierarchy/EditorDataStorageHierarchyImplementation.cpp
--- a/Engine/Plugins/Experimental/EditorDataStorage/Source/TedsCore/Private/Hierarchy/EditorDataStorageHierarchyImplementation.cpp
+++ b/Engine/Plugins/Experimental/EditorDataStorage/Source/TedsCore/Private/Hierarchy/EditorDataStorageHierarchyImplementation.cpp
@@ -45,6 +45,14 @@ namespace UE::Editor::DataStorage
 		return nullptr;
 	}
 
+	void FTedsHierarchyRegistrar::ListHierarchyNames(TFunctionRef<void(const FName& HierarchyName)> Callback) const
+	{
+		for (int32 HierarchyIndex = 0; HierarchyIndex < RegisteredHierarchies.Num(); ++HierarchyIndex)
+		{
+			Callback(RegisteredHierarchies[HierarchyIndex].Name);
+		}
+	}
+
 	FHierarchyHandle FTedsHierarchyRegistrar::RegisterHierarchy(ICoreProvider* InProvider, const FHierarchyRegistrationParams& RegistrationParams)
 	{
 		// Create an access interface
diff --git a/Engine/Plugins/Experimental/EditorDataStorage/Source/TedsCore/Private/Hierarchy/EditorDataStorageHierarchyImplementation.h b/Engine/Plugins/Experimental/EditorDataStorage/Source/TedsCore/Private/Hierarchy/EditorDataStorageHierarchyImplementation.h
--- a/Engine/Plugins/Experimental/EditorDataStorage/Source/TedsCore/Private/Hierarchy/EditorDataStorageHierarchyImplementation.h
+++ b/Engine/Plugins/Experimental/EditorDataStorage/Source/TedsCore/Private/Hierarchy/EditorDataStorageHierarchyImplementation.h
@@ -25,6 +25,12 @@ namespace UE::Editor::DataStorage
 		 * Gets an interface to add and remove rows from a hierarchy and walk the hierarchy
 		 */
 		const FTedsHierarchyAccessInterface* GetAccessInterface(FHierarchyHandle) const;
+
+		/**
+		 * Iterate over all registered hierarchyes
+		 */
+		void ListHierarchyNames(TFunctionRef<void(const FName& HierarchyName)> Callback) const;
+		
 	private:
 		void RegisterObservers(ICoreProvider& CoreProvider, const FHierarchyHandle& Handle);
 		void RegisterProcessors(ICoreProvider& CoreProvider, const FHierarchyHandle& Handle);
diff --git a/Engine/Plugins/Experimental/EditorDataStorage/Source/TedsCore/Private/TypedElementDatabase.cpp b/Engine/Plugins/Experimental/EditorDataStorage/Source/TedsCore/Private/TypedElementDatabase.cpp
--- a/Engine/Plugins/Experimental/EditorDataStorage/Source/TedsCore/Private/TypedElementDatabase.cpp
+++ b/Engine/Plugins/Experimental/EditorDataStorage/Source/TedsCore/Private/TypedElementDatabase.cpp
@@ -1109,6 +1109,11 @@ const UScriptStruct* UEditorDataStorage::GetHierarchyDataColumnType(UE::Editor::
 	return nullptr;
 }
 
+void UEditorDataStorage::ListHierarchyNames(TFunctionRef<void(const FName& HierarchyName)> Callback) const
+{
+	Environment->GetHierarchyRegistrar().ListHierarchyNames(Callback);
+}
+
 void UEditorDataStorage::SetParentRow(UE::Editor::DataStorage::FHierarchyHandle InHierarchyHandle, UE::Editor::DataStorage::RowHandle Target,
 	UE::Editor::DataStorage::RowHandle Parent)
 {
diff --git a/Engine/Plugins/Experimental/EditorDataStorage/Source/TedsCore/Private/TypedElementDatabase.h b/Engine/Plugins/Experimental/EditorDataStorage/Source/TedsCore/Private/TypedElementDatabase.h
--- a/Engine/Plugins/Experimental/EditorDataStorage/Source/TedsCore/Private/TypedElementDatabase.h
+++ b/Engine/Plugins/Experimental/EditorDataStorage/Source/TedsCore/Private/TypedElementDatabase.h
@@ -149,6 +149,7 @@ public:
 		UE::Editor::DataStorage::FHierarchyHandle) const override;
 	UE_API virtual const UScriptStruct* GetHierarchyDataColumnType(
 		UE::Editor::DataStorage::FHierarchyHandle) const override;
+	UE_API virtual void ListHierarchyNames(TFunctionRef<void(const FName& HierarchyName)> Callback) const override;
 	UE_API virtual void SetParentRow(
 		UE::Editor::DataStorage::FHierarchyHandle, UE::Editor::DataStorage::RowHandle Target,	UE::Editor::DataStorage::RowHandle Parent) override;
 	virtual void SetUnresolvedParent(
diff --git a/Engine/Plugins/Experimental/EditorDataStorageFeatures/Source/TedsDebugger/Private/QueryEditor/TedsQueryEditor.cpp b/Engine/Plugins/Experimental/EditorDataStorageFeatures/Source/TedsDebugger/Private/QueryEditor/TedsQueryEditor.cpp
--- a/Engine/Plugins/Experimental/EditorDataStorageFeatures/Source/TedsDebugger/Private/QueryEditor/TedsQueryEditor.cpp
+++ b/Engine/Plugins/Experimental/EditorDataStorageFeatures/Source/TedsDebugger/Private/QueryEditor/TedsQueryEditor.cpp
@@ -8,6 +8,7 @@
 #include "Widgets/Input/SEditableComboBox.h"
 #include "Widgets/QueryEditor/TedsConditionSelectionComboWidget.h"
 #include "Widgets/QueryEditor/TedsConditionCollectionViewWidget.h"
+#include "Widgets/QueryEditor/TedsHierarchySelectionComboWidget.h"
 #include "Widgets/QueryEditor/TedsQueryEditorResultsView.h"
 
 #define LOCTEXT_NAMESPACE "TedsQueryEditor"
@@ -28,6 +29,8 @@ namespace UE::Editor::DataStorage::Debug::QueryEditor
 	{	
 		ComboItems.Reset();
 		Model = &QueryEditorModel;
+
+		HierarchyComboWidget = SNew(SHierarchyComboWidget, *Model);
 	
 		ChildSlot
 		[
@@ -37,6 +40,11 @@ namespace UE::Editor::DataStorage::Debug::QueryEditor
 				SNew(SVerticalBox)
 				+SVerticalBox::Slot()
 				.AutoHeight()
+				[
+					CreateToolbar()
+				]
+				+SVerticalBox::Slot()
+				.AutoHeight()
 				[
 					SNew(SHorizontalBox)
 					+SHorizontalBox::Slot()
@@ -93,12 +101,97 @@ namespace UE::Editor::DataStorage::Debug::QueryEditor
 				]
 				+SVerticalBox::Slot()
 				[
-					SNew(SResultsView, *Model)
+					SAssignNew(ResultsView, SResultsView, *Model)
 				]
 			]
 			
 		];
 	}
+
+	TSharedRef<SWidget> SQueryEditorWidget::CreateToolbar()
+	{
+		// We're directly using FMenuBuilder since this is a toolbar specific to the query editor that doesn't need extension via UToolMenus
+		FSlimHorizontalToolBarBuilder ToolBarBuilder(nullptr, FMultiBoxCustomization::None);
+
+		ToolBarBuilder.AddComboButton( 
+			FUIAction(),
+			FOnGetContent::CreateSP(this, &SQueryEditorWidget::GetViewModesMenuWidget),
+			TAttribute<FText>::CreateSP(this, &SQueryEditorWidget::GetCurrentViewModeText),
+			TAttribute<FText>(),
+			FSlateIcon(),
+			/*bInSimpleComboBox*/ false);
+
+		if (HierarchyComboWidget)
+		{
+			ToolBarBuilder.AddWidget(
+				HierarchyComboWidget.ToSharedRef(),
+				NAME_None, // default value
+				true, // default value
+				HAlign_Fill, // default value
+				FNewMenuDelegate(), // default value
+				TAttribute<EVisibility>::CreateLambda([this]()
+				{
+					// Only show the hierarchy combo box if we are using the hierarchy viewer
+					return ResultsView->GetViewMode() == ETableViewMode::Tree ? EVisibility::Visible : EVisibility::Collapsed;
+				}));
+		}
+
+		return ToolBarBuilder.MakeWidget();
+	}
+
+	TSharedRef<SWidget> SQueryEditorWidget::GetViewModesMenuWidget()
+	{
+		FMenuBuilder MenuBuilder(true, nullptr);
+
+		// Manually add the list and tree (hierarchy) view - The tile view is not supported yet
+		MenuBuilder.AddMenuEntry(
+			GetViewModeAsText(ETableViewMode::List),
+			FText::GetEmpty(),
+			FSlateIcon(), 
+			FUIAction(FExecuteAction::CreateLambda([this]()
+			{
+				if (ResultsView)
+				{
+					ResultsView->SetViewMode(ETableViewMode::List);
+				}
+			}))
+		);
+
+		MenuBuilder.AddMenuEntry(
+			GetViewModeAsText(ETableViewMode::Tree),
+			FText::GetEmpty(),
+			FSlateIcon(), 
+			FUIAction(FExecuteAction::CreateLambda([this]()
+			{
+				if (ResultsView)
+				{
+					ResultsView->SetViewMode(ETableViewMode::Tree);
+				}
+			}))
+		);
+
+		return MenuBuilder.MakeWidget();
+	}
+
+	FText SQueryEditorWidget::GetCurrentViewModeText() const
+	{
+		return ResultsView ? GetViewModeAsText(ResultsView->GetViewMode()) : FText::GetEmpty();
+	}
+
+	FText SQueryEditorWidget::GetViewModeAsText(ETableViewMode::Type InViewMode)
+	{
+		switch (InViewMode)
+		{
+		case ETableViewMode::List:
+			return LOCTEXT("ListViewMode", "List View");
+		case ETableViewMode::Tile:
+			return LOCTEXT("TileViewMode", "Tile View");
+		case ETableViewMode::Tree:
+			return LOCTEXT("TreeViewMode", "Hierarchy View");
+		default:
+			return FText::GetEmpty();
+		}
+	}
 } // namespace UE::Editor::DataStorage::Debug::QueryEditor
 
 #undef LOCTEXT_NAMESPACE
\ No newline at end of file
diff --git a/Engine/Plugins/Experimental/EditorDataStorageFeatures/Source/TedsDebugger/Private/QueryEditor/TedsQueryEditor.h b/Engine/Plugins/Experimental/EditorDataStorageFeatures/Source/TedsDebugger/Private/QueryEditor/TedsQueryEditor.h
--- a/Engine/Plugins/Experimental/EditorDataStorageFeatures/Source/TedsDebugger/Private/QueryEditor/TedsQueryEditor.h
+++ b/Engine/Plugins/Experimental/EditorDataStorageFeatures/Source/TedsDebugger/Private/QueryEditor/TedsQueryEditor.h
@@ -1,13 +1,17 @@
 // Copyright Epic Games, Inc. All Rights Reserved.
 #pragma once
-#include "Widgets/SCompoundWidget.h"
+
+#include "Framework/Views/ITypedTableView.h"
 #include "Templates/SharedPointerFwd.h"
+#include "Widgets/SCompoundWidget.h"
 
 class SDockTab;
 class SWindow;
 
 namespace UE::Editor::DataStorage::Debug::QueryEditor
 {
+	class SHierarchyComboWidget;
+	class SResultsView;
 	class FTedsQueryEditorModel;
 	class SQueryEditorWidget : public SCompoundWidget
 	{
@@ -24,9 +28,20 @@ namespace UE::Editor::DataStorage::Debug::QueryEditor
 		 * Wrapper class for a combo item in the dropdown selection
 		 */
 		struct ColumnComboItem;
+
+	private:
+
+		TSharedRef<SWidget> CreateToolbar();
+		TSharedRef<SWidget> GetViewModesMenuWidget();
+		FText GetCurrentViewModeText() const;
+		static FText GetViewModeAsText(ETableViewMode::Type InViewMode);
+		
 	private:
+		
 		FTedsQueryEditorModel* Model = nullptr;
 		TArray<TSharedPtr<ColumnComboItem>> ComboItems;
+		TSharedPtr<SResultsView> ResultsView;
+		TSharedPtr<SHierarchyComboWidget> HierarchyComboWidget;
 	};
 }
 
diff --git a/Engine/Plugins/Experimental/EditorDataStorageFeatures/Source/TedsDebugger/Private/QueryEditor/TedsQueryEditorModel.cpp b/Engine/Plugins/Experimental/EditorDataStorageFeatures/Source/TedsDebugger/Private/QueryEditor/TedsQueryEditorModel.cpp
--- a/Engine/Plugins/Experimental/EditorDataStorageFeatures/Source/TedsDebugger/Private/QueryEditor/TedsQueryEditorModel.cpp
+++ b/Engine/Plugins/Experimental/EditorDataStorageFeatures/Source/TedsDebugger/Private/QueryEditor/TedsQueryEditorModel.cpp
@@ -421,4 +421,20 @@ namespace UE::Editor::DataStorage::Debug::QueryEditor
 				});
 		}
 	}
+
+	void FTedsQueryEditorModel::SetHierarchy(const FName& InHierarchy)
+	{
+		CurrentHierarchy = InHierarchy;
+		HierarchyChangedDelegate.Broadcast();
+	}
+
+	FName FTedsQueryEditorModel::GetHierarchyName() const
+	{
+		return CurrentHierarchy;
+	}
+
+	FHierarchyHandle FTedsQueryEditorModel::GetHierarchy() const
+	{
+		return EditorDataStorageProvider.FindHierarchyByName(CurrentHierarchy);
+	}
 } // namespace UE::Editor::DataStorage::Debug::QueryEditor
diff --git a/Engine/Plugins/Experimental/EditorDataStorageFeatures/Source/TedsDebugger/Private/QueryEditor/TedsQueryEditorModel.h b/Engine/Plugins/Experimental/EditorDataStorageFeatures/Source/TedsDebugger/Private/QueryEditor/TedsQueryEditorModel.h
--- a/Engine/Plugins/Experimental/EditorDataStorageFeatures/Source/TedsDebugger/Private/QueryEditor/TedsQueryEditorModel.h
+++ b/Engine/Plugins/Experimental/EditorDataStorageFeatures/Source/TedsDebugger/Private/QueryEditor/TedsQueryEditorModel.h
@@ -92,8 +92,15 @@ namespace UE::Editor::DataStorage::Debug::QueryEditor
 		const UScriptStruct* GetColumnScriptStruct(FConditionEntryHandle Handle) const;
 		
 		FTedsQueryEditorModel_ModelChanged& GetModelChangedDelegate() { return ModelChangedDelegate; }
+
+		// Separate event for hierarchy changes since those should not refresh the selected operators
+		FTedsQueryEditorModel_ModelChanged& OnHierarchyChanged() { return HierarchyChangedDelegate; }
 		
 		void RegenerateColumnsList();
+
+		void SetHierarchy(const FName& InHierarchy);
+		FName GetHierarchyName() const;
+		FHierarchyHandle GetHierarchy() const;
 	
 	private:
 		
@@ -119,5 +126,9 @@ namespace UE::Editor::DataStorage::Debug::QueryEditor
 		uint64 CurrentVersion = 0;
 		uint64 GeneratedVersion = 0;
 		QueryHandle QueryHandle;
+
+		// The current hierarchy being viewed (only applicable for hierarchy based widgets)
+		FName CurrentHierarchy;
+		FTedsQueryEditorModel_ModelChanged HierarchyChangedDelegate;
 	};
 } // namespace UE::Editor::DataStorage::Debug::QueryEditor
diff --git a/Engine/Plugins/Experimental/EditorDataStorageFeatures/Source/TedsDebugger/Private/Widgets/QueryEditor/TedsHierarchySelectionComboWidget.cpp b/Engine/Plugins/Experimental/EditorDataStorageFeatures/Source/TedsDebugger/Private/Widgets/QueryEditor/TedsHierarchySelectionComboWidget.cpp
new file mode 100644
--- /dev/null
+++ b/Engine/Plugins/Experimental/EditorDataStorageFeatures/Source/TedsDebugger/Private/Widgets/QueryEditor/TedsHierarchySelectionComboWidget.cpp
@@ -0,0 +1,52 @@
+// Copyright Epic Games, Inc. All Rights Reserved.
+
+#include "TedsHierarchySelectionComboWidget.h"
+
+#include "Elements/Interfaces/TypedElementDataStorageInterface.h"
+
+namespace UE::Editor::DataStorage::Debug::QueryEditor
+{
+	void SHierarchyComboWidget::Construct(const FArguments& InArgs, FTedsQueryEditorModel& InModel)
+	{
+		Model = &InModel;
+
+		GenerateHierarchyList();
+		
+		ChildSlot
+		[
+			SAssignNew(SearchableComboBox, SSearchableComboBox)
+            	.OptionsSource(&Hierarchies)
+            	.OnSelectionChanged_Lambda([this](TSharedPtr<FString> SelectedItem, ESelectInfo::Type SelectInfo)
+            	{
+            		Model->SetHierarchy(FName(*SelectedItem));
+            	})
+            	.OnGenerateWidget_Lambda([](TSharedPtr<FString> InItem)
+            	{
+            		return SNew(STextBlock)
+            		.Text(FText::FromString(InItem.IsValid() ? *InItem : FString()));
+            	})
+				.OnComboBoxOpening_Lambda([this]()
+				{
+					GenerateHierarchyList();
+				})
+            	.Content()
+            	[
+            		SNew(STextBlock)
+					.Text_Lambda([this]()
+					{
+						return FText::FromName(Model->GetHierarchyName());
+					})
+            	]
+		];
+	}
+
+	void SHierarchyComboWidget::GenerateHierarchyList()
+	{
+		Hierarchies.Empty();
+		
+		Model->GetTedsInterface().ListHierarchyNames([this](const FName& HierarchyName)
+		{
+			Hierarchies.Add(MakeShared<FString>(HierarchyName.ToString()));
+		});
+	}
+}
diff --git a/Engine/Plugins/Experimental/EditorDataStorageFeatures/Source/TedsDebugger/Private/Widgets/QueryEditor/TedsHierarchySelectionComboWidget.h b/Engine/Plugins/Experimental/EditorDataStorageFeatures/Source/TedsDebugger/Private/Widgets/QueryEditor/TedsHierarchySelectionComboWidget.h
new file mode 100644
--- /dev/null
+++ b/Engine/Plugins/Experimental/EditorDataStorageFeatures/Source/TedsDebugger/Private/Widgets/QueryEditor/TedsHierarchySelectionComboWidget.h
@@ -0,0 +1,26 @@
+// Copyright Epic Games, Inc. All Rights Reserved.
+
+#pragma once
+
+#include "SSearchableComboBox.h"
+#include "QueryEditor/TedsQueryEditorModel.h"
+
+namespace UE::Editor::DataStorage::Debug::QueryEditor
+{
+	class SHierarchyComboWidget : public SCompoundWidget
+	{
+	public:
+		SLATE_BEGIN_ARGS( SHierarchyComboWidget ){}
+		SLATE_END_ARGS()
+		
+		void Construct(const FArguments& InArgs, FTedsQueryEditorModel& InModel);
+
+	protected:
+		void GenerateHierarchyList();
+		
+	protected:
+		TArray<TSharedPtr<FString>> Hierarchies;
+		FTedsQueryEditorModel* Model = nullptr;
+		TSharedPtr<SSearchableComboBox> SearchableComboBox;
+	};
+}
diff --git a/Engine/Plugins/Experimental/EditorDataStorageFeatures/Source/TedsDebugger/Private/Widgets/QueryEditor/TedsQueryEditorResultsView.cpp b/Engine/Plugins/Experimental/EditorDataStorageFeatures/Source/TedsDebugger/Private/Widgets/QueryEditor/TedsQueryEditorResultsView.cpp
--- a/Engine/Plugins/Experimental/EditorDataStorageFeatures/Source/TedsDebugger/Private/Widgets/QueryEditor/TedsQueryEditorResultsView.cpp
+++ b/Engine/Plugins/Experimental/EditorDataStorageFeatures/Source/TedsDebugger/Private/Widgets/QueryEditor/TedsQueryEditorResultsView.cpp
@@ -18,6 +18,9 @@
 #include "Elements/Columns/TypedElementHiearchyColumns.h"
 #include "TedsQueryNode.h"
 #include "TedsRowSortNode.h"
+#include "QueryStack/TopLevelRowsNode.h"
+#include "Widgets/Layout/SBorder.h"
+#include "Widgets/STedsHierarchyViewer.h"
 #include "Widgets/STedsTableViewer.h"
 #include "Widgets/SRowDetails.h"
 #include "Widgets/Text/STextBlock.h"
@@ -29,6 +32,7 @@ namespace UE::Editor::DataStorage::Debug::QueryEditor
 	SResultsView::~SResultsView()
 	{
 		Model->GetModelChangedDelegate().Remove(ModelChangedDelegateHandle);
+		Model->OnHierarchyChanged().Remove(HierarchyChangedDelegateHandle);
 	}
 
 	void SResultsView::Construct(const FArguments& InArgs, FTedsQueryEditorModel& InModel)
@@ -37,18 +41,21 @@ namespace UE::Editor::DataStorage::Debug::QueryEditor
 		
 		Model = &InModel;
 		ModelChangedDelegateHandle = Model->GetModelChangedDelegate().AddRaw(this, &SResultsView::OnModelChanged);
+		HierarchyChangedDelegateHandle = Model->OnHierarchyChanged().AddRaw(this, &SResultsView::OnHierarchyChanged);
+		
 		UiProvider = GetMutableDataStorageFeature<IUiProvider>(UiFeatureName);
 		checkf(UiProvider, TEXT("Cannot create SResultsView without TEDS UI"));
 
-		
 		// Create a custom column for the table viewer to display row handles
 		CreateRowHandleColumn();
 	
 		QueryNode = MakeShared<FQueryNode>(Model->GetTedsInterface());
 
+		// Use RefreshOnUpdate to re-run the query every frame since we don't have a good change detection mechanism for hierarchies right now
+		// and just using observers can cause the visual hierarchy to go out of sync
 		RowQueryStack = MakeShared<FRowQueryResultsNode>(Model->GetTedsInterface(), QueryNode,
-			FRowQueryResultsNode::ESyncFlags::RefreshOnQueryChange | FRowQueryResultsNode::ESyncFlags::IncrementWhenDifferent);
-		RowQueryStack = MakeShared<FRowMonitorNode>(Model->GetTedsInterface(), RowQueryStack, QueryNode);
+			FRowQueryResultsNode::ESyncFlags::RefreshOnUpdate);
+		//RowQueryStack = MakeShared<FRowMonitorNode>(Model->GetTedsInterface(), RowQueryStack, QueryNode);
 		
 		ChildSlot
 		[
@@ -59,24 +66,11 @@ namespace UE::Editor::DataStorage::Debug::QueryEditor
 				+SSplitter::Slot()
 				.Value(0.5f)
 				[
-					SAssignNew(TableViewer, STedsTableViewer)
-					.QueryStack(RowQueryStack)
-					.EmptyRowsMessage(LOCTEXT("EmptyRowsMessage", "The provided query has no results."))
-					.OnSelectionChanged(STedsTableViewer::FOnSelectionChanged::CreateLambda(
-						[this](RowHandle SelectedRow)
-							{
-								if(RowDetailsWidget)
-								{
-									if(SelectedRow != InvalidRowHandle)
-									{
-										RowDetailsWidget->SetRow(SelectedRow);
-									}
-									else
-									{
-										RowDetailsWidget->ClearRow();
-									}
-								}
-							}))
+					SAssignNew(TableViewerSlot, SBorder)
+					[
+						CreateTableViewer()
+					]
+					
 				]
 				+SSplitter::Slot()
 				.Value(0.5f)
@@ -118,7 +112,6 @@ namespace UE::Editor::DataStorage::Debug::QueryEditor
 
 		ICoreProvider& TedsInterface = Model->GetTedsInterface();
 
-		
 		if(bModelDirty)
 		{
 			{
@@ -168,12 +161,33 @@ namespace UE::Editor::DataStorage::Debug::QueryEditor
 		}
 	}
 
+	ETableViewMode::Type SResultsView::GetViewMode() const
+	{
+		return ResultsViewMode;
+	}
+
+	void SResultsView::SetViewMode(ETableViewMode::Type InViewMode)
+	{
+		ResultsViewMode = InViewMode;
+		TableViewerSlot->SetContent(CreateTableViewer()); // Regenerate the table viewer widget
+	}
+
 	void SResultsView::OnModelChanged()
 	{
 		bModelDirty = true;
 		Invalidate(EInvalidateWidgetReason::Layout);
 	}
 
+	void SResultsView::OnHierarchyChanged()
+	{
+		HierarchyHandle = Model->GetHierarchy();
+
+		if (ResultsViewMode == ETableViewMode::Tree)
+		{
+			TableViewerSlot->SetContent(CreateTableViewer()); // Regenerate the table viewer widget
+		}
+	}
+
 	void SResultsView::CreateRowHandleColumn()
 	{
 		auto AssignWidgetToColumn = [this](TUniquePtr<FTypedElementWidgetConstructor> Constructor, TConstArrayView<TWeakObjectPtr<const UScriptStruct>>)
@@ -182,12 +196,81 @@ namespace UE::Editor::DataStorage::Debug::QueryEditor
 			RowHandleColumn = MakeShared<FTedsTableViewerColumn>(TEXT("Row Handle"), WidgetConstructor);
 			return false;
 		};
-	
 
 		UiProvider->CreateWidgetConstructors(
 			UiProvider->FindPurpose(IUiProvider::FPurposeInfo("General", "Cell", "RowHandle").GeneratePurposeID()),
 			FMetaDataView(), AssignWidgetToColumn);
 	}
+
+	TSharedRef<SWidget> SResultsView::CreateTableViewer()
+	{
+		switch (ResultsViewMode)
+		{
+		case ETableViewMode::List:
+			TableViewer = SNew(STedsTableViewer)
+					.QueryStack(RowQueryStack)
+					.EmptyRowsMessage(LOCTEXT("EmptyRowsMessage", "The provided query has no results."))
+					.OnSelectionChanged(STedsTableViewer::FOnSelectionChanged::CreateLambda(
+						[this](RowHandle SelectedRow)
+							{
+								if(RowDetailsWidget)
+								{
+									if(SelectedRow != InvalidRowHandle)
+									{
+										RowDetailsWidget->SetRow(SelectedRow);
+									}
+									else
+									{
+										RowDetailsWidget->ClearRow();
+									}
+								}
+							}));
+			break;
+		
+		case ETableViewMode::Tree:
+
+			HierarchyData = MakeShared<FHierarchyViewerData>(HierarchyHandle);
+			TopLevelRowsNode = MakeShared<QueryStack::FTopLevelRowsNode>(&Model->GetTedsInterface(), HierarchyData, RowQueryStack);
+
+			TableViewer =
+				SNew(SHierarchyViewer, HierarchyData)
+					.AllNodeProvider(RowQueryStack)
+					.TopLevelNodeProvider(TopLevelRowsNode)
+					.EmptyRowsMessage(LOCTEXT("EmptyRowsMessage", "The provided query has no results."))
+					.OnSelectionChanged(STedsTableViewer::FOnSelectionChanged::CreateLambda(
+						[this](RowHandle SelectedRow)
+							{
+								if(RowDetailsWidget)
+								{
+									if(SelectedRow != InvalidRowHandle)
+									{
+										RowDetailsWidget->SetRow(SelectedRow);
+									}
+									else
+									{
+										RowDetailsWidget->ClearRow();
+									}
+								}
+							}));
+			break;
+
+			// NYI
+		case ETableViewMode::Tile:
+			[[fallthrough]];
+		default:
+			checkf(false, TEXT("Unknown view mode requested! Make sure the options in the UI to select view modes matches the implementation")); 
+			return SNullWidget::NullWidget;
+		}
+
+		if (RowDetailsWidget)
+		{
+			RowDetailsWidget->ClearRow();
+		}
+
+		bModelDirty = true; // We need to set the columns again
+		
+		return TableViewer->AsWidget();
+	}
 } // namespace UE::Editor::DataStorage::Debug::QueryEditor
 
 #undef LOCTEXT_NAMESPACE
diff --git a/Engine/Plugins/Experimental/EditorDataStorageFeatures/Source/TedsDebugger/Private/Widgets/QueryEditor/TedsQueryEditorResultsView.h b/Engine/Plugins/Experimental/EditorDataStorageFeatures/Source/TedsDebugger/Private/Widgets/QueryEditor/TedsQueryEditorResultsView.h
--- a/Engine/Plugins/Experimental/EditorDataStorageFeatures/Source/TedsDebugger/Private/Widgets/QueryEditor/TedsQueryEditorResultsView.h
+++ b/Engine/Plugins/Experimental/EditorDataStorageFeatures/Source/TedsDebugger/Private/Widgets/QueryEditor/TedsQueryEditorResultsView.h
@@ -1,7 +1,9 @@
 // Copyright Epic Games, Inc. All Rights Reserved.
 #pragma once
 
+#include "HierarchyViewerIntefaces.h"
 #include "Elements/Interfaces/TypedElementDataStorageInterface.h"
+#include "Framework/Views/ITypedTableView.h"
 #include "Widgets/SCompoundWidget.h"
 
 
@@ -9,9 +11,11 @@ class ISceneOutliner;
 class SSceneOutliner;
 class SHorizontalBox;
 struct FTypedElementWidgetConstructor;
+class SBorder;
 
 namespace UE::Editor::DataStorage
 {
+	class ITableViewer;
 	class STedsTableViewer;
 	class FTedsTableViewerColumn;
 	class SRowDetails;
@@ -19,6 +23,7 @@ namespace UE::Editor::DataStorage
 
 	namespace QueryStack
 	{
+		class FTopLevelRowsNode;
 		class FQueryNode;
 		class IRowNode;
 	}
@@ -37,22 +42,35 @@ namespace UE::Editor::DataStorage
 			void Construct(const FArguments& InArgs, FTedsQueryEditorModel& InModel);
 			void Tick(const FGeometry& AllottedGeometry, const double InCurrentTime, const float InDeltaTime) override;
 
+			ETableViewMode::Type GetViewMode() const;
+			void SetViewMode(ETableViewMode::Type InViewMode);
+
 		private:
 
 			void OnModelChanged();
+			void OnHierarchyChanged();
 			void CreateRowHandleColumn();
+			TSharedRef<SWidget> CreateTableViewer();
 		
 			FTedsQueryEditorModel* Model = nullptr;
 			FDelegateHandle ModelChangedDelegateHandle;
+			FDelegateHandle HierarchyChangedDelegateHandle;
 			bool bModelDirty = true;
 
 
 			QueryHandle CountQueryHandle = InvalidQueryHandle;
 			
 			TArray<RowHandle> TableViewerRows;
-			TSharedPtr<STedsTableViewer> TableViewer;
+			TSharedPtr<ITableViewer> TableViewer;
 			TSharedPtr<QueryStack::FQueryNode> QueryNode;
 			TSharedPtr<QueryStack::IRowNode> RowQueryStack;
+
+			// Hierarchy Data used by SHierarchyView to resolve parenting
+			TSharedPtr<FHierarchyViewerData> HierarchyData;
+			// Node to only collect top level rows in RowQueryStack for the hierarchy view
+			TSharedPtr<QueryStack::FTopLevelRowsNode> TopLevelRowsNode;
+			// The current hierarchy the widget is viewing
+			FHierarchyHandle HierarchyHandle;
 			
 			// Custom column for the table viewer to display row handles
 			TSharedPtr<FTedsTableViewerColumn> RowHandleColumn;
@@ -61,6 +79,11 @@ namespace UE::Editor::DataStorage
 			TSharedPtr<SRowDetails> RowDetailsWidget;
 			
 			IUiProvider* UiProvider = nullptr;
+
+			// Border inside which the table viewer is stored so we can re-create it when needed
+			TSharedPtr<SBorder> TableViewerSlot;
+
+			ETableViewMode::Type ResultsViewMode = ETableViewMode::List;
 		};
 
 	} // namespace Debug::QueryEditor
diff --git a/Engine/Source/Runtime/TypedElementFramework/Public/Elements/Interfaces/TypedElementDataStorageInterface.h b/Engine/Source/Runtime/TypedElementFramework/Public/Elements/Interfaces/TypedElementDataStorageInterface.h
--- a/Engine/Source/Runtime/TypedElementFramework/Public/Elements/Interfaces/TypedElementDataStorageInterface.h
+++ b/Engine/Source/Runtime/TypedElementFramework/Public/Elements/Interfaces/TypedElementDataStorageInterface.h
@@ -382,7 +382,7 @@ public:
 	virtual UE::Editor::DataStorage::FHierarchyHandle FindHierarchyByName(const FName& Name) const = 0;
 	virtual bool IsValidHierachyHandle(FHierarchyHandle) const = 0;
 
-		// Gets the tag type indicating the row is a child in this hierarchy and thus has a parent
+	// Gets the tag type indicating the row is a child in this hierarchy and thus has a parent
 	virtual const UScriptStruct* GetChildTagType(FHierarchyHandle InHierarchyHandle) const = 0;
 	// Gets the tag type indicating the row is a parent in this hierarchy and thus has a child
 	virtual const UScriptStruct* GetParentTagType(FHierarchyHandle InHierarchyHandle) const = 0;
@@ -390,6 +390,8 @@ public:
 	// and is used internally.  It is the caller's responsibility to provide at least Read access
 	// to this column type in query requirements in order to use the HierarchyAccessInterface on the row
 	virtual const UScriptStruct* GetHierarchyDataColumnType(FHierarchyHandle InHierarchyHandle) const = 0;
+	// List all the currently known hierarchies by name
+	virtual void ListHierarchyNames(TFunctionRef<void(const FName& HierarchyName)> Callback) const = 0;
 
 	// CoreProvider Context API
 	// =========================
