TEDS Debugger: Fix potential crash due to unsafe ptr dereference without checking

diff --git a/Engine/Plugins/Experimental/EditorDataStorageFeatures/Source/TedsDebugger/Private/QueryEditor/TedsQueryEditor.cpp b/Engine/Plugins/Experimental/EditorDataStorageFeatures/Source/TedsDebugger/Private/QueryEditor/TedsQueryEditor.cpp
--- a/Engine/Plugins/Experimental/EditorDataStorageFeatures/Source/TedsDebugger/Private/QueryEditor/TedsQueryEditor.cpp
+++ b/Engine/Plugins/Experimental/EditorDataStorageFeatures/Source/TedsDebugger/Private/QueryEditor/TedsQueryEditor.cpp
@@ -132,7 +132,7 @@ namespace UE::Editor::DataStorage::Debug::QueryEditor
 				TAttribute<EVisibility>::CreateLambda([this]()
 				{
 					// Only show the hierarchy combo box if we are using the hierarchy viewer
-					return ResultsView->GetViewMode() == ETableViewMode::Tree ? EVisibility::Visible : EVisibility::Collapsed;
+					return ResultsView && ResultsView->GetViewMode() == ETableViewMode::Tree ? EVisibility::Visible : EVisibility::Collapsed;
 				}));
 		}
 
diff --git a/Engine/Plugins/Experimental/EditorDataStorageFeatures/Source/TedsDebugger/Private/Widgets/QueryEditor/TedsHierarchySelectionComboWidget.cpp b/Engine/Plugins/Experimental/EditorDataStorageFeatures/Source/TedsDebugger/Private/Widgets/QueryEditor/TedsHierarchySelectionComboWidget.cpp
--- a/Engine/Plugins/Experimental/EditorDataStorageFeatures/Source/TedsDebugger/Private/Widgets/QueryEditor/TedsHierarchySelectionComboWidget.cpp
+++ b/Engine/Plugins/Experimental/EditorDataStorageFeatures/Source/TedsDebugger/Private/Widgets/QueryEditor/TedsHierarchySelectionComboWidget.cpp
@@ -18,7 +18,10 @@ namespace UE::Editor::DataStorage::Debug::QueryEditor
             	.OptionsSource(&Hierarchies)
             	.OnSelectionChanged_Lambda([this](TSharedPtr<FString> SelectedItem, ESelectInfo::Type SelectInfo)
             	{
-            		Model->SetHierarchy(FName(*SelectedItem));
+            		if (SelectedItem)
+            		{
+            			Model->SetHierarchy(FName(*SelectedItem));
+            		}
             	})
             	.OnGenerateWidget_Lambda([](TSharedPtr<FString> InItem)
             	{
