Disable teds actor hierarchies to fix heap corruption issue.

diff --git a/Engine/Plugins/Experimental/EditorDataStorageFeatures/Source/TedsActorCompatibility/Private/Compatibility/TedsActorParentQueries.cpp b/Engine/Plugins/Experimental/EditorDataStorageFeatures/Source/TedsActorCompatibility/Private/Compatibility/TedsActorParentQueries.cpp
--- a/Engine/Plugins/Experimental/EditorDataStorageFeatures/Source/TedsActorCompatibility/Private/Compatibility/TedsActorParentQueries.cpp
+++ b/Engine/Plugins/Experimental/EditorDataStorageFeatures/Source/TedsActorCompatibility/Private/Compatibility/TedsActorParentQueries.cpp
@@ -15,7 +15,7 @@
 namespace UE::Editor::DataStorage::ActorParentQueries::Private
 {
 	static FName ActorEditorHierarchyName(TEXT("ActorEditorHierarchy"));
-	static bool bAddParentColumnToActors = true;
+	static bool bAddParentColumnToActors = false;
 
 	// Cvar to add parenting info for actors to TEDS
 	static FAutoConsoleVariableRef CvarAddParentColumnToActors(
