[TEDS Outliner Filters] Add the ability to create a TEDS Outliner filter with Query Functions and add further functionality for filter UI (icon, tooltip)

diff --git a/Engine/Plugins/ChaosVD/Source/ChaosVD/Private/ChaosVDWorldOutlinerTab.cpp b/Engine/Plugins/ChaosVD/Source/ChaosVD/Private/ChaosVDWorldOutlinerTab.cpp
--- a/Engine/Plugins/ChaosVD/Source/ChaosVD/Private/ChaosVDWorldOutlinerTab.cpp
+++ b/Engine/Plugins/ChaosVD/Source/ChaosVD/Private/ChaosVDWorldOutlinerTab.cpp
@@ -51,7 +51,6 @@ void FChaosVDWorldOutlinerTab::CreateWorldOutlinerWidget()
 
 	FTedsOutlinerParams Params(nullptr);
 	Params.QueryDescription = OutlinerQueryDescription;
-	Params.bUseDefaultTedsFilters = false;
 	Params.bShowRowHandleColumn = false;
 	Params.bUseDefaultObservers = false; // CVD uses custom observers to time slice addition and removal currently
 
diff --git a/Engine/Plugins/Experimental/Dataflow/Source/DataflowEditor/Private/Dataflow/DataflowOutlinerView.cpp b/Engine/Plugins/Experimental/Dataflow/Source/DataflowEditor/Private/Dataflow/DataflowOutlinerView.cpp
--- a/Engine/Plugins/Experimental/Dataflow/Source/DataflowEditor/Private/Dataflow/DataflowOutlinerView.cpp
+++ b/Engine/Plugins/Experimental/Dataflow/Source/DataflowEditor/Private/Dataflow/DataflowOutlinerView.cpp
@@ -12,6 +12,8 @@
 #include "Elements/Columns/TypedElementVisibilityColumns.h"
 #include "Elements/Framework/TypedElementQueryBuilder.h"
 
+#define LOCTEXT_NAMESPACE "DataflowOutlinerView"
+
 FDataflowOutlinerView::FDataflowOutlinerView(TWeakPtr<FDataflowConstructionScene> InConstructionScene, TWeakPtr<FDataflowSimulationScene> InSimulationScene, TObjectPtr<UDataflowBaseContent> InContent)
 	: FDataflowNodeView(InContent)
 	, OutlinerWidget(nullptr)
@@ -28,12 +30,15 @@ TSharedPtr<ISceneOutliner> FDataflowOutlinerView::CreateWidget()
 {
 	using namespace UE::Editor::DataStorage::Queries;
 	using namespace UE::Editor::Outliner;
-	UE::Editor::DataStorage::FQueryDescription RowQueryDescription = 
+
+	ICoreProvider* Storage = GetMutableDataStorageFeature<ICoreProvider>(StorageFeatureName);
+	
+	const UE::Editor::DataStorage::FQueryDescription RowQueryDescription = 
 		Select()
 		.Where(TColumn<FDataflowSceneObjectTag>(GetEditorContent()->GetDataflowOwner().GetFName()) || TColumn<FDataflowSceneStructTag>(GetEditorContent()->GetDataflowOwner().GetFName()) )
 		.Compile();
 
-	UE::Editor::DataStorage::FQueryDescription ColumnQueryDescription =
+	const UE::Editor::DataStorage::FQueryDescription ColumnQueryDescription =
 		Select()
 		.ReadOnly<FTypedElementClassTypeInfoColumn, FVisibleInEditorColumn>()
 		.Compile();
@@ -46,32 +51,34 @@ TSharedPtr<ISceneOutliner> FDataflowOutlinerView::CreateWidget()
 	
 	UE::Editor::Outliner::FTedsOutlinerParams Params(nullptr);
 	Params.QueryDescription = RowQueryDescription;
-	Params.bUseDefaultTedsFilters = false;
 	Params.bShowRowHandleColumn = false;
 	Params.ColumnQueryDescription = ColumnQueryDescription;
-
+	
 	// Add outliner filter queries
-	UE::Editor::DataStorage::FQueryDescription SimulationFilterQuery =
-		Select()
-		.Where(TColumn<FDataflowSimulationObjectTag>())
-		.Compile();
-	UE::Editor::DataStorage::FQueryDescription ConstructionFilterQuery =
-		Select()
-		.Where(TColumn<FDataflowConstructionObjectTag>())
-		.Compile();
-	UE::Editor::DataStorage::FQueryDescription ElementsFilterQuery =
-		Select()
-		.Where(TColumn<FDataflowSceneStructTag>(GetEditorContent()->GetDataflowOwner().GetFName()))
-		.Compile();
-	UE::Editor::DataStorage::FQueryDescription ComponentsFilterQuery =
-		Select()
-		.Where(TColumn<FDataflowSceneObjectTag>(GetEditorContent()->GetDataflowOwner().GetFName()))
-		.Compile();
-
-	Params.FilterQueries.Emplace("Dataflow Construction", ConstructionFilterQuery);
-	Params.FilterQueries.Emplace("Dataflow Simulation", SimulationFilterQuery);
-	Params.FilterQueries.Emplace("Dataflow Elements", ElementsFilterQuery);
-	Params.FilterQueries.Emplace("Dataflow Components", ComponentsFilterQuery);
+	Params.Filters.Emplace("Dataflow Construction",
+		LOCTEXT("FilterDataflowConstructionDisplayName", "Dataflow Construction"),
+		Storage->RegisterQuery(
+			Select()
+			.Where(TColumn<FDataflowConstructionObjectTag>())
+			.Compile()));
+	Params.Filters.Emplace("Dataflow Simulation",
+		LOCTEXT("FilterDataflowSimulationDisplayName", "Dataflow Simulation"),
+		Storage->RegisterQuery(
+			Select()
+			.Where(TColumn<FDataflowSimulationObjectTag>())
+			.Compile()));
+	Params.Filters.Emplace("Dataflow Elements",
+		LOCTEXT("FilterDataflowElementsDisplayName", "Dataflow Elements"),
+		Storage->RegisterQuery(
+			Select()
+			.Where(TColumn<FDataflowSceneStructTag>(GetEditorContent()->GetDataflowOwner().GetFName()))
+			.Compile()));
+	Params.Filters.Emplace("Dataflow Components",
+		LOCTEXT("FilterDataflowComponentsDisplayName", "Dataflow Components"),
+		Storage->RegisterQuery(
+			Select()
+			.Where(TColumn<FDataflowSceneObjectTag>(GetEditorContent()->GetDataflowOwner().GetFName()))
+			.Compile()));
 	
 	// Empty selection set name is currently the level editor
 	Params.SelectionSetOverride = FName("DataflowSelection");
@@ -82,8 +89,8 @@ TSharedPtr<ISceneOutliner> FDataflowOutlinerView::CreateWidget()
 		return new FDataflowOutlinerMode(Params, ConstructionScene, SimulationScene);
 	});
 	
-
 	InitOptions.ColumnMap.Add(FSceneOutlinerBuiltInColumnTypes::Label(), FSceneOutlinerColumnInfo(ESceneOutlinerColumnVisibility::Visible, 10));
+
 	OutlinerWidget = SNew(SSceneOutliner, InitOptions);
 
 	return OutlinerWidget;
@@ -149,3 +156,5 @@ void FDataflowOutlinerView::AddReferencedObjects(FReferenceCollector& Collector)
 {
 	FDataflowNodeView::AddReferencedObjects(Collector);
 }
+
+#undef LOCTEXT_NAMESPACE
\ No newline at end of file
diff --git a/Engine/Plugins/Experimental/EditorDataStorageFeatures/Source/TedsOutliner/Private/TedsOutlinerFilter.cpp b/Engine/Plugins/Experimental/EditorDataStorageFeatures/Source/TedsOutliner/Private/TedsOutlinerFilter.cpp
--- a/Engine/Plugins/Experimental/EditorDataStorageFeatures/Source/TedsOutliner/Private/TedsOutlinerFilter.cpp
+++ b/Engine/Plugins/Experimental/EditorDataStorageFeatures/Source/TedsOutliner/Private/TedsOutlinerFilter.cpp
@@ -2,22 +2,117 @@
 
 #include "TedsOutlinerFilter.h"
 
+#include "TedsOutlinerHelpers.h"
+#include "DataStorage/Features.h"
+#include "Elements/Columns/TypedElementTypeInfoColumns.h"
+#include "Elements/Framework/TypedElementQueryContext.h"
+#include "Styling/SlateIconFinder.h"
 #include "TedsOutlinerImpl.h"
 
+#define LOCTEXT_NAMESPACE "TEDSOutlinerFilter"
+
 namespace UE::Editor::Outliner
 {
-FTedsOutlinerFilter::FTedsOutlinerFilter(const FName& InFilterName, const FText& InFilterDisplayName,
-	TSharedPtr<FFilterCategory> InCategory, TSharedRef<FTedsOutlinerImpl> InTedsOutlinerImpl,
-	const UE::Editor::DataStorage::FQueryDescription& InFilterQuery)
+FTedsFilterData::FTedsFilterData(const FName& InFilterName, const FText& InFilterDisplayName, const FText& InFilterToolTip, const FName& InFilterIconName,
+	const TSharedPtr<FFilterCategory>& InCategory, const QueryHandle& InFilterQuery)
+	: FilterName(InFilterName)
+	, FilterDisplayName(InFilterDisplayName)
+	, FilterToolTip(InFilterToolTip)
+	, FilterIconName(InFilterIconName)
+	, FilterCategory(InCategory)
+{
+	FilterQuery.Set<QueryHandle>(Helpers::CheckValidFilterQueryHandle(InFilterQuery) ? InFilterQuery : InvalidQueryHandle);
+}
+
+FTedsFilterData::FTedsFilterData(const FName& InFilterName, const FText& InFilterDisplayName, const QueryHandle& InFilterQuery)
+	: FTedsFilterData(
+		InFilterName, 
+		InFilterDisplayName, 
+		FText::FromString(InFilterName.ToString()), 
+		FName(), 
+		nullptr, 
+		InFilterQuery)
+{}
+
+FTedsFilterData::FTedsFilterData(const FName& InFilterName, const FText& InFilterDisplayName, const FText& InFilterToolTip, const FName& InFilterIconName,
+	const TSharedPtr<FFilterCategory>& InCategory, const TQueryFunction<bool>& InFilterQuery)
+	: FilterName(InFilterName)
+	, FilterDisplayName(InFilterDisplayName)
+	, FilterToolTip(InFilterToolTip)
+	, FilterIconName(InFilterIconName)
+	, FilterCategory(InCategory)
+{
+	FilterQuery.Set<TQueryFunction<bool>>(InFilterQuery);
+}
+
+FTedsFilterData::FTedsFilterData(const FName& InFilterName, const FText& InFilterDisplayName, const TQueryFunction<bool>& InFilterQuery)
+	: FTedsFilterData(
+		InFilterName, 
+		InFilterDisplayName, 
+		FText::FromString(InFilterName.ToString()), 
+		FName(), 
+		nullptr, 
+		InFilterQuery)
+{}
+	
+FTedsOutlinerFilter::FTedsOutlinerFilter(const FTedsFilterData& InTedsFilterData, const TSharedPtr<FTedsOutlinerImpl>& InTedsOutlinerImpl)
+	: FFilterBase(InTedsFilterData.FilterCategory)
+	, FilterName(InTedsFilterData.FilterName)
+	, FilterDisplayName(InTedsFilterData.FilterDisplayName)
+	, FilterToolTip(InTedsFilterData.FilterToolTip)
+	, FilterIconName(InTedsFilterData.FilterIconName)
+	, bIsClassFilter(false)
+	, TedsOutlinerImpl(InTedsOutlinerImpl)
+	, FilterQuery(InTedsFilterData.FilterQuery)
+{}
+
+FTedsOutlinerFilter::FTedsOutlinerFilter(const FName& InFilterName, const FText& InFilterDisplayName, const FText& InFilterToolTip, const FName& InFilterIconName, 
+	const TSharedPtr<FFilterCategory>& InCategory, const TSharedPtr<FTedsOutlinerImpl>& InTedsOutlinerImpl, 
+	const QueryHandle& InFilterQuery)
 	: FFilterBase(InCategory)
 	, FilterName(InFilterName)
 	, FilterDisplayName(InFilterDisplayName)
+	, FilterToolTip(InFilterToolTip)
+	, FilterIconName(InFilterIconName)
+	, bIsClassFilter(false)
 	, TedsOutlinerImpl(InTedsOutlinerImpl)
-	, FilterQuery(InFilterQuery)
 {
-	
+	FilterQuery.Set<QueryHandle>(Helpers::CheckValidFilterQueryHandle(InFilterQuery) ? InFilterQuery : InvalidQueryHandle);
+}
+
+FTedsOutlinerFilter::FTedsOutlinerFilter(const FName& InFilterName, const FText& InFilterDisplayName, const FText& InFilterToolTip, 
+	const FName& InFilterIconName, const TSharedPtr<FFilterCategory>& InCategory, const TSharedPtr<FTedsOutlinerImpl>& InTedsOutlinerImpl,
+	const TQueryFunction<bool>& InFilterQuery)
+	: FFilterBase(InCategory)
+	, FilterName(InFilterName)
+	, FilterDisplayName(InFilterDisplayName)
+	, FilterToolTip(InFilterToolTip)
+	, FilterIconName(InFilterIconName)
+	, bIsClassFilter(false)
+	, TedsOutlinerImpl(InTedsOutlinerImpl)
+{
+	FilterQuery.Set<TQueryFunction<bool>>(InFilterQuery);
 }
 
+FTedsOutlinerFilter::FTedsOutlinerFilter(const UClass* InClass, const TSharedPtr<FFilterCategory>& InCategory, const TSharedPtr<FTedsOutlinerImpl>& InTedsOutlinerImpl)
+	: FFilterBase(InCategory)
+	, FilterName(InClass->GetFName())
+	, FilterDisplayName(InClass->GetDisplayNameText())
+	, FilterToolTip(FText::Format(LOCTEXT("FilterClassTooltip", "Filter by {0}"), InClass->GetDisplayNameText()))
+	, FilterIconName(FSlateIconFinder::FindIconForClass(InClass).GetStyleName())
+	, bIsClassFilter(true)
+	, TedsOutlinerImpl(InTedsOutlinerImpl)
+{
+	FilterQuery.Set<TQueryFunction<bool>>(
+		Queries::BuildQueryFunction<bool>([InClass](TQueryContext<RowBatchInfo> Context, TResult<bool>& Result, TConstBatch<FTypedElementClassTypeInfoColumn> TypeInfoColumns)
+		{
+			Context.ForEachRow([&Result, InClass](const FTypedElementClassTypeInfoColumn& TypeInfoColumn)
+			{
+				Result.Add(TypeInfoColumn.TypeInfo->IsChildOf(InClass));
+			}, TypeInfoColumns);
+		}));
+}
+	
 FString FTedsOutlinerFilter::GetName() const
 {
 	return FilterName.ToString();
@@ -30,7 +125,7 @@ FText FTedsOutlinerFilter::GetDisplayName() const
 
 FText FTedsOutlinerFilter::GetToolTipText() const
 {
-	return FText::FromName(FilterName);
+	return FilterToolTip;
 }
 
 FLinearColor FTedsOutlinerFilter::GetColor() const
@@ -40,7 +135,7 @@ FLinearColor FTedsOutlinerFilter::GetColor() const
 
 FName FTedsOutlinerFilter::GetIconName() const
 {
-	return FName();
+	return FilterIconName;
 }
 
 bool FTedsOutlinerFilter::IsInverseFilter() const
@@ -50,15 +145,41 @@ bool FTedsOutlinerFilter::IsInverseFilter() const
 
 void FTedsOutlinerFilter::ActiveStateChanged(bool bActive)
 {
-	if (TSharedPtr<FTedsOutlinerImpl> TedsOutlinerImplPin = TedsOutlinerImpl.Pin())
+	using namespace UE::Editor::DataStorage::Queries;
+	
+	const TSharedPtr<FTedsOutlinerImpl> TedsOutlinerImplPin = TedsOutlinerImpl.Pin();
+	if (ensureMsgf(TedsOutlinerImplPin, TEXT("No TedsOutliner Context was set for the %s filter."), *FilterDisplayName.ToString()))
 	{
 		if(bActive)
 		{
-			TedsOutlinerImplPin->AddExternalQuery(FilterName, FilterQuery);
+			if (bIsClassFilter)
+			{
+				// Class filters are separated since we want to OR them with each other
+				TedsOutlinerImplPin->AddClassQueryFunction(FilterName, FilterQuery.Get<TQueryFunction<bool>>());
+			}
+			else if (FilterQuery.IsType<QueryHandle>())
+			{
+				TedsOutlinerImplPin->AddExternalQuery(FilterName, FilterQuery.Get<QueryHandle>());
+			}
+			else /* if (FilterQuery.IsType<TQueryFunction<bool>>()) */
+			{
+				TedsOutlinerImplPin->AddExternalQueryFunction(FilterName, FilterQuery.Get<TQueryFunction<bool>>());
+			}
 		}
 		else
 		{
-			TedsOutlinerImplPin->RemoveExternalQuery(FilterName);
+			if (bIsClassFilter)
+			{
+				TedsOutlinerImplPin->RemoveClassQueryFunction(FilterName);
+			}
+			else if (FilterQuery.IsType<QueryHandle>())
+			{
+				TedsOutlinerImplPin->RemoveExternalQuery(FilterName);
+			}
+			else /* if (FilterQuery.IsType<TQueryFunction<bool>>()) */
+			{
+				TedsOutlinerImplPin->RemoveExternalQueryFunction(FilterName);
+			}
 		}
 	}
 }
@@ -78,7 +199,7 @@ void FTedsOutlinerFilter::LoadSettings(const FString& IniFilename, const FString
 	
 }
 
-bool FTedsOutlinerFilter::PassesFilter(SceneOutliner::FilterBarType InItem) const
+bool FTedsOutlinerFilter::PassesFilter(const ISceneOutlinerTreeItem& InItem) const
 {
 	if (TSharedPtr<FTedsOutlinerImpl> TedsOutlinerImplPin = TedsOutlinerImpl.Pin())
 	{
@@ -90,7 +211,9 @@ bool FTedsOutlinerFilter::PassesFilter(SceneOutliner::FilterBarType InItem) cons
 		}
 	}
 
-	// The filter is applied through a TEDS query and this is just a dummy to activate it, so we can simply return true otherwise
+	// The filter is applied through a TEDS query and this is just a dummy to activate it, so we can simply return false otherwise
 	return false;
 }
 } // namespace UE::Editor::Outliner
+
+#undef LOCTEXT_NAMESPACE
diff --git a/Engine/Plugins/Experimental/EditorDataStorageFeatures/Source/TedsOutliner/Private/TedsOutlinerHelpers.cpp b/Engine/Plugins/Experimental/EditorDataStorageFeatures/Source/TedsOutliner/Private/TedsOutlinerHelpers.cpp
--- a/Engine/Plugins/Experimental/EditorDataStorageFeatures/Source/TedsOutliner/Private/TedsOutlinerHelpers.cpp
+++ b/Engine/Plugins/Experimental/EditorDataStorageFeatures/Source/TedsOutliner/Private/TedsOutlinerHelpers.cpp
@@ -4,6 +4,7 @@
 
 #include "Columns/TedsOutlinerColumns.h"
 #include "Compatibility/SceneOutlinerTedsBridge.h"
+#include "DataStorage/Features.h"
 #include "Elements/Interfaces/TypedElementDataStorageInterface.h"
 #include "Modules/ModuleManager.h"
 #include "ISceneOutliner.h"
@@ -83,4 +84,18 @@ namespace UE::Editor::Outliner::Helpers
 
 		return NAME_None;
 	}
+	
+    bool CheckValidFilterQueryHandle(const DataStorage::QueryHandle& InQueryHandle)
+    {
+		using namespace UE::Editor::DataStorage;
+    	// Check if the given Query Handle is valid for a filter (Not an Observer Query Handle)
+    	ICoreProvider* Storage = GetMutableDataStorageFeature<ICoreProvider>(StorageFeatureName);
+    	if (ensureMsgf(Storage, TEXT("TEDS must be initialized before TEDS Filters")))
+    	{
+    		const EQueryCallbackType QueryHandleCallbackType = Storage->GetQueryDescription(InQueryHandle).Callback.Type;
+    		return ensureMsgf(QueryHandleCallbackType == EQueryCallbackType::None || QueryHandleCallbackType == EQueryCallbackType::Processor,
+    			TEXT("TEDS Filters cannot accept Observer Query Handles."));
+    	}
+    	return false;
+    }
 }
diff --git a/Engine/Plugins/Experimental/EditorDataStorageFeatures/Source/TedsOutliner/Private/TedsOutlinerImpl.cpp b/Engine/Plugins/Experimental/EditorDataStorageFeatures/Source/TedsOutliner/Private/TedsOutlinerImpl.cpp
--- a/Engine/Plugins/Experimental/EditorDataStorageFeatures/Source/TedsOutliner/Private/TedsOutlinerImpl.cpp
+++ b/Engine/Plugins/Experimental/EditorDataStorageFeatures/Source/TedsOutliner/Private/TedsOutlinerImpl.cpp
@@ -19,9 +19,11 @@
 #include "TedsOutlinerFilter.h"
 #include "TedsOutlinerHelpers.h"
 #include "TedsOutlinerItem.h"
+#include "TedsQueryHandleNode.h"
 #include "TedsQueryMergeNode.h"
 #include "TedsQueryNode.h"
 #include "TedsRowCopyNode.h"
+#include "TedsRowFilterNode.h"
 #include "TedsRowHandleSortNode.h"
 #include "TedsRowMergeNode.h"
 #include "TedsRowMonitorNode.h"
@@ -34,7 +36,6 @@
 #include "Elements/Columns/TypedElementOverrideColumns.h"
 #include "Elements/Columns/TypedElementSelectionColumns.h"
 #include "Elements/Columns/TypedElementSlateWidgetColumns.h"
-#include "Elements/Framework/TypedElementQueryBuilder.h"
 #include "Filters/FilterBase.h"
 #include "Widgets/SWidget.h"
 
@@ -42,6 +43,16 @@
 
 namespace UE::Editor::Outliner
 {
+namespace Private
+{
+	static bool bShowTedsColumnFilters = false;
+} // namespace Private
+
+static FAutoConsoleVariableRef CVarShowTedsColumnFilters(
+	TEXT("TEDS.UI.Outliner.ShowTedsColumnFilters"),
+	Private::bShowTedsColumnFilters,
+	TEXT("Show Teds Column/Tag Filters in the Outliner (requires TEDS Outliner to be enabled and the Outliner to be reopened)"));
+	
 namespace QueryUtils
 {
 	static bool CanDisplayRow(DataStorage::IQueryContext& Context, const FTedsOutlinerColumn& TedsOutlinerColumn, DataStorage::RowHandle Row, SSceneOutliner& SceneOutliner)
@@ -94,7 +105,6 @@ namespace QueryUtils
 	};
 }
 
-
 namespace LabelWidgetUtils
 {
 	static TUniquePtr<FTypedElementWidgetConstructor> CreateLabelWidgetConstructorInternal(
@@ -176,7 +186,6 @@ namespace LabelWidgetUtils
 FTedsOutlinerParams::FTedsOutlinerParams(SSceneOutliner* InSceneOutliner)
 	: SceneOutliner(InSceneOutliner)
 	, QueryDescription()
-	, bUseDefaultTedsFilters(false)
 	, bShowRowHandleColumn(true)
 	, bForceShowParents(true)
 	, bUseDefaultObservers(true)
@@ -221,7 +230,7 @@ void FTedsOutlinerImpl::CreateFilterQueries()
 {
 	using namespace UE::Editor::DataStorage::Queries;
 
-	if (CreationParams.bUseDefaultTedsFilters)
+	if (Private::bShowTedsColumnFilters)
 	{
 		// Create separate categories for columns and tags
 		TSharedRef<FFilterCategory> TedsColumnFilterCategory = MakeShared<FFilterCategory>(LOCTEXT("TedsColumnFilters", "TEDS Columns"), LOCTEXT("TedsColumnFiltersTooltip", "Filter by TEDS columns"));
@@ -238,42 +247,49 @@ void FTedsOutlinerImpl::CreateFilterQueries()
 				if (Struct->IsChildOf(TedsColumn) || Struct->IsChildOf(TedsTag))
 				{
 					// Create a query description to filter for this tag/column
-					FQueryDescription FilterQueryDesc;
+					QueryHandle FilterQuery;
 
 					if (bUsingQueryConditionsSyntax)
 					{
-						FilterQueryDesc = Select()
+						FilterQuery = Storage->RegisterQuery(
+							Select()
 							.Where(TColumn(Struct))
-							.Compile();
+							.Compile());
 					}
 					else
 					{
-						FilterQueryDesc =
+						FilterQuery = Storage->RegisterQuery(
 							Select()
 							.Where()
 								.All(Struct)
-							.Compile();
+							.Compile());
 					}
 
 					// Create the filter
-					TSharedRef<FTedsOutlinerFilter> TedsFilter = MakeShared<FTedsOutlinerFilter>(Struct->GetFName(), Struct->GetDisplayNameText(),
-						Struct->IsChildOf(TedsColumn) ? TedsColumnFilterCategory : TedsTagFilterCategory, AsShared(), FilterQueryDesc);
+					TSharedRef<FTedsOutlinerFilter> TedsFilter = MakeShared<FTedsOutlinerFilter>(Struct->GetFName(), Struct->GetDisplayNameText(), 
+						Struct->GetDisplayNameText(), FName(), Struct->IsChildOf(TedsColumn) ? TedsColumnFilterCategory : TedsTagFilterCategory, 
+						AsShared(), FilterQuery);
 					SceneOutliner->AddFilterToFilterBar(TedsFilter);
 				}
 			}
 		});
 	}
 
-	// Custom filters input by the user
-	TSharedRef<FFilterCategory> CustomFiltersCategory = MakeShared<FFilterCategory>(LOCTEXT("TedsFilters", "TEDS Custom Filters"), LOCTEXT("TedsFiltersTooltip", "Filter by custom TEDS queries"));
-
-	for(const TPair<FName, const DataStorage::FQueryDescription>& FilterQuery : CreationParams.FilterQueries)
+	TSharedRef<FFilterCategory> CustomClassFiltersCategory = MakeShared<FFilterCategory>(LOCTEXT("CustomTypeFilters", "Custom Type Filters"), LOCTEXT("CustomTypeFiltersTooltip", "Filter by custom class types"));
+	for(UClass* Class : CreationParams.ClassFilters)
 	{
-		// TEDS-Outliner TODO: Custom filters need a localizable display name instead of using the FName, but we need to change how they are added first
-		// to see if it can be consolidated with the SFilterBar API
-		TSharedRef<FTedsOutlinerFilter> TedsFilter = MakeShared<FTedsOutlinerFilter>(FilterQuery.Key, FText::FromName(FilterQuery.Key), CustomFiltersCategory, AsShared(), FilterQuery.Value);
-		SceneOutliner->AddFilterToFilterBar(TedsFilter);
+		SceneOutliner->AddFilterToFilterBar(MakeShared<FTedsOutlinerFilter>(Class, CustomClassFiltersCategory, AsShared()));
+	}
 	
+	TSharedRef<FFilterCategory> CustomFiltersCategory = MakeShared<FFilterCategory>(LOCTEXT("TedsFilters", "TEDS Custom Filters"), LOCTEXT("TedsFiltersTooltip", "Filter by custom TEDS queries"));
+		
+	for(FTedsFilterData& Filter : CreationParams.Filters)
+	{
+		if (!Filter.FilterCategory)
+		{
+			Filter.FilterCategory = CustomFiltersCategory;
+		}
+		SceneOutliner->AddFilterToFilterBar(MakeShared<FTedsOutlinerFilter>(Filter, AsShared()));
 	}
 }
 
@@ -296,27 +312,35 @@ void FTedsOutlinerImpl::Init()
 
 void FTedsOutlinerImpl::OnMonitoredRowsAdded(DataStorage::FRowHandleArrayView InRows)
 {
-	for (DataStorage::RowHandle Row : InRows)
+	// Refresh queries so that the added rows can pass through the filters if applied
+	RefreshQueries();
+
+	const bool bHasNoActiveFilters = FilterNodes.IsEmpty() && ClassFilterNodes.IsEmpty();
+	for (DataStorage::RowHandle RowToAdd : InRows)
 	{
-		// If the row we are adding is selected, we are probably synced to some external selection and the Outliner needs to update selection
-		if (Storage->HasColumns<FTypedElementSelectionColumn>(Row))
+		// Don't display the new row unless it is in the filtered row list, but there is no need to check if there are no filters active
+		if (bHasNoActiveFilters || FinalCombinedRowNode->GetRows().Contains(RowToAdd))
 		{
-			bSelectionDirty = true;
+			// If the row we are adding is selected, we are probably synced to some external selection and the Outliner needs to update selection
+			if (Storage->HasColumns<FTypedElementSelectionColumn>(RowToAdd))
+			{
+				bSelectionDirty = true;
+			}
+			FSceneOutlinerHierarchyChangedData EventData;
+			EventData.Type = FSceneOutlinerHierarchyChangedData::Added;
+			EventData.Items.Add(SceneOutlinerMode->CreateItemFor<FTedsOutlinerTreeItem>(FTedsOutlinerTreeItem(RowToAdd, AsShared())));
+			HierarchyChangedEvent.Broadcast(EventData);
 		}
-		FSceneOutlinerHierarchyChangedData EventData;
-		EventData.Type = FSceneOutlinerHierarchyChangedData::Added;
-		EventData.Items.Add(SceneOutlinerMode->CreateItemFor<FTedsOutlinerTreeItem>(FTedsOutlinerTreeItem(Row, AsShared())));
-		HierarchyChangedEvent.Broadcast(EventData);
 	}
 }
 
 void FTedsOutlinerImpl::OnMonitoredRowsRemoved(DataStorage::FRowHandleArrayView InRows)
 {
-	for (DataStorage::RowHandle Row : InRows)
+	for (DataStorage::RowHandle RowToRemove : InRows)
 	{
 		FSceneOutlinerHierarchyChangedData EventData;
 		EventData.Type = FSceneOutlinerHierarchyChangedData::Removed;
-		EventData.ItemIDs.Add(Row);
+		EventData.ItemIDs.Add(RowToRemove);
 		HierarchyChangedEvent.Broadcast(EventData);
 	}
 }
@@ -386,28 +410,48 @@ TSharedRef<SWidget>FTedsOutlinerImpl::CreateLabelWidgetForItem(DataStorage::RowH
 	return CreateLabelWidget(*Storage, *StorageUi, LabelWidgetPurpose, TargetRow, TreeItem, RowItem, bIsInteractable);
 }
 
-void FTedsOutlinerImpl::AddExternalQuery(FName QueryName, const DataStorage::FQueryDescription& InQueryDescription)
+void FTedsOutlinerImpl::AddExternalQuery(const FName& QueryName, const QueryHandle& InQueryHandle)
 {
-	ExternalQueries.Emplace(QueryName, MakeShared<DataStorage::QueryStack::FQueryNode>(*Storage, InQueryDescription));
-
-	RecompileQueries();
+	if (ensureMsgf(InQueryHandle != InvalidQueryHandle, TEXT("An Invalid Query Handle cannot be used for a TEDS Filter")))
+	{
+		ExternalQueries.Emplace(QueryName, MakeShared<DataStorage::QueryStack::FQueryHandleNode>(InQueryHandle));
+	}
 }
 
-void FTedsOutlinerImpl::RemoveExternalQuery(FName QueryName)
+void FTedsOutlinerImpl::RemoveExternalQuery(const FName& QueryName)
 {
 	ExternalQueries.Remove(QueryName);
 }
 
-
-
-void FTedsOutlinerImpl::AppendExternalQueries(DataStorage::FQueryDescription& OutQuery)
+void FTedsOutlinerImpl::AppendExternalQueries(FQueryDescription& OutQuery)
 {
-	for(const TPair<FName, TSharedPtr<DataStorage::QueryStack::IQueryNode>>& ExternalQuery : ExternalQueries)
+	for(const TPair<FName, TSharedPtr<QueryStack::IQueryNode>>& ExternalQuery : ExternalQueries)
 	{
 		DataStorage::Queries::MergeQueries(OutQuery, Storage->GetQueryDescription(ExternalQuery.Value->GetQuery()));
 	}
 }
 
+void FTedsOutlinerImpl::AddExternalQueryFunction(const FName& QueryName, const DataStorage::Queries::TQueryFunction<bool>& InQueryFunction)
+{
+	// Store the external query functions as functions instead of RowNodes since the row node to filter on has not been initialized yet
+	ExternalQueryFunctions.Emplace(QueryName, InQueryFunction);
+}
+
+void FTedsOutlinerImpl::RemoveExternalQueryFunction(const FName& QueryName)
+{
+	ExternalQueryFunctions.Remove(QueryName);
+}
+
+void FTedsOutlinerImpl::AddClassQueryFunction(const FName& ClassName, const TQueryFunction<bool>& InClassQueryFunction)
+{
+	ClassFilters.Emplace(ClassName, InClassQueryFunction);
+}
+
+void FTedsOutlinerImpl::RemoveClassQueryFunction(const FName& ClassName)
+{
+	ClassFilters.Remove(ClassName);
+}
+
 bool FTedsOutlinerImpl::CanDisplayRow(DataStorage::RowHandle ItemRowHandle) const
 {
 	/*
@@ -588,16 +632,10 @@ void FTedsOutlinerImpl::CreateItemsFromQuery(TArray<FSceneOutlinerTreeItemPtr>&
 	// update. Addition/Removal in between refreshes is handled by the row monitor node
 	// Technically the RowMonitorNode also handles keeping this list up to date using observers, but the monitor node is optional and it's safer
 	// to do a full clean refresh at this point anyways
-	RowCollectorNode->Refresh();
-	FRowHandleArrayView Rows = RowCollectorNode->GetRows();
 	
-	for (RowHandle Row : Rows)
+	RefreshQueries();
+	for (RowHandle Row : FinalCombinedRowNode->GetRows())
 	{
-		if (!CanDisplayRow(Row))
-		{
-			continue;
-		}
-		
 		if (FSceneOutlinerTreeItemPtr TreeItem = InMode->CreateItemFor<FTedsOutlinerTreeItem>(FTedsOutlinerTreeItem(Row, AsShared()), false))
 		{
 			OutItems.Add(TreeItem);
@@ -606,7 +644,6 @@ void FTedsOutlinerImpl::CreateItemsFromQuery(TArray<FSceneOutlinerTreeItemPtr>&
 
 	// We need to update selection since the whole outliner was re-populated
 	bSelectionDirty = true;
-
 }
 
 void FTedsOutlinerImpl::CreateChildren(const FSceneOutlinerTreeItemPtr& Item, TArray<FSceneOutlinerTreeItemPtr>& OutChildren) const
@@ -696,7 +733,7 @@ void FTedsOutlinerImpl::CreateChildren(const FSceneOutlinerTreeItemPtr& Item, TA
 	// Actually create the items for the child entities 
 	for (RowHandle ChildItemRowHandle : ChildItems)
 	{
-		if (!CanDisplayRow(ChildItemRowHandle) || !Rows.Contains(ChildItemRowHandle))
+		if (!Rows.Contains(ChildItemRowHandle))
 		{
 			continue;
 		}
@@ -732,11 +769,6 @@ DataStorage::RowHandle FTedsOutlinerImpl::GetParentRow(DataStorage::RowHandle In
 		return DataStorage::InvalidRowHandle;
 	}
 
-	if (!CanDisplayRow(ParentRowHandle))
-	{
-		return DataStorage::InvalidRowHandle;
-	}
-
 	// If the row doesn't exist in our query stack results, it doesn't match the query and should not be displayed
 	DataStorage::FRowHandleArrayView Rows(RowHandleSortNode->GetRows());
 
@@ -748,22 +780,31 @@ DataStorage::RowHandle FTedsOutlinerImpl::GetParentRow(DataStorage::RowHandle In
 	return ParentRowHandle;
 }
 
-bool FTedsOutlinerImpl::ShouldForceShowParentRows()
+bool FTedsOutlinerImpl::ShouldForceShowParentRows() const
 {
 	return bForceShowParents;
 }
 
+void FTedsOutlinerImpl::RefreshQueries() const
+{
+	InitialRowCollectorNode->Refresh();
+	CombinedRowCollectorNode->Refresh();
+	FinalCombinedRowNode->Update();
+}
+
 void FTedsOutlinerImpl::RecompileQueries()
 {
 	using namespace UE::Editor::DataStorage::Queries;
 	using namespace UE::Editor::DataStorage::QueryStack;
-
+	
 	UnregisterQueries();
 
-	// Get a list of all the query nodes (initial query + filters) to combine them into one composite query
+	// Get a list of all the query nodes (initial query + query description filters) to combine them into one composite query
 	TArray<TSharedPtr<IQueryNode>> QueryNodes;
-	QueryNodes.Add(MakeShared<FQueryNode>(*Storage, InitialQueryDescription));
 
+	TSharedPtr<IQueryNode> InitialQueryNode = MakeShared<FQueryNode>(*Storage, InitialQueryDescription);
+	
+	QueryNodes.Add(InitialQueryNode);
 	for (const TPair<FName, TSharedPtr<IQueryNode>>& ExternalQuery : ExternalQueries)
 	{
 		QueryNodes.Add(ExternalQuery.Value);
@@ -771,20 +812,52 @@ void FTedsOutlinerImpl::RecompileQueries()
 
 	TSharedPtr<IQueryNode> CombinedQueryNode = MakeShared<FQueryMergeNode>(*Storage, QueryNodes);
 
-	// If we have a monitor node, it handles updating the collector node automatically. Otherwise we need it to update on tick
+	// If we have a monitor node, it handles updating the collector node automatically. Otherwise, we need it to update on tick
 	FRowQueryResultsNode::ESyncFlags SyncFlags =
 		CreationParams.bUseDefaultObservers ? FRowQueryResultsNode::ESyncFlags::None : FRowQueryResultsNode::ESyncFlags::RefreshOnUpdate;
-	
+
+	// Node to collect all rows matching the initial query (stored so OR filters can filter on the original query)
+	InitialRowCollectorNode = MakeShared<FRowQueryResultsNode>(*Storage, InitialQueryNode, SyncFlags);
+
 	// Node to collect all rows matching the composite query
-	RowCollectorNode = MakeShared<FRowQueryResultsNode>(*Storage, CombinedQueryNode, SyncFlags);
+	CombinedRowCollectorNode = MakeShared<FRowQueryResultsNode>(*Storage, CombinedQueryNode, SyncFlags);
 
-	// The "actual" collector node passed into other misc nodes is the monitor node if it exists, the collector node otherwise 
-	TSharedPtr<IRowNode> ActualCollectorNode = RowCollectorNode;
+	FilterNodes.Empty();
+	TSharedPtr<IRowNode> PrimaryFilterNode = CombinedRowCollectorNode;
+	// Combines all active external filters using AND
+	for (const TPair<FName, TQueryFunction<bool>>& ExternalQueryFunction : ExternalQueryFunctions)
+	{
+		TSharedPtr<FRowFilterNode> FilterNode = MakeShared<FRowFilterNode>(Storage, PrimaryFilterNode, ExternalQueryFunction.Value);
+		FilterNodes.Add(FilterNode);
+		PrimaryFilterNode = FilterNode;
+	}
+
+	ClassFilterNodes.Empty();
+	// Combines all active class filters using OR
+	if (!ClassFilters.IsEmpty())
+	{
+		for (const TPair<FName, TQueryFunction<bool>>& ClassFilter : ClassFilters)
+		{
+			ClassFilterNodes.Add(MakeShared<FRowFilterNode>(Storage, InitialRowCollectorNode, ClassFilter.Value));
+		}
+		
+		TArray<TSharedPtr<IRowNode>> ClassBaseFilterNodes(ClassFilterNodes);
+		CombinedClassFilterRowNode = MakeShared<FRowMergeNode>(ClassBaseFilterNodes, FRowMergeNode::EMergeApproach::Unique);
+
+		FinalCombinedRowNode = MakeShared<FRowMergeNode>(MakeConstArrayView({PrimaryFilterNode, CombinedClassFilterRowNode}), FRowMergeNode::EMergeApproach::Repeating);
+	}
+	else
+	{
+		FinalCombinedRowNode = PrimaryFilterNode;
+	}
+
+	// The "actual" collector node passed into other misc nodes is the monitor node if it exists, the collector node otherwise
+	TSharedPtr<IRowNode> ActualCollectorNode = FinalCombinedRowNode;
 	
 	// Node to monitor row addition/removal
 	if (CreationParams.bUseDefaultObservers)
 	{
-		RowMonitorNode = MakeShared<FRowMonitorNode>(*Storage, RowCollectorNode, CombinedQueryNode);
+		RowMonitorNode = MakeShared<FRowMonitorNode>(*Storage, FinalCombinedRowNode, CombinedQueryNode);
 		RowMonitorNode->OnMonitoredRowsAdded().AddRaw(this, &FTedsOutlinerImpl::OnMonitoredRowsAdded);
 		RowMonitorNode->OnMonitoredRowsRemoved().AddRaw(this, &FTedsOutlinerImpl::OnMonitoredRowsRemoved);
 		ActualCollectorNode = RowMonitorNode;
@@ -852,11 +925,6 @@ void FTedsOutlinerImpl::RecompileQueries()
 								
 						if (QueryUtils::HasItemParentChanged(Context, Rows[RowIndex], ParentRowHandle, *SceneOutliner))
 						{
-							if (TedsOutlinerColumns && !QueryUtils::CanDisplayRow(Context, TedsOutlinerColumns[RowIndex], Rows[RowIndex], *SceneOutliner))
-							{
-								continue;
-							}
-						
 							FSceneOutlinerHierarchyChangedData EventData;
 							EventData.Type = FSceneOutlinerHierarchyChangedData::Moved;
 							EventData.ItemIDs.Add(Rows[RowIndex]);
diff --git a/Engine/Plugins/Experimental/EditorDataStorageFeatures/Source/TedsOutliner/Private/TedsOutlinerModule.cpp b/Engine/Plugins/Experimental/EditorDataStorageFeatures/Source/TedsOutliner/Private/TedsOutlinerModule.cpp
--- a/Engine/Plugins/Experimental/EditorDataStorageFeatures/Source/TedsOutliner/Private/TedsOutlinerModule.cpp
+++ b/Engine/Plugins/Experimental/EditorDataStorageFeatures/Source/TedsOutliner/Private/TedsOutlinerModule.cpp
@@ -13,15 +13,20 @@
 #include "Elements/Columns/TypedElementPackageColumns.h"
 #include "Elements/Columns/TypedElementTypeInfoColumns.h"
 #include "Elements/Columns/TypedElementVisibilityColumns.h"
+#include "Columns/TedsActorComponentCompatibilityColumns.h"
 #include "Columns/TedsActorMobilityColumns.h"
+#include "TedsAlertColumns.h"
 #include "DataStorage/Features.h"
 #include "Elements/Interfaces/TypedElementDataStorageInterface.h"
+#include "Elements/Framework/TypedElementQueryContext.h"
 #include "Modules/ModuleManager.h"
 #include "Compatibility/SceneOutlinerTedsBridge.h"
 #include "Compatibility/SceneOutlinerRowHandleColumn.h"
-#include "TedsAlertColumns.h"
 #include "Widgets/Docking/SDockTab.h"
 
+#include "Engine/DirectionalLight.h"
+#include "Engine/StaticMeshActor.h"
+
 #define LOCTEXT_NAMESPACE "TedsOutlinerModule"
 
 namespace UE::Editor::Outliner
@@ -70,7 +75,7 @@ TSharedRef<ISceneOutliner> FTedsOutlinerModule::CreateTedsOutliner(const FSceneO
 {
 	using namespace UE::Editor::DataStorage;
 	ensureMsgf(AreEditorDataStorageFeaturesEnabled(), TEXT("Unable to initialize the Teds-Outliner before TEDS itself is initialized."));
-
+	
 	FSceneOutlinerInitializationOptions InitOptions(InInitOptions);
 	FTedsOutlinerParams InitTedsOptions(InInitTedsOptions);
 
@@ -117,7 +122,7 @@ UE::Editor::DataStorage::FQueryDescription FTedsOutlinerModule::GetLevelEditorTe
 {
 	using namespace UE::Editor::DataStorage;
 	using namespace UE::Editor::DataStorage::Columns;
-	ICoreProvider* Storage = GetMutableDataStorageFeature<ICoreProvider>(StorageFeatureName);
+	Storage = GetMutableDataStorageFeature<ICoreProvider>(StorageFeatureName);
 		
 	using namespace DataStorage::Queries;
 
@@ -173,15 +178,38 @@ TSharedRef<SWidget> FTedsOutlinerModule::CreateLevelEditorTedsOutliner()
 	FTedsOutlinerParams Params(nullptr);
 	Params.QueryDescription = OutlinerQueryDescription;
 	Params.ColumnQueryDescription = GetLevelEditorTedsOutlinerColumnQueryDescription();
-	Params.bUseDefaultTedsFilters = true;
 
-	// Example Query to filter for actors
-	DataStorage::FQueryDescription ActorFilterQuery =
-		Select()
-		.Where()
-			.All<FTypedElementActorTag>()
-		.Compile();
-	Params.FilterQueries.Emplace("Actors", ActorFilterQuery);
+	Params.ClassFilters.Add(AActor::StaticClass());
+	Params.ClassFilters.Add(AStaticMeshActor::StaticClass());
+	Params.ClassFilters.Add(ADirectionalLight::StaticClass());
+	
+	// Example ways to create filters using query descriptions and functions
+	// Note QueryDescriptions should always be used over QueryFunctions when possible as they are more performant
+	Params.Filters.Emplace(
+		"ActorComponents", 
+		LOCTEXT("ActorFilterDisplayName", "Actor Components"), 
+		LOCTEXT("ActorFilterToolTip", "Filters for actor components"), 
+		"ClassIcon.Actor", 
+		nullptr,
+		Storage->RegisterQuery(
+			Select()
+			.Where()
+				.All<FActorComponentTypeTag>()
+			.Compile()));
+
+	// TODO: Make it possible to set EFunctionCallConfig::VerifyColumns in BuildQueryFunction as we will usually want to exclude rows that 
+	// don't have the columns we want to filter (MobilityColumn=Static will only filter on columns that have the Mobility Column and
+	// return true for rows that don't have that column), perhaps add it to EFilterOptions?
+	Params.Filters.Emplace(
+		"Static",
+		LOCTEXT("StaticFilterDisplayName", "Static"), 
+		LOCTEXT("StaticFilterToolTip", "Filter for Static Actors"),
+		FName(), 
+		nullptr,
+		[](TQueryContext<SingleRowInfo> Context, const FTedsActorMobilityColumn& MobilityColumn)
+		{
+			return MobilityColumn.Mobility == EComponentMobility::Static;
+		});
 		
 	// Empty selection set name is currently the level editor
 	Params.SelectionSetOverride = FName();
diff --git a/Engine/Plugins/Experimental/EditorDataStorageFeatures/Source/TedsOutliner/Public/TedsOutlinerFilter.h b/Engine/Plugins/Experimental/EditorDataStorageFeatures/Source/TedsOutliner/Public/TedsOutlinerFilter.h
--- a/Engine/Plugins/Experimental/EditorDataStorageFeatures/Source/TedsOutliner/Public/TedsOutlinerFilter.h
+++ b/Engine/Plugins/Experimental/EditorDataStorageFeatures/Source/TedsOutliner/Public/TedsOutlinerFilter.h
@@ -2,62 +2,131 @@
 
 #pragma once
 
+#include <type_traits>
 #include "Filters/FilterBase.h"
 #include "SceneOutlinerPublicTypes.h"
 #include "DataStorage/Queries/Description.h"
-#include "Elements/Framework/TypedElementMetaData.h"
+#include "Elements/Framework/TypedElementQueryBuilder.h"
 
 namespace UE::Editor::Outliner
 {
-class FTedsOutlinerImpl;
-
-// TEDS-Outliner TODO: Lots of minor missing functionality for icon, tooltip, color etc
-class FTedsOutlinerFilter : public FFilterBase<SceneOutliner::FilterBarType>
-{
-public:
-	FTedsOutlinerFilter(const FName& InFilterName, const FText& InFilterDisplayName, TSharedPtr<FFilterCategory> InCategory,
-		TSharedRef<FTedsOutlinerImpl> InTedsOutlinerImpl, const DataStorage::FQueryDescription& InFilterQuery);
-
-	/** Returns the system name for this filter */
-	virtual FString GetName() const override;
-
-	/** Returns the human readable name for this filter */
-	virtual FText GetDisplayName() const override;
-
-	/** Returns the tooltip for this filter, shown in the filters menu */
-	virtual FText GetToolTipText() const override;
-
-	/** Returns the color this filter button will be when displayed as a button */
-	virtual FLinearColor GetColor() const override;
-
-	/** Returns the name of the icon to use in menu entries */
-	virtual FName GetIconName() const override;
-
-	/** If true, the filter will be active in the FilterBar when it is inactive in the UI (i.e the filter pill is grayed out)
-	 */
-	virtual bool IsInverseFilter() const override;
-
-	/** Notification that the filter became active or inactive */
-	virtual void ActiveStateChanged(bool bActive) override;
-
-	/** Called when the right-click context menu is being built for this filter */
-	virtual void ModifyContextMenu(FMenuBuilder& MenuBuilder) override;
-	
-	/** Can be overriden for custom FilterBar subclasses to save settings, currently not implemented in any generic Filter Bar */
-	virtual void SaveSettings(const FString& IniFilename, const FString& IniSection, const FString& SettingsString) const override;
-
-	/** Can be overridden for custom FilterBar subclasses to load settings, currently not implemented in any generic Filter Bar */
-	virtual void LoadSettings(const FString& IniFilename, const FString& IniSection, const FString& SettingsString) override;
-	
-	/** Returns whether the specified Item passes the Filter's restrictions */
-	virtual bool PassesFilter( SceneOutliner::FilterBarType InItem ) const override;
-
-protected:
-
-	FName FilterName;
-	FText FilterDisplayName;
-
-	TWeakPtr<FTedsOutlinerImpl> TedsOutlinerImpl;
-	const DataStorage::FQueryDescription FilterQuery;
-};
+	using namespace UE::Editor::DataStorage::Queries;
+	class FTedsOutlinerImpl;
+
+	struct TEDSOUTLINER_API FTedsFilterData
+	{
+		FTedsFilterData(const FName& InFilterName, const FText& InFilterDisplayName, const FText& InFilterToolTip, const FName& InFilterIconName,
+			const TSharedPtr<FFilterCategory>& InCategory, const QueryHandle& InFilterQuery);
+
+		FTedsFilterData(const FName& InFilterName, const FText& InFilterDisplayName, const QueryHandle& InFilterQuery);
+		
+		FTedsFilterData(const FName& InFilterName, const FText& InFilterDisplayName, const FText& InFilterToolTip, const FName& InFilterIconName,
+			const TSharedPtr<FFilterCategory>& InCategory, const TQueryFunction<bool>& InFilterQuery);
+
+		FTedsFilterData(const FName& InFilterName, const FText& InFilterDisplayName, const TQueryFunction<bool>& InFilterQuery);
+
+		/*
+		 * Require that the template type cannot be one of the pre-defined FilterQuery types so it will pass to the defined
+		 * constructor instead of this one.
+		 */
+		template< typename FilterFunction UE_REQUIRES(
+			!std::is_same_v<std::remove_cvref_t<FilterFunction>, QueryHandle> &&
+			!std::is_same_v<std::remove_cvref_t<FilterFunction>, TQueryFunction<bool>>) >
+		FTedsFilterData(const FName& InFilterName, const FText& InFilterDisplayName, const FText& InFilterToolTip, const FName& InFilterIconName,
+			const TSharedPtr<FFilterCategory>& InCategory, FilterFunction&& InFilterQuery)
+			: FTedsFilterData(
+				InFilterName, 
+				InFilterDisplayName, 
+				InFilterToolTip, 
+				InFilterIconName, 
+				InCategory,
+				Queries::BuildQueryFunction<bool>(Forward<FilterFunction>(InFilterQuery)))
+		{}
+
+		/*
+		 * Require that the template type cannot be one of the pre-defined FilterQuery types so it will pass to the defined
+		 * constructor instead of this one.
+		 */
+		template< typename FilterFunction UE_REQUIRES(
+			!std::is_same_v<std::remove_cvref_t<FilterFunction>, QueryHandle> &&
+			!std::is_same_v<std::remove_cvref_t<FilterFunction>, TQueryFunction<bool>>) >
+	    FTedsFilterData(const FName& InFilterName, const FText& InFilterDisplayName, FilterFunction&& InFilterQuery)
+			: FTedsFilterData(
+				InFilterName, 
+				InFilterDisplayName, 
+				FText::FromString(InFilterName.ToString()), 
+				FName(), 
+				nullptr, 
+				Queries::BuildQueryFunction<bool>(Forward<FilterFunction>(InFilterQuery)))
+		{}
+		
+		FName FilterName;
+		FText FilterDisplayName;
+		FText FilterToolTip;
+		FName FilterIconName;
+		
+		// The Category can be set as a nullptr in FTedsFilterData to go to the default 'Other Filters' category when adding via the 
+		// FTedsOutlinerParams.Filters list, otherwise the filter will not display if it's not given a category
+		TSharedPtr<FFilterCategory> FilterCategory;
+		TVariant<QueryHandle, TQueryFunction<bool>> FilterQuery;
+	};
+		
+	class TEDSOUTLINER_API FTedsOutlinerFilter : public FFilterBase<SceneOutliner::FilterBarType>
+	{
+	public:
+		FTedsOutlinerFilter(const FTedsFilterData& InTedsFilterData, const TSharedPtr<FTedsOutlinerImpl>& InTedsOutlinerImpl);
+
+		FTedsOutlinerFilter(const FName& InFilterName, const FText& InFilterDisplayName, const FText& InFilterToolTip, const FName& InFilterIconName, 
+			const TSharedPtr<FFilterCategory>& InCategory, const TSharedPtr<FTedsOutlinerImpl>& InTedsOutlinerImpl, const QueryHandle& InFilterQuery);
+		
+		FTedsOutlinerFilter(const FName& InFilterName, const FText& InFilterDisplayName, const FText& InFilterToolTip, const FName& InFilterIconName,
+			const TSharedPtr<FFilterCategory>& InCategory, const TSharedPtr<FTedsOutlinerImpl>& InTedsOutlinerImpl, const TQueryFunction<bool>& InFilterQuery);
+
+		FTedsOutlinerFilter(const UClass* InClass, const TSharedPtr<FFilterCategory>& InCategory, const TSharedPtr<FTedsOutlinerImpl>& InTedsOutlinerImpl);
+
+		/** Returns the system name for this filter */
+		virtual FString GetName() const override;
+
+		/** Returns the human-readable name for this filter */
+		virtual FText GetDisplayName() const override;
+
+		/** Returns the tooltip for this filter, shown in the filters menu */
+		virtual FText GetToolTipText() const override;
+
+		/** Returns the color this filter button will be when displayed as a button */
+		virtual FLinearColor GetColor() const override;
+
+		/** Returns the name of the icon to use in menu entries */
+		virtual FName GetIconName() const override;
+
+		/** If true, the filter will be active in the FilterBar when it is inactive in the UI (i.e the filter pill is grayed out) */
+		virtual bool IsInverseFilter() const override;
+
+		/** Notification that the filter became active or inactive */
+		virtual void ActiveStateChanged(bool bActive) override;
+
+		/** Called when the right-click context menu is being built for this filter */
+		virtual void ModifyContextMenu(FMenuBuilder& MenuBuilder) override;
+		
+		/** Can be overriden for custom FilterBar subclasses to save settings, currently not implemented in any generic Filter Bar */
+		virtual void SaveSettings(const FString& IniFilename, const FString& IniSection, const FString& SettingsString) const override;
+
+		/** Can be overridden for custom FilterBar subclasses to load settings, currently not implemented in any generic Filter Bar */
+		virtual void LoadSettings(const FString& IniFilename, const FString& IniSection, const FString& SettingsString) override;
+		
+		/** Returns whether the specified Item passes the Filter's restrictions */
+		virtual bool PassesFilter(const ISceneOutlinerTreeItem& InItem ) const override;
+
+	protected:
+
+		FName FilterName;
+		FText FilterDisplayName;
+		FText FilterToolTip;
+		FName FilterIconName;
+
+		const bool bIsClassFilter;
+
+		TWeakPtr<FTedsOutlinerImpl> TedsOutlinerImpl;
+		TVariant<QueryHandle, TQueryFunction<bool>> FilterQuery;
+	};
 } // namespace UE::Editor::Outliner
diff --git a/Engine/Plugins/Experimental/EditorDataStorageFeatures/Source/TedsOutliner/Public/TedsOutlinerHelpers.h b/Engine/Plugins/Experimental/EditorDataStorageFeatures/Source/TedsOutliner/Public/TedsOutlinerHelpers.h
--- a/Engine/Plugins/Experimental/EditorDataStorageFeatures/Source/TedsOutliner/Public/TedsOutlinerHelpers.h
+++ b/Engine/Plugins/Experimental/EditorDataStorageFeatures/Source/TedsOutliner/Public/TedsOutlinerHelpers.h
@@ -25,4 +25,7 @@ namespace UE::Editor::Outliner::Helpers
 	
 	// Get the name of the Outliner column corresponding to the given TEDS column (if any)
 	FName FindOutlinerColumnFromTedsColumns(const DataStorage::ICoreProvider* Storage, TConstArrayView<TWeakObjectPtr<const UScriptStruct>> TEDSColumns);
+	
+	// Helper function to check if the given QueryHandle is valid to be used in a filter (is not an observer).
+    bool CheckValidFilterQueryHandle(const DataStorage::QueryHandle& InQueryHandle);
 }
diff --git a/Engine/Plugins/Experimental/EditorDataStorageFeatures/Source/TedsOutliner/Public/TedsOutlinerImpl.h b/Engine/Plugins/Experimental/EditorDataStorageFeatures/Source/TedsOutliner/Public/TedsOutlinerImpl.h
--- a/Engine/Plugins/Experimental/EditorDataStorageFeatures/Source/TedsOutliner/Public/TedsOutlinerImpl.h
+++ b/Engine/Plugins/Experimental/EditorDataStorageFeatures/Source/TedsOutliner/Public/TedsOutlinerImpl.h
@@ -9,6 +9,8 @@
 #include "Elements/Columns/TypedElementHiearchyColumns.h"
 #include "Elements/Framework/TypedElementRowHandleArrayView.h"
 #include "Elements/Interfaces/TypedElementDataStorageUiInterface.h"
+#include "Elements/Framework/TypedElementQueryBuilder.h"
+#include "TedsOutlinerFilter.h"
 
 #define UE_API TEDSOUTLINER_API
 
@@ -22,6 +24,8 @@ namespace UE::Editor::DataStorage
 		class FRowQueryResultsNode;
 		class IQueryNode;
 		class FRowMonitorNode;
+		class FRowMergeNode;
+		class FRowFilterNode;
 		class IRowNode;
 	}
 
@@ -104,19 +108,17 @@ struct FTedsOutlinerParams
 
 	SSceneOutliner* SceneOutliner;
 
-	// The query description that will be used to populate rows in the TEDS-Outliner
+	// The query description that will be used to populate rows in the TEDS Outliner
 	DataStorage::FQueryDescription QueryDescription;
 
-	// The query description that will be used to populate the initial columns in the Teds Outliner
+	// The query description that will be used to populate the initial columns in the TEDS Outliner
 	DataStorage::FQueryDescription ColumnQueryDescription;
 	
-	// TEDS queries that will be used to create filters in this Outliner
-	// TEDS-Outliner TODO: Can we consolidate this with the SceneOutliner API to create filters? Currently has to be separate because FTEDSOutlinerFilter
-	// needs a reference to the mode which is not possible since filters with the Outliner API are added before the mode is init
-	TMap<FName, const DataStorage::FQueryDescription> FilterQueries;
+	// TEDS Filters that utilize QueryDescriptions or QueryFunctions in this TEDS Outliner
+	TArray<FTedsFilterData> Filters;
 
-	// If true, this Outliner will automatically add all TEDS tags and columns as filters
-	bool bUseDefaultTedsFilters;
+	// Class types to create filters for in this TEDS Outliner
+	TArray<UClass*> ClassFilters;
 
 	// If true, this Outliner will include a column for row handle
 	bool bShowRowHandleColumn;
@@ -195,12 +197,20 @@ public:
 	// Get the hierarchy data associated with this table viewer
 	UE_API const TOptional<FTedsOutlinerHierarchyData>& GetHierarchyData();
 	
-	// Add an external query to the Outliner
-	UE_API void AddExternalQuery(FName QueryName, const DataStorage::FQueryDescription& InQueryDescription);
-	UE_API void RemoveExternalQuery(FName QueryName);
+	// Add an external query description to the Outliner
+	UE_API void AddExternalQuery(const FName& QueryName, const QueryHandle& InQueryHandle);
+	UE_API void RemoveExternalQuery(const FName& QueryName);
 
-	// Append all external queries into the given query
+	// Append all external queries into the given query description
 	UE_API void AppendExternalQueries(DataStorage::FQueryDescription& OutQuery);
+
+	// Add an external query function to the Outliner
+	UE_API void AddExternalQueryFunction(const FName& QueryName, const TQueryFunction<bool>& InQueryFunction);
+	UE_API void RemoveExternalQueryFunction(const FName& QueryName);
+
+	// Add an external class query function to the Outliner (uses OR) 
+	UE_API void AddClassQueryFunction(const FName& ClassName, const TQueryFunction<bool>& InClassQueryFunction);
+	UE_API void RemoveClassQueryFunction(const FName& ClassName);
 	
 	// Outliner specific functionality
 	UE_API void CreateItemsFromQuery(TArray<FSceneOutlinerTreeItemPtr>& OutItems, ISceneOutlinerMode* InMode);
@@ -209,7 +219,10 @@ public:
 	// Get the parent row for a given row
 	UE_API DataStorage::RowHandle GetParentRow(DataStorage::RowHandle InRowHandle);
 
-	UE_API bool ShouldForceShowParentRows();
+	UE_API bool ShouldForceShowParentRows() const;
+
+	// Refresh the primary query node and its filters; use instead of RecompileQueries if the queries have not changed 
+	UE_API void RefreshQueries() const;
 
 	// Recompile all queries used by this table viewer
 	UE_API void RecompileQueries();
@@ -263,14 +276,34 @@ protected:
 	// Initial query provided by user
 	DataStorage::FQueryDescription InitialQueryDescription;
 
-	// External queries that are currently active (e.g Filters)
+	// External query descriptions that are currently active (e.g Filters)
 	TMap<FName, TSharedPtr<DataStorage::QueryStack::IQueryNode>> ExternalQueries;
+	
+	// External query functions that are currently active (e.g Filters)
+    TMap<FName, DataStorage::Queries::TQueryFunction<bool>> ExternalQueryFunctions;
+
+	// External query functions used to filter by class, these are the only filters that use OR instead of AND
+	TMap<FName, DataStorage::Queries::TQueryFunction<bool>> ClassFilters;
+
+	// Array to store all filter nodes created from the active query functions
+	TArray<TSharedPtr<DataStorage::QueryStack::FRowFilterNode>> FilterNodes;
+
+	// Array to store all filter nodes created from the active class filters
+	TArray<TSharedPtr<DataStorage::QueryStack::FRowFilterNode>> ClassFilterNodes;
 
 	// Optional Hierarchy Data
 	TOptional<FTedsOutlinerHierarchyData> HierarchyData;
-
+	// The query stack node responsible for collecting all rows that match the initial query on FullRefresh()
+	TSharedPtr<DataStorage::QueryStack::FRowQueryResultsNode> InitialRowCollectorNode;
 	// The query stack node responsible for collecting all rows that match the composite query on FullRefresh()
-	TSharedPtr<DataStorage::QueryStack::FRowQueryResultsNode> RowCollectorNode;
+	TSharedPtr<DataStorage::QueryStack::FRowQueryResultsNode> CombinedRowCollectorNode;
+
+	// The query stack node responsible for combining class filters - utilizes the unique merge approach (OR)
+	TSharedPtr<DataStorage::QueryStack::IRowNode> CombinedClassFilterRowNode;
+
+	// The query stack node responsible for combining all filters and queries as the 'final' node responsible
+	// for getting all desired rows - utilizes the repeating merge approach (AND)
+	TSharedPtr<DataStorage::QueryStack::IRowNode> FinalCombinedRowNode;
 
 	// The query stack node responsible for tracking and row addition/removals that happen in between refreshes
 	TSharedPtr<DataStorage::QueryStack::FRowMonitorNode> RowMonitorNode;
diff --git a/Engine/Plugins/Experimental/EditorDataStorageFeatures/Source/TedsOutliner/Public/TedsOutlinerModule.h b/Engine/Plugins/Experimental/EditorDataStorageFeatures/Source/TedsOutliner/Public/TedsOutlinerModule.h
--- a/Engine/Plugins/Experimental/EditorDataStorageFeatures/Source/TedsOutliner/Public/TedsOutlinerModule.h
+++ b/Engine/Plugins/Experimental/EditorDataStorageFeatures/Source/TedsOutliner/Public/TedsOutlinerModule.h
@@ -60,5 +60,6 @@ private:
 private:
 	FName TedsOutlinerTabName;
 	FDelegateHandle LevelEditorTabManagerChangedHandle;
+	DataStorage::ICoreProvider* Storage = nullptr;
 };
 } // namespace UE::Editor::Outliner
diff --git a/Engine/Plugins/Experimental/EditorDataStorageFeatures/Source/TedsQueryStack/Private/TedsRowMergeNode.cpp b/Engine/Plugins/Experimental/EditorDataStorageFeatures/Source/TedsQueryStack/Private/TedsRowMergeNode.cpp
--- a/Engine/Plugins/Experimental/EditorDataStorageFeatures/Source/TedsQueryStack/Private/TedsRowMergeNode.cpp
+++ b/Engine/Plugins/Experimental/EditorDataStorageFeatures/Source/TedsQueryStack/Private/TedsRowMergeNode.cpp
@@ -31,7 +31,6 @@ namespace UE::Editor::DataStorage::QueryStack
 			{
 				bRebuild = true;
 				ParentInfo.Revision = ParentRevision;
-				break;
 			}
 		}
 
diff --git a/Engine/Plugins/Experimental/EditorDataStorageFeatures/Source/TedsQueryStack/Public/TedsRowFilterNode.h b/Engine/Plugins/Experimental/EditorDataStorageFeatures/Source/TedsQueryStack/Public/TedsRowFilterNode.h
--- a/Engine/Plugins/Experimental/EditorDataStorageFeatures/Source/TedsQueryStack/Public/TedsRowFilterNode.h
+++ b/Engine/Plugins/Experimental/EditorDataStorageFeatures/Source/TedsQueryStack/Public/TedsRowFilterNode.h
@@ -22,26 +22,50 @@ namespace UE::Editor::DataStorage::QueryStack
 	class FRowFilterNode final : public IRowNode
 	{
 	public:
-		template<typename FilterFunction>
 		FRowFilterNode(ICoreProvider* Storage, const TSharedPtr<IRowNode>& ParentRowNode, ICoreProvider::EFilterOptions Options, 
-			FilterFunction&& Filter)
+			const Queries::TQueryFunction<bool>& Filter)
 			: Storage(Storage)
-			, Filter(Queries::BuildQueryFunction<bool>(Forward<FilterFunction>(Filter)))
+			, Filter(Filter)
 			, ParentRowNode(ParentRowNode)
 			, CachedParentRevisionID(ParentRowNode->GetRevision())
 			, Options(Options)
 		{
 		}
+    
+		FRowFilterNode(ICoreProvider* Storage, const TSharedPtr<IRowNode>& ParentRowNode, const Queries::TQueryFunction<bool>& Filter)
+			: FRowFilterNode(
+				Storage, 
+				ParentRowNode, 
+				ICoreProvider::EFilterOptions::Inclusive, 
+				Filter)
+		{}
+            
+		/*
+		 * Require that the template type cannot already be a TQueryFunction so it will pass to the defined constructor instead of this one.
+		 */
+		template< typename FilterFunction UE_REQUIRES(
+			!std::is_same_v<std::remove_cvref_t<FilterFunction>, Queries::TQueryFunction<bool>>) >
+		FRowFilterNode(ICoreProvider* Storage, const TSharedPtr<IRowNode>& ParentRowNode, ICoreProvider::EFilterOptions Options, 
+			FilterFunction&& Filter)
+			: FRowFilterNode(
+				Storage, 
+				ParentRowNode, 
+				Options, 
+				Queries::BuildQueryFunction<bool>(Forward<FilterFunction>(Filter)))
+		{}
 
-		template<typename FilterFunction>
+		/*
+		 * Require that the template type cannot already be a TQueryFunction so it will pass to the defined constructor instead of this one.
+		 */
+		template< typename FilterFunction UE_REQUIRES(
+			!std::is_same_v<std::remove_cvref_t<FilterFunction>, Queries::TQueryFunction<bool>>) >
 		FRowFilterNode(ICoreProvider* Storage, const TSharedPtr<IRowNode>& ParentRowNode, FilterFunction&& Filter)
-			: Storage(Storage)
-			, Filter(Queries::BuildQueryFunction<bool>(Forward<FilterFunction>(Filter)))
-			, ParentRowNode(ParentRowNode)
-			, CachedParentRevisionID(ParentRowNode->GetRevision())
-			, Options(ICoreProvider::EFilterOptions::Inclusive)
-		{
-		}
+			: FRowFilterNode(
+				Storage, 
+				ParentRowNode, 
+				ICoreProvider::EFilterOptions::Inclusive, 
+				Queries::BuildQueryFunction<bool>(Forward<FilterFunction>(Filter)))
+		{}
 
 		virtual ~FRowFilterNode() override = default;
 
diff --git a/Engine/Source/Runtime/TypedElementFramework/Public/Elements/Framework/TypedElementQueryFunctionArguments.inl b/Engine/Source/Runtime/TypedElementFramework/Public/Elements/Framework/TypedElementQueryFunctionArguments.inl
--- a/Engine/Source/Runtime/TypedElementFramework/Public/Elements/Framework/TypedElementQueryFunctionArguments.inl
+++ b/Engine/Source/Runtime/TypedElementFramework/Public/Elements/Framework/TypedElementQueryFunctionArguments.inl
@@ -113,19 +113,19 @@ namespace UE::Editor::DataStorage::Queries::Private
 			static TConstArrayView<const UScriptStruct*> Result =
 				[&]()
 				{
-					static const UScriptStruct* Columns[ColumnCount];
+					static const UScriptStruct* ResultColumns[ColumnCount];
 					int32 Index = 0;
 					(
 						[&]
 						{
 							if constexpr (EnumHasAllFlags(TArgument<Args>::Flags, Flags))
 							{
-								Columns[Index++] = TArgument<Args>::ColumnType::StaticStruct();
+								ResultColumns[Index++] = TArgument<Args>::ColumnType::StaticStruct();
 							}
 						}(), ...
 					);
 
-					return TConstArrayView<const UScriptStruct*>(Columns, ColumnCount);
+					return TConstArrayView<const UScriptStruct*>(ResultColumns, ColumnCount);
 				}();
 			return Result;
 		}
