TEDS hierarchy viewer: Various changes to support filtering/sorting/searching properly

Hierarchy Viewer: Remove the TopLevelRowNode argument, the hierarchy viewer now only requires a single node provider and automatically creates the top level row node internally
FTopLevelRowsNode: Make the node private since it's internal only now

Table Viewer/ Tile Viewer: Various minor fixes

diff --git a/Engine/Plugins/Experimental/EditorDataStorageFeatures/Source/TedsActorCompatibility/Private/ActorComponentDebugHierarchyWidget/ActorComponentTree.cpp b/Engine/Plugins/Experimental/EditorDataStorageFeatures/Source/TedsActorCompatibility/Private/ActorComponentDebugHierarchyWidget/ActorComponentTree.cpp
--- a/Engine/Plugins/Experimental/EditorDataStorageFeatures/Source/TedsActorCompatibility/Private/ActorComponentDebugHierarchyWidget/ActorComponentTree.cpp
+++ b/Engine/Plugins/Experimental/EditorDataStorageFeatures/Source/TedsActorCompatibility/Private/ActorComponentDebugHierarchyWidget/ActorComponentTree.cpp
@@ -3,6 +3,7 @@
 #include "ActorComponentTree.h"
 
 #include "SActorComponentTreeWidget.h"
+#include "TedsHierarchyNode.h"
 #include "TedsQueryNode.h"
 #include "TedsQueryNode.h"
 #include "TedsRowQueryResultsNode.h"
@@ -68,7 +69,7 @@ namespace UE::Editor::DataStorage::ActorCompatibility
 							}
 						}
 					};
-				
+
 					TSharedPtr<QueryStack::IRowNode> TopLevelRowNode = MakeShared<QueryStack::FRowQueryCallbackResultsNode>(
 						*DataStorageInterface,
 						QueryNode,
@@ -76,8 +77,15 @@ namespace UE::Editor::DataStorage::ActorCompatibility
 						FilterRowsBySelectionSet
 						);
 
+					// Use the actor-component hierarchy to enumerate the components of all the selected actors
+					TSharedPtr<QueryStack::IRowNode> AllRowNode = MakeShared<FHierarchyRowNode>(
+						*DataStorageInterface,
+						HierarchyHandle,
+						TopLevelRowNode,
+						FHierarchyRowNode::ESyncFlags::Always);
+
 					// Create a widget to display the view model, widget internally binds the results of the query and displays appropriately
-					TSharedRef<SActorComponentTreeWidget> TreeWidget = SNew(SActorComponentTreeWidget, HierarchyHandle, TopLevelRowNode);
+					TSharedRef<SActorComponentTreeWidget> TreeWidget = SNew(SActorComponentTreeWidget, HierarchyHandle, AllRowNode);
 			
 					MajorTab->SetContent(TreeWidget);
 				}
diff --git a/Engine/Plugins/Experimental/EditorDataStorageFeatures/Source/TedsActorCompatibility/Private/ActorComponentDebugHierarchyWidget/SActorComponentTreeWidget.cpp b/Engine/Plugins/Experimental/EditorDataStorageFeatures/Source/TedsActorCompatibility/Private/ActorComponentDebugHierarchyWidget/SActorComponentTreeWidget.cpp
--- a/Engine/Plugins/Experimental/EditorDataStorageFeatures/Source/TedsActorCompatibility/Private/ActorComponentDebugHierarchyWidget/SActorComponentTreeWidget.cpp
+++ b/Engine/Plugins/Experimental/EditorDataStorageFeatures/Source/TedsActorCompatibility/Private/ActorComponentDebugHierarchyWidget/SActorComponentTreeWidget.cpp
@@ -18,35 +18,17 @@ SActorComponentTreeWidget::~SActorComponentTreeWidget() = default;
 void SActorComponentTreeWidget::Construct(
 	const FArguments& InArgs,
 	UE::Editor::DataStorage::FHierarchyHandle InHierarchyHandle,
-	TSharedPtr<UE::Editor::DataStorage::QueryStack::IRowNode>& InTopLevelRowProvider)
+	TSharedPtr<UE::Editor::DataStorage::QueryStack::IRowNode>& InRowProvider)
 {
 	using namespace UE::Editor::DataStorage;
 	using namespace UE::Editor::DataStorage::Queries;
 
-	TopLevelRowProvider = InTopLevelRowProvider;
-	HierarchyHandle = InHierarchyHandle;
-
-	ICoreProvider* DataStorageInterface = GetMutableDataStorageFeature<ICoreProvider>(StorageFeatureName);
-	
-	TFunction<RowHandle(const void*, const UScriptStruct*)> GetParentRowFromColumnPredicate = DataStorageInterface->CreateParentExtractionFunction(HierarchyHandle);
-	
-	const UScriptStruct* HierarchyStruct = DataStorageInterface->GetHierarchyDataColumnType(HierarchyHandle);
-
-	TSharedRef<QueryStack::IRowNode> HierarchyRowNode = MakeShared<FHierarchyRowNode>(
-		*DataStorageInterface,
-		HierarchyHandle,
-		TopLevelRowProvider,
-		FHierarchyRowNode::ESyncFlags::IncrementWhenDifferent);
-	
-	TSharedPtr<QueryStack::IRowNode> HierarchyRowNodePtr = HierarchyRowNode.ToSharedPtr();
-
-	TSharedPtr<FHierarchyViewerLegacyData> HierarchyData = MakeShared<FHierarchyViewerLegacyData>(HierarchyStruct, MoveTemp(GetParentRowFromColumnPredicate));
+	TSharedPtr<FHierarchyViewerData> HierarchyData = MakeShared<FHierarchyViewerData>(InHierarchyHandle);
 	
 	TSharedRef<SHierarchyViewer> HierarchyViewer =
 		SNew(SHierarchyViewer,
 			MoveTemp(HierarchyData))
-			.AllNodeProvider(HierarchyRowNode)
-			.TopLevelNodeProvider(TopLevelRowProvider)
+			.AllNodeProvider(InRowProvider)
 			.EmptyRowsMessage(LOCTEXT("EmptyRowsMessage", "The provided query has no results."))
 			.CellWidgetPurpose(IUiProvider::FPurposeInfo("General", "RowLabel", NAME_None).GeneratePurposeID())
 			.Columns({FTypedElementLabelColumn::StaticStruct()});
@@ -57,11 +39,5 @@ void SActorComponentTreeWidget::Construct(
 	];
 }
 
-bool SActorComponentTreeWidget::TickProvider(float)
-{
-	TopLevelRowProvider->Update();
-	return true; // Keep ticking
-}
-
 #undef LOCTEXT_NAMESPACE
 
diff --git a/Engine/Plugins/Experimental/EditorDataStorageFeatures/Source/TedsActorCompatibility/Private/ActorComponentDebugHierarchyWidget/SActorComponentTreeWidget.h b/Engine/Plugins/Experimental/EditorDataStorageFeatures/Source/TedsActorCompatibility/Private/ActorComponentDebugHierarchyWidget/SActorComponentTreeWidget.h
--- a/Engine/Plugins/Experimental/EditorDataStorageFeatures/Source/TedsActorCompatibility/Private/ActorComponentDebugHierarchyWidget/SActorComponentTreeWidget.h
+++ b/Engine/Plugins/Experimental/EditorDataStorageFeatures/Source/TedsActorCompatibility/Private/ActorComponentDebugHierarchyWidget/SActorComponentTreeWidget.h
@@ -33,15 +33,7 @@ public:
 	void Construct(
 		const FArguments& InArgs,
 		UE::Editor::DataStorage::FHierarchyHandle InHierarchyHandle,
-		TSharedPtr<UE::Editor::DataStorage::QueryStack::IRowNode>& InTopLevelRowProvider);
+		TSharedPtr<UE::Editor::DataStorage::QueryStack::IRowNode>& InRowProvider);
 
 private:
-
-	bool TickProvider(float);
-
-	// Provides the top level rows
-	TSharedPtr<UE::Editor::DataStorage::QueryStack::IRowNode> TopLevelRowProvider;
-	// Provides the rows that will be shown in the associated tree view
-	TSharedPtr<UE::Editor::DataStorage::QueryStack::IRowNode> TreeRowProvider;
-	UE::Editor::DataStorage::FHierarchyHandle HierarchyHandle;
 };
\ No newline at end of file
diff --git a/Engine/Plugins/Experimental/EditorDataStorageFeatures/Source/TedsDebugger/Private/Widgets/QueryEditor/TedsQueryEditorResultsView.cpp b/Engine/Plugins/Experimental/EditorDataStorageFeatures/Source/TedsDebugger/Private/Widgets/QueryEditor/TedsQueryEditorResultsView.cpp
--- a/Engine/Plugins/Experimental/EditorDataStorageFeatures/Source/TedsDebugger/Private/Widgets/QueryEditor/TedsQueryEditorResultsView.cpp
+++ b/Engine/Plugins/Experimental/EditorDataStorageFeatures/Source/TedsDebugger/Private/Widgets/QueryEditor/TedsQueryEditorResultsView.cpp
@@ -18,7 +18,6 @@
 #include "Elements/Columns/TypedElementHiearchyColumns.h"
 #include "TedsQueryNode.h"
 #include "TedsRowSortNode.h"
-#include "QueryStack/TopLevelRowsNode.h"
 #include "Widgets/Layout/SBorder.h"
 #include "Widgets/STedsHierarchyViewer.h"
 #include "Widgets/STedsTableViewer.h"
@@ -230,12 +229,10 @@ namespace UE::Editor::DataStorage::Debug::QueryEditor
 		case ETableViewMode::Tree:
 
 			HierarchyData = MakeShared<FHierarchyViewerData>(HierarchyHandle);
-			TopLevelRowsNode = MakeShared<QueryStack::FTopLevelRowsNode>(&Model->GetTedsInterface(), HierarchyData, RowQueryStack);
 
 			TableViewer =
 				SNew(SHierarchyViewer, HierarchyData)
 					.AllNodeProvider(RowQueryStack)
-					.TopLevelNodeProvider(TopLevelRowsNode)
 					.EmptyRowsMessage(LOCTEXT("EmptyRowsMessage", "The provided query has no results."))
 					.OnSelectionChanged(STedsTableViewer::FOnSelectionChanged::CreateLambda(
 						[this](RowHandle SelectedRow)
diff --git a/Engine/Plugins/Experimental/EditorDataStorageFeatures/Source/TedsDebugger/Private/Widgets/QueryEditor/TedsQueryEditorResultsView.h b/Engine/Plugins/Experimental/EditorDataStorageFeatures/Source/TedsDebugger/Private/Widgets/QueryEditor/TedsQueryEditorResultsView.h
--- a/Engine/Plugins/Experimental/EditorDataStorageFeatures/Source/TedsDebugger/Private/Widgets/QueryEditor/TedsQueryEditorResultsView.h
+++ b/Engine/Plugins/Experimental/EditorDataStorageFeatures/Source/TedsDebugger/Private/Widgets/QueryEditor/TedsQueryEditorResultsView.h
@@ -67,8 +67,6 @@ namespace UE::Editor::DataStorage
 
 			// Hierarchy Data used by SHierarchyView to resolve parenting
 			TSharedPtr<FHierarchyViewerData> HierarchyData;
-			// Node to only collect top level rows in RowQueryStack for the hierarchy view
-			TSharedPtr<QueryStack::FTopLevelRowsNode> TopLevelRowsNode;
 			// The current hierarchy the widget is viewing
 			FHierarchyHandle HierarchyHandle;
 			
diff --git a/Engine/Plugins/Experimental/EditorDataStorageFeatures/Source/TedsEverythingPicker/Private/Context/TedsPickerContextUtil.cpp b/Engine/Plugins/Experimental/EditorDataStorageFeatures/Source/TedsEverythingPicker/Private/Context/TedsPickerContextUtil.cpp
--- a/Engine/Plugins/Experimental/EditorDataStorageFeatures/Source/TedsEverythingPicker/Private/Context/TedsPickerContextUtil.cpp
+++ b/Engine/Plugins/Experimental/EditorDataStorageFeatures/Source/TedsEverythingPicker/Private/Context/TedsPickerContextUtil.cpp
@@ -11,8 +11,6 @@
 #include "Elements/Columns/TedsTypeInfoColumns.h"
 #include "Elements/Framework/TypedElementQueryContext.h"
 
-#include "QueryStack/TopLevelRowsNode.h"
-
 #include "Widgets/STedsHierarchyViewer.h"
 #include "Widgets/STedsTableViewer.h"
 
@@ -80,12 +78,8 @@ namespace UE::Editor::DataStorage::Picker
 
 		TSharedPtr<FHierarchyViewerLegacyData> HierarchyData = MakeShared<FHierarchyViewerLegacyData>(FTableRowParentColumn::StaticStruct(), MoveTemp(ParentExtractionFunction));
 
-		// Node that filters out the ReferenceQueryResultsNode or SearchNode (if enabled) to generate all the top level nodes
-		TSharedRef<QueryStack::FTopLevelRowsNode> TopLevelNode = MakeShared<QueryStack::FTopLevelRowsNode>(Storage, HierarchyData, SearchNode);
-
 		TSharedPtr<SHierarchyViewer> ReferenceView = SNew(SHierarchyViewer, HierarchyData)
 			.AllNodeProvider(SearchNode)
-			.TopLevelNodeProvider(TopLevelNode)
 			.Columns({ FTypedElementLabelColumn::StaticStruct(), FTypedElementClassTypeInfoColumn::StaticStruct() })
 			.CellWidgetPurpose(PurposeId)
 			.OnSelectionChanged(InArgs._OnSelectionChanged)
diff --git a/Engine/Plugins/Experimental/EditorDataStorageFeatures/Source/TedsTableViewer/Private/QueryStack/TopLevelRowsNode.h b/Engine/Plugins/Experimental/EditorDataStorageFeatures/Source/TedsTableViewer/Private/QueryStack/TopLevelRowsNode.h
new file mode 100644
--- /dev/null
+++ b/Engine/Plugins/Experimental/EditorDataStorageFeatures/Source/TedsTableViewer/Private/QueryStack/TopLevelRowsNode.h
@@ -0,0 +1,43 @@
+// Copyright Epic Games, Inc. All Rights Reserved.
+
+#pragma once
+
+#include "TedsQueryStackInterfaces.h"
+#include "Elements/Interfaces/TypedElementDataStorageInterface.h"
+
+#define UE_API TEDSTABLEVIEWER_API
+
+namespace UE::Editor::DataStorage
+{
+	class IHierarchyViewerDataInterface;
+}
+
+namespace UE::Editor::DataStorage::QueryStack
+{
+	// A custom query stack node to get top level rows of the given rows to display in the hierarchy viewer
+	class FTopLevelRowsNode : public IRowNode
+	{
+	public:
+		
+		UE_API FTopLevelRowsNode(const ICoreProvider* InStorage, TSharedPtr<IHierarchyViewerDataInterface> InHierarchyData, TSharedPtr<IRowNode> InParent);
+		UE_API virtual ~FTopLevelRowsNode() override = default;
+		
+		/** Retrieve access to the rows used by this node. */
+		virtual FRowHandleArrayView GetRows() const override;
+		virtual FRowHandleArray& GetMutableRows() override;
+		virtual RevisionId GetRevision() const override;
+		virtual void Update() override;
+
+	protected:
+		void UpdateRows();
+	protected:
+		RevisionId Revision = 0;
+		const ICoreProvider* Storage = nullptr;
+		TSharedPtr<IHierarchyViewerDataInterface> HierarchyData;
+		TSharedPtr<IRowNode> Parent;
+		FRowHandleArray Rows;
+		RevisionId CachedParentRevision = TNumericLimits<RevisionId>::Max();
+	};
+}
+
+#undef UE_API
\ No newline at end of file
diff --git a/Engine/Plugins/Experimental/EditorDataStorageFeatures/Source/TedsTableViewer/Private/TedsTableViewerModel.cpp b/Engine/Plugins/Experimental/EditorDataStorageFeatures/Source/TedsTableViewer/Private/TedsTableViewerModel.cpp
--- a/Engine/Plugins/Experimental/EditorDataStorageFeatures/Source/TedsTableViewer/Private/TedsTableViewerModel.cpp
+++ b/Engine/Plugins/Experimental/EditorDataStorageFeatures/Source/TedsTableViewer/Private/TedsTableViewerModel.cpp
@@ -69,6 +69,11 @@ namespace UE::Editor::DataStorage
 		return !Storage->HasColumns<FHideRowFromUITag>(InRowHandle);
 	}
 
+	TSharedPtr<QueryStack::IRowNode> FTedsTableViewerModel::GetRowNode() const
+	{
+		return RowQueryStack;
+	}
+
 	void FTedsTableViewerModel::ValidateRequestedColumns()
 	{
 		RequestedTedsColumns.RemoveAll([](const TWeakObjectPtr<const UScriptStruct>& Column)
diff --git a/Engine/Plugins/Experimental/EditorDataStorageFeatures/Source/TedsTableViewer/Private/Widgets/STedsHierarchyViewer.cpp b/Engine/Plugins/Experimental/EditorDataStorageFeatures/Source/TedsTableViewer/Private/Widgets/STedsHierarchyViewer.cpp
--- a/Engine/Plugins/Experimental/EditorDataStorageFeatures/Source/TedsTableViewer/Private/Widgets/STedsHierarchyViewer.cpp
+++ b/Engine/Plugins/Experimental/EditorDataStorageFeatures/Source/TedsTableViewer/Private/Widgets/STedsHierarchyViewer.cpp
@@ -9,6 +9,7 @@
 #include "Elements/Columns/TypedElementUIColumns.h"
 #include "Elements/Framework/TypedElementAttributeBinding.h"
 #include "Elements/Framework/TypedElementDataStorageWidget.h"
+#include "QueryStack/TopLevelRowsNode.h"
 #include "Widgets/STedsTableViewerRow.h"
 #include "Widgets/Views/SListView.h"
 #include "Widgets/Text/STextBlock.h"
@@ -27,12 +28,13 @@ namespace UE::Editor::DataStorage
 		EmptyRowsMessage = InArgs._EmptyRowsMessage;
 		DefaultExpansionState = InArgs._DefaultExpansionState;
 		
-		IUiProvider::FPurposeID CellWigetPurpose = InArgs._CellWidgetPurpose;
+		IUiProvider::FPurposeID CellWidgetPurpose = InArgs._CellWidgetPurpose;
 
-		if (!CellWigetPurpose.IsSet())
+		IUiProvider* StorageUi = GetMutableDataStorageFeature<IUiProvider>(UiFeatureName);
+
+		if (!CellWidgetPurpose.IsSet())
 		{
-			IUiProvider* StorageUi = GetMutableDataStorageFeature<IUiProvider>(UiFeatureName);
-			CellWigetPurpose = StorageUi->GetGeneralWidgetPurposeID();
+			CellWidgetPurpose = StorageUi->GetGeneralWidgetPurposeID();
 		}
 
 		IUiProvider::FPurposeID HeaderWidgetPurpose = InArgs._HeaderWidgetPurpose;
@@ -42,16 +44,14 @@ namespace UE::Editor::DataStorage
 			HeaderWidgetPurpose = IUiProvider::FPurposeInfo("General", "Header", NAME_None).GeneratePurposeID();
 		}
 
-		Model = MakeShared<FTedsTableViewerModel>(InArgs._AllNodeProvider, InArgs._Columns, CellWigetPurpose, HeaderWidgetPurpose,
+		Model = MakeShared<FTedsTableViewerModel>(InArgs._AllNodeProvider, InArgs._Columns, CellWidgetPurpose, HeaderWidgetPurpose,
 			FTedsTableViewerModel::FIsItemVisible::CreateSP(this, &SHierarchyViewer::IsItemVisible));
 		
-		HierarchyNode = InArgs._TopLevelNodeProvider;
+		TopLevelRowsNode = MakeShared<QueryStack::FTopLevelRowsNode>(Model->GetDataStorageInterface(), HierarchyInterface, Model->GetRowNode());
 		
 		HeaderRowWidget = SNew( SHeaderRow )
 							.CanSelectGeneratedColumn(true);
 		
-		IUiProvider* StorageUi = Model->GetDataStorageUiProvider();
-
 		if (ensure(StorageUi))
 		{
 			TedsWidget = StorageUi->CreateContainerTedsWidget(InvalidRowHandle);
@@ -88,9 +88,12 @@ namespace UE::Editor::DataStorage
 	{
 		TopLevelRows.Empty();
 
-		HierarchyNode->Update();
+		// The current logic relies on the query stack to invalidate and cause a model refresh when the hierarchy changes since we don't have a
+		// generic way of detecting hierarchy changes. If the model isn't refreshed on hierarchy change the top level rows will still use the old
+		// hierarchy
+		TopLevelRowsNode->Update();
 		
-		FRowHandleArrayView Rows = HierarchyNode->GetRows();
+		FRowHandleArrayView Rows(TopLevelRowsNode->GetRows());
 
 		for(const RowHandle RowHandle : Rows)
 		{
diff --git a/Engine/Plugins/Experimental/EditorDataStorageFeatures/Source/TedsTableViewer/Private/Widgets/STedsTileViewer.cpp b/Engine/Plugins/Experimental/EditorDataStorageFeatures/Source/TedsTableViewer/Private/Widgets/STedsTileViewer.cpp
--- a/Engine/Plugins/Experimental/EditorDataStorageFeatures/Source/TedsTableViewer/Private/Widgets/STedsTileViewer.cpp
+++ b/Engine/Plugins/Experimental/EditorDataStorageFeatures/Source/TedsTableViewer/Private/Widgets/STedsTileViewer.cpp
@@ -177,7 +177,7 @@ namespace UE::Editor::DataStorage
 
 	void STedsTileViewer::AddCustomRowWidget(const TSharedRef<FTedsTableViewerColumn>& InColumn)
 	{
-		// The Tile Viewer doesn't have traditional columns, so adding a custom column is ambiguous (where should the column be placed?)
+		// The Tile Viewer doesn't have traditional columns, so adding a custom row widget is not implemented yet (where should the widget be placed on the tile?)
 		ensureMsgf(false, TEXT("Adding custom columns is currently not supported in the Tile Viewer!"));
 	}
 
diff --git a/Engine/Plugins/Experimental/EditorDataStorageFeatures/Source/TedsTableViewer/Public/QueryStack/TopLevelRowsNode.h b/Engine/Plugins/Experimental/EditorDataStorageFeatures/Source/TedsTableViewer/Public/QueryStack/TopLevelRowsNode.h
deleted file mode 100644
diff --git a/Engine/Plugins/Experimental/EditorDataStorageFeatures/Source/TedsTableViewer/Public/TedsTableViewerModel.h b/Engine/Plugins/Experimental/EditorDataStorageFeatures/Source/TedsTableViewer/Public/TedsTableViewerModel.h
--- a/Engine/Plugins/Experimental/EditorDataStorageFeatures/Source/TedsTableViewer/Public/TedsTableViewerModel.h
+++ b/Engine/Plugins/Experimental/EditorDataStorageFeatures/Source/TedsTableViewer/Public/TedsTableViewerModel.h
@@ -18,6 +18,7 @@ namespace UE::Editor::DataStorage
 	{
 		class FRowOrderInversionNode;
 		class FRowSortNode;
+		class IRowNode;
 	}
 	class FTedsTableViewerColumn;
 	class ICompatibilityProvider;
@@ -69,7 +70,7 @@ namespace UE::Editor::DataStorage
 		// Clear the current list of columns being displayed and set it to the given list
 		TEDSTABLEVIEWER_API void SetColumns(const TArray<TWeakObjectPtr<const UScriptStruct>>& InColumns);
 
-		// Add a custom column to display in the table viewer, that doesn't necessarily map to a Teds column
+		// Add a custom row widget to display in the table viewer, that doesn't necessarily map to a Teds column
 		TEDSTABLEVIEWER_API void AddCustomRowWidget(const TSharedRef<FTedsTableViewerColumn>& InColumn);
 
 		TEDSTABLEVIEWER_API ICoreProvider* GetDataStorageInterface() const;
@@ -79,6 +80,9 @@ namespace UE::Editor::DataStorage
 		// Check whether a row is allowed to be displayed in the table viewer
 		TEDSTABLEVIEWER_API bool IsRowDisplayable(RowHandle InRowHandle) const;
 
+		// Get the bottom-most row node used in the query stack for this model
+		TEDSTABLEVIEWER_API TSharedPtr<QueryStack::IRowNode> GetRowNode() const;
+		
 	protected:
 
 		// Generate the actual columns to display in the UI using TEDS UI
diff --git a/Engine/Plugins/Experimental/EditorDataStorageFeatures/Source/TedsTableViewer/Public/Widgets/STedsHierarchyViewer.h b/Engine/Plugins/Experimental/EditorDataStorageFeatures/Source/TedsTableViewer/Public/Widgets/STedsHierarchyViewer.h
--- a/Engine/Plugins/Experimental/EditorDataStorageFeatures/Source/TedsTableViewer/Public/Widgets/STedsHierarchyViewer.h
+++ b/Engine/Plugins/Experimental/EditorDataStorageFeatures/Source/TedsTableViewer/Public/Widgets/STedsHierarchyViewer.h
@@ -56,9 +56,6 @@ namespace UE::Editor::DataStorage
 			
 		}
 
-		// RowNode that provides the top level rows for the hierarchy
-		SLATE_ARGUMENT(TSharedPtr<QueryStack::IRowNode>, TopLevelNodeProvider)
-
 		// Query Stack that will supply the rows to be displayed
 		SLATE_ARGUMENT(TSharedPtr<QueryStack::IRowNode>, AllNodeProvider)
 
@@ -96,7 +93,8 @@ namespace UE::Editor::DataStorage
 		// Clear the current list of columns being displayed and set it to the given list
 		TEDSTABLEVIEWER_API virtual void SetColumns(const TArray<TWeakObjectPtr<const UScriptStruct>>& Columns) override;
 
-		// Add a custom column to display in the table viewer, that doesn't necessarily map to a Teds column
+    	// Add a custom per-row widget to the table viewer (that doesn't necessarily map to a TEDS column)
+		// This means a new column for the hierarchy viewer
 		TEDSTABLEVIEWER_API virtual void AddCustomRowWidget(const TSharedRef<FTedsTableViewerColumn>& InColumn) override;
 
 		// Execute the given callback for each row that is selected in the table viewer
@@ -163,7 +161,7 @@ namespace UE::Editor::DataStorage
 		TAttribute<EExpansionState> DefaultExpansionState;
 
 		// A specialized query stack node to only contain the top level rows present in the actual query stack we are using.
-		TSharedPtr<QueryStack::IRowNode> HierarchyNode;
+		TSharedPtr<QueryStack::IRowNode> TopLevelRowsNode;
 		
 		// The cached list of top level rows
 		TArray<TableViewerItemPtr> TopLevelRows;
diff --git a/Engine/Plugins/Experimental/EditorDataStorageFeatures/Source/TedsTableViewer/Public/Widgets/STedsTableViewer.h b/Engine/Plugins/Experimental/EditorDataStorageFeatures/Source/TedsTableViewer/Public/Widgets/STedsTableViewer.h
--- a/Engine/Plugins/Experimental/EditorDataStorageFeatures/Source/TedsTableViewer/Public/Widgets/STedsTableViewer.h
+++ b/Engine/Plugins/Experimental/EditorDataStorageFeatures/Source/TedsTableViewer/Public/Widgets/STedsTableViewer.h
@@ -85,7 +85,8 @@ namespace UE::Editor::DataStorage
 		// Clear the current list of columns being displayed and set it to the given list
 		TEDSTABLEVIEWER_API virtual void SetColumns(const TArray<TWeakObjectPtr<const UScriptStruct>>& Columns) override;
 
-		// Add a custom column to display in the table viewer, that doesn't necessarily map to a Teds column
+		// Add a custom per-row widget to the table viewer (that doesn't necessarily map to a TEDS column)
+		// This means a new column for the table viewer
 		TEDSTABLEVIEWER_API virtual void AddCustomRowWidget(const TSharedRef<FTedsTableViewerColumn>& InColumn) override;
 
 		// Execute the given callback for each row that is selected in the table viewer
diff --git a/Engine/Plugins/Experimental/EditorDataStorageFeatures/Source/TedsTableViewer/Public/Widgets/STedsTileViewer.h b/Engine/Plugins/Experimental/EditorDataStorageFeatures/Source/TedsTableViewer/Public/Widgets/STedsTileViewer.h
--- a/Engine/Plugins/Experimental/EditorDataStorageFeatures/Source/TedsTableViewer/Public/Widgets/STedsTileViewer.h
+++ b/Engine/Plugins/Experimental/EditorDataStorageFeatures/Source/TedsTableViewer/Public/Widgets/STedsTileViewer.h
@@ -104,7 +104,7 @@ namespace UE::Editor::DataStorage
 		// Clear the current list of columns being displayed and set it to the given list
 		TEDSTABLEVIEWER_API virtual void SetColumns(const TArray<TWeakObjectPtr<const UScriptStruct>>& Columns) override;
 
-		// Add a custom column to display in the table viewer, that doesn't necessarily map to a Teds column
+		// Add a custom per-row widget to the table viewer (that doesn't necessarily map to a TEDS column)
 		TEDSTABLEVIEWER_API virtual void AddCustomRowWidget(const TSharedRef<FTedsTableViewerColumn>& InColumn) override;
 
 	protected:
