STedsHierarchyViewer: Ensure that TEDS UI is initialized before using it

diff --git a/Engine/Plugins/Experimental/EditorDataStorageFeatures/Source/TedsTableViewer/Private/Widgets/STedsHierarchyViewer.cpp b/Engine/Plugins/Experimental/EditorDataStorageFeatures/Source/TedsTableViewer/Private/Widgets/STedsHierarchyViewer.cpp
--- a/Engine/Plugins/Experimental/EditorDataStorageFeatures/Source/TedsTableViewer/Private/Widgets/STedsHierarchyViewer.cpp
+++ b/Engine/Plugins/Experimental/EditorDataStorageFeatures/Source/TedsTableViewer/Private/Widgets/STedsHierarchyViewer.cpp
@@ -32,6 +32,11 @@ namespace UE::Editor::DataStorage
 
 		IUiProvider* StorageUi = GetMutableDataStorageFeature<IUiProvider>(UiFeatureName);
 
+		if (!ensureMsgf(StorageUi, TEXT("Cannot use SHierarchyVIewer without TEDS UI being initialized")))
+		{
+			return;
+		}
+
 		if (!CellWidgetPurpose.IsSet())
 		{
 			CellWidgetPurpose = StorageUi->GetGeneralWidgetPurposeID();
@@ -52,15 +57,12 @@ namespace UE::Editor::DataStorage
 		HeaderRowWidget = SNew( SHeaderRow )
 							.CanSelectGeneratedColumn(true);
 		
-		if (ensure(StorageUi))
-		{
-			TedsWidget = StorageUi->CreateContainerTedsWidget(InvalidRowHandle);
+		TedsWidget = StorageUi->CreateContainerTedsWidget(InvalidRowHandle);
 			
-			ChildSlot
-			[
-				TedsWidget->AsWidget()
-			];
-		}
+		ChildSlot
+		[
+			TedsWidget->AsWidget()
+		];
 		
 		AddWidgetColumns();
 
