STedsTableViewer : Add a slate argument to set CanSelectGeneratedColumn on the SHeaderRow

diff --git a/Engine/Plugins/Experimental/EditorDataStorageFeatures/Source/TedsTableViewer/Private/Widgets/STedsTableViewer.cpp b/Engine/Plugins/Experimental/EditorDataStorageFeatures/Source/TedsTableViewer/Private/Widgets/STedsTableViewer.cpp
--- a/Engine/Plugins/Experimental/EditorDataStorageFeatures/Source/TedsTableViewer/Private/Widgets/STedsTableViewer.cpp
+++ b/Engine/Plugins/Experimental/EditorDataStorageFeatures/Source/TedsTableViewer/Private/Widgets/STedsTableViewer.cpp
@@ -42,7 +42,7 @@ namespace UE::Editor::DataStorage
 			FTedsTableViewerModel::FIsItemVisible::CreateSP(this, &STedsTableViewer::IsItemVisible));
 		
 		HeaderRowWidget = SNew( SHeaderRow )
-							.CanSelectGeneratedColumn(true);
+							.CanSelectGeneratedColumn(InArgs._CanSelectGeneratedColumn);
 		
 		IUiProvider* StorageUi = Model->GetDataStorageUiProvider();
 
diff --git a/Engine/Plugins/Experimental/EditorDataStorageFeatures/Source/TedsTableViewer/Public/Widgets/STedsTableViewer.h b/Engine/Plugins/Experimental/EditorDataStorageFeatures/Source/TedsTableViewer/Public/Widgets/STedsTableViewer.h
--- a/Engine/Plugins/Experimental/EditorDataStorageFeatures/Source/TedsTableViewer/Public/Widgets/STedsTableViewer.h
+++ b/Engine/Plugins/Experimental/EditorDataStorageFeatures/Source/TedsTableViewer/Public/Widgets/STedsTableViewer.h
@@ -39,7 +39,8 @@ namespace UE::Editor::DataStorage
 		DECLARE_DELEGATE_OneParam(FOnSelectionChanged, RowHandle)
 
 		SLATE_BEGIN_ARGS(STedsTableViewer)
-			: _ListSelectionMode(ESelectionMode::Type::Single)
+			: _CanSelectGeneratedColumn(true)
+			, _ListSelectionMode(ESelectionMode::Type::Single)
 			, _ItemHeight()
 			, _ItemPadding(FMargin(0.f))
 		{
@@ -59,6 +60,9 @@ namespace UE::Editor::DataStorage
 			
 		// The widget purpose to use to create the header widgets
 		SLATE_ARGUMENT(IUiProvider::FPurposeID, HeaderWidgetPurpose)
+
+		// Whether or not the user can select the columns generated by right clicking on the header menu.
+		SLATE_ARGUMENT(bool, CanSelectGeneratedColumn)
 		
 		// Delegate called when the selection changes
 		SLATE_EVENT(FOnSelectionChanged, OnSelectionChanged)
