Add Time Slicing for TEDs Asset Data Processing
- Provides a fixed budget per tick for the main 3 processors
- Tracks the dynamic cost of avg unit time per tick to rebalance for the next

diff --git a/Engine/Plugins/Experimental/EditorDataStorageFeatures/Source/TedsAssetData/Private/CB/TedsAssetDataCBDataSource.cpp b/Engine/Plugins/Experimental/EditorDataStorageFeatures/Source/TedsAssetData/Private/CB/TedsAssetDataCBDataSource.cpp
--- a/Engine/Plugins/Experimental/EditorDataStorageFeatures/Source/TedsAssetData/Private/CB/TedsAssetDataCBDataSource.cpp
+++ b/Engine/Plugins/Experimental/EditorDataStorageFeatures/Source/TedsAssetData/Private/CB/TedsAssetDataCBDataSource.cpp
@@ -3,6 +3,7 @@
 #include "TedsAssetDataCBDataSource.h"
 
 #include "TedsAssetDataColumns.h"
+#include "TedsAssetData.h"
 
 #include "AssetDefinitionRegistry.h"
 #include "AssetRegistry/IAssetRegistry.h"
@@ -16,7 +17,6 @@
 #include "Factories/Factory.h"
 #include "HAL/IConsoleManager.h"
 #include "Interfaces/IPluginManager.h"
-#include "Interfaces/IPluginManager.h"
 #include "Internationalization/Text.h"
 #include "Misc/PathViews.h"
 #include "PluginDescriptor.h"
@@ -349,8 +349,20 @@ void FTedsAssetDataCBDataSource::AddAssetDataColumns(UE::Editor::DataStorage::IQ
 
 void FTedsAssetDataCBDataSource::ProcessPathQueryCallback(DataStorage::IQueryContext& Context, const DataStorage::RowHandle* Rows, const FAssetPathColumn_Experimental* PathColumn) const
 {
-	int32 NumOfRowToProcess = Context.GetRowCount();
+	uint64 StartTime = FTedsAssetData::ProcessPathAddStartTime.load();
+	uint64 CurrentTime = FPlatformTime::Cycles64();
+
+	// Processors run in chunks on a worker thread. 
+	// Using compare_exchange_weak allows us to take the min start time and the max end time of all chunks in a thread safe way
+	while (CurrentTime < StartTime)
+	{
+		if (FTedsAssetData::ProcessPathAddStartTime.compare_exchange_weak(StartTime, CurrentTime))
+		{
+			break;
+		}
+	}
 
+	int32 NumOfRowToProcess = Context.GetRowCount();
 	FNameBuilder InternalPath;
 	FNameBuilder VirtualPath;
 
@@ -385,6 +397,19 @@ void FTedsAssetDataCBDataSource::ProcessPathQueryCallback(DataStorage::IQueryCon
 			Context.AddColumn(Rows[Index], MoveTemp(FolderType));
 		}
 	}
+
+	uint64 EndTime = FTedsAssetData::ProcessPathAddEndTime.load();
+	CurrentTime = FPlatformTime::Cycles64();
+
+	while (CurrentTime > EndTime)
+	{
+		if (FTedsAssetData::ProcessPathAddEndTime.compare_exchange_weak(EndTime, CurrentTime))
+		{
+			break;
+		}
+	}
+
+	FTedsAssetData::ProcessedAddedPathBatchSize.fetch_add(NumOfRowToProcess);
 }
 
 void FTedsAssetDataCBDataSource::ProcessAssetDataPathUpdateQueryCallback(DataStorage::IQueryContext& Context, const DataStorage::RowHandle* Rows, const FAssetDataColumn_Experimental* AssetDataColumn) const
@@ -410,8 +435,20 @@ void FTedsAssetDataCBDataSource::ProcessAssetDataPathUpdateQueryCallback(DataSto
 
 void FTedsAssetDataCBDataSource::ProcessAssetDataAndPathUpdateQueryCallback(DataStorage::IQueryContext& Context, const DataStorage::RowHandle* Rows, const FAssetDataColumn_Experimental* AssetDataColumn) const
 {
-	const int32 RowCount = Context.GetRowCount();
+	uint64 StartTime = FTedsAssetData::ProcessAssetAddStartTime.load();
+	uint64 CurrentTime = FPlatformTime::Cycles64();
 
+	// Processors run in chunks on a worker thread. 
+	// Using compare_exchange_weak allows us to take the min start time and the max end time of all chunks in a thread safe way
+	while (CurrentTime < StartTime)
+	{
+		if (FTedsAssetData::ProcessAssetAddStartTime.compare_exchange_weak(StartTime, CurrentTime))
+		{
+			break;
+		}
+	}
+
+	const int32 RowCount = Context.GetRowCount();
 	TArray<FName, TInlineAllocator<32>> PackageNames;
 	TArray<TPair<DataStorage::RowHandle, const FAssetData*>, TInlineAllocator<32>> RowsAndAssetData;
 	PackageNames.Reserve(RowCount);
@@ -424,23 +461,20 @@ void FTedsAssetDataCBDataSource::ProcessAssetDataAndPathUpdateQueryCallback(Data
 	{ 
 		const FAssetData& AssetData = AssetDataColumn[Index].AssetData;
 
-		if (ContentBrowserDataUtils::IsPrimaryAsset(AssetData))
-		{ 
-			InternalPath.Reset();
-			AssetData.AppendObjectPath(InternalPath);
+		InternalPath.Reset();
+		AssetData.AppendObjectPath(InternalPath);
 
-			if (GenerateVirtualPath(InternalPath, VirtualPath))
-			{
-				const DataStorage::RowHandle Row = Rows[Index];
+		if (GenerateVirtualPath(InternalPath, VirtualPath))
+		{
+			const DataStorage::RowHandle Row = Rows[Index];
 
-				FVirtualPathColumn_Experimental VirtualPathColumn;
-				VirtualPathColumn.VirtualPath = *VirtualPath;
-				// Todo investigate for a batch add Maybe?
-				Context.AddColumn(Row, MoveTemp(VirtualPathColumn));
+			FVirtualPathColumn_Experimental VirtualPathColumn;
+			VirtualPathColumn.VirtualPath = *VirtualPath;
+			// Todo investigate for a batch add Maybe?
+			Context.AddColumn(Row, MoveTemp(VirtualPathColumn));
 
-				PackageNames.Add(AssetData.PackageName);
-				RowsAndAssetData.Emplace(Row, &AssetData);
-			}
+			PackageNames.Add(AssetData.PackageName);
+			RowsAndAssetData.Emplace(Row, &AssetData);
 		}
 	}
 
@@ -452,10 +486,36 @@ void FTedsAssetDataCBDataSource::ProcessAssetDataAndPathUpdateQueryCallback(Data
 		const TOptional<FAssetPackageData>& AssetPackageData = AssetPackageDatas[Index];
 		AddAssetDataColumns(Context, Pair.Key, *Pair.Value, AssetPackageData.GetPtrOrNull());
 	}
+
+	uint64 EndTime = FTedsAssetData::ProcessAssetAddEndTime.load();
+	CurrentTime = FPlatformTime::Cycles64();
+
+	while (CurrentTime > EndTime)
+	{
+		if (FTedsAssetData::ProcessAssetAddEndTime.compare_exchange_weak(EndTime, CurrentTime))
+		{
+			break;
+		}
+	}
+
+	FTedsAssetData::ProcessedAddedAssetBatchSize.fetch_add(RowCount);
 }
 
 void FTedsAssetDataCBDataSource::ProcessAssetDataUpdateQueryCallback(DataStorage::IQueryContext& Context, const DataStorage::RowHandle* Rows, const FAssetDataColumn_Experimental* AssetDataColumn) const
 {
+	uint64 StartTime = FTedsAssetData::ProcessAssetUpdateStartTime.load();
+	uint64 CurrentTime = FPlatformTime::Cycles64();
+
+	// Processors run in chunks on a worker thread. 
+	// Using compare_exchange_weak allows us to take the min start time and the max end time of all chunks in a thread safe way
+	while (CurrentTime < StartTime)
+	{
+		if (FTedsAssetData::ProcessAssetUpdateStartTime.compare_exchange_weak(StartTime, CurrentTime))
+		{
+			break;
+		}
+	}
+
 	const int32 RowCount = Context.GetRowCount();
 	TArray<FName, TInlineAllocator<32>> PackageNames;
 	PackageNames.Reserve(RowCount);
@@ -471,6 +531,19 @@ void FTedsAssetDataCBDataSource::ProcessAssetDataUpdateQueryCallback(DataStorage
 	{
 		AddAssetDataColumns(Context, Rows[Index], AssetDataColumn[Index].AssetData, AssetPackageDatas[Index].GetPtrOrNull());
 	}
+
+	uint64 EndTime = FTedsAssetData::ProcessAssetUpdateEndTime.load();
+	CurrentTime = FPlatformTime::Cycles64();
+
+	while (CurrentTime > EndTime)
+	{
+		if (FTedsAssetData::ProcessAssetUpdateEndTime.compare_exchange_weak(EndTime, CurrentTime))
+		{
+			break;
+		}
+	}
+
+	FTedsAssetData::ProcessedUpdatedAssetBatchSize.fetch_add(RowCount);
 }
 
 void FTedsAssetDataCBDataSource::PrepopulateTagsMetadataCache()
diff --git a/Engine/Plugins/Experimental/EditorDataStorageFeatures/Source/TedsAssetData/Private/TedsAssetData.cpp b/Engine/Plugins/Experimental/EditorDataStorageFeatures/Source/TedsAssetData/Private/TedsAssetData.cpp
--- a/Engine/Plugins/Experimental/EditorDataStorageFeatures/Source/TedsAssetData/Private/TedsAssetData.cpp
+++ b/Engine/Plugins/Experimental/EditorDataStorageFeatures/Source/TedsAssetData/Private/TedsAssetData.cpp
@@ -14,11 +14,82 @@
 #include "Elements/Framework/TypedElementQueryBuilder.h"
 #include "Elements/Interfaces/TypedElementDataStorageInterface.h"
 #include "Engine/Engine.h"
+#include "DataStorage/Debug/Log.h"
+
+#include "ContentBrowserDataUtils.h"
+
+#include "ProfilingDebugging/CountersTrace.h"
+
+#include <limits>
 
 #define TRACK_TEDSASSETDATA_MEMORY 0
 
+#if COUNTERSTRACE_ENABLED
+
+DECLARE_STATS_GROUP(TEXT("Teds Asset Data"), STATGROUP_TedsAssetData, STATCAT_Advanced);
+
+DECLARE_DWORD_COUNTER_STAT(TEXT("Path Adds Remaining"), STAT_TedsAssetData_PathAddsRemaining, STATGROUP_TedsAssetData);
+DECLARE_DWORD_COUNTER_STAT(TEXT("Asset Adds Remaining"), STAT_TedsAssetData_AssetAddsRemaining, STATGROUP_TedsAssetData);
+DECLARE_DWORD_COUNTER_STAT(TEXT("Asset Updates Remaining"), STAT_TedsAssetData_AssetUpdatesRemaining, STATGROUP_TedsAssetData);
+
+DECLARE_FLOAT_COUNTER_STAT(TEXT("Path Add Avg Time"), STAT_TedsAssetData_PathAddAvgTime, STATGROUP_TedsAssetData);
+DECLARE_FLOAT_COUNTER_STAT(TEXT("Asset Add Avg Time"), STAT_TedsAssetData_AssetAddAvgTime, STATGROUP_TedsAssetData);
+DECLARE_FLOAT_COUNTER_STAT(TEXT("Asset Update Avg Time"), STAT_TedsAssetData_AssetUpdateAvgTime, STATGROUP_TedsAssetData);
+
+DECLARE_DWORD_COUNTER_STAT(TEXT("Path Add Batch Size"), STAT_TedsAssetData_PathAddBatchSize, STATGROUP_TedsAssetData);
+DECLARE_DWORD_COUNTER_STAT(TEXT("Asset Add Batch Size"), STAT_TedsAssetData_AssetAddBatchSize, STATGROUP_TedsAssetData);
+DECLARE_DWORD_COUNTER_STAT(TEXT("Asset Update Batch Size"), STAT_TedsAssetData_AssetUpdateBatchSize, STATGROUP_TedsAssetData);
+
+DECLARE_DWORD_COUNTER_STAT(TEXT("Path Add Batch Size Received"), STAT_TedsAssetData_PathAddBatchSizeReceived, STATGROUP_TedsAssetData);
+DECLARE_DWORD_COUNTER_STAT(TEXT("Asset Add Batch Size Received"), STAT_TedsAssetData_AssetAddBatchSizeReceived, STATGROUP_TedsAssetData);
+DECLARE_DWORD_COUNTER_STAT(TEXT("Asset Update Batch Size Received"), STAT_TedsAssetData_AssetUpdateBatchSizeReceived, STATGROUP_TedsAssetData);
+
+DECLARE_FLOAT_COUNTER_STAT(TEXT("Path Add Batch Last Duration"), STAT_TedsAssetData_PathAddBatchLastDuration, STATGROUP_TedsAssetData);
+DECLARE_FLOAT_COUNTER_STAT(TEXT("Asset Add Batch Last Duration"), STAT_TedsAssetData_AssetAddBatchLastDuration, STATGROUP_TedsAssetData);
+DECLARE_FLOAT_COUNTER_STAT(TEXT("Asset Update Batch Last Duration"), STAT_TedsAssetData_AssetUpdateBatchLastDuration, STATGROUP_TedsAssetData);
+
+TRACE_DECLARE_INT_COUNTER(COUNTER_TedsAssetData_PathAddsRemaining, TEXT("TedsAssetData/PathAdd/Remaining"));
+TRACE_DECLARE_INT_COUNTER(COUNTER_TedsAssetData_AssetAddsRemaining, TEXT("TedsAssetData/AssetAdd/Remaining"));
+TRACE_DECLARE_INT_COUNTER(COUNTER_TedsAssetData_AssetUpdatesRemaining, TEXT("TedsAssetData/AssetUpdate/Remaining"));
+
+TRACE_DECLARE_FLOAT_COUNTER(COUNTER_TedsAssetData_PathAddAvgTime, TEXT("TedsAssetData/PathAdd/AvgTime"));
+TRACE_DECLARE_FLOAT_COUNTER(COUNTER_TedsAssetData_AssetAddAvgTime, TEXT("TedsAssetData/AssetAdd/AvgTime"));
+TRACE_DECLARE_FLOAT_COUNTER(COUNTER_TedsAssetData_AssetUpdateAvgTime, TEXT("TedsAssetData/AssetUpdate/AvgTime"));
+
+TRACE_DECLARE_INT_COUNTER(COUNTER_TedsAssetData_PathAddBatchSize, TEXT("TedsAssetData/PathAdd/Batch/Size"));
+TRACE_DECLARE_INT_COUNTER(COUNTER_TedsAssetData_AssetAddBatchSize, TEXT("TedsAssetData/AssetAdd/Batch/Size"));
+TRACE_DECLARE_INT_COUNTER(COUNTER_TedsAssetData_AssetUpdateBatchSize, TEXT("TedsAssetData/AssetUpdate/Batch/Size"));
+
+TRACE_DECLARE_INT_COUNTER(COUNTER_TedsAssetData_PathAddBatchSizeReceived, TEXT("TedsAssetData/PathAdd/Batch/Size/Received"));
+TRACE_DECLARE_INT_COUNTER(COUNTER_TedsAssetData_AssetAddBatchSizeReceived, TEXT("TedsAssetData/AssetAdd/Batch/Size/Received"));
+TRACE_DECLARE_INT_COUNTER(COUNTER_TedsAssetData_AssetUpdateBatchSizeReceived, TEXT("TedsAssetData/AssetUpdate/Batch/Size/Received"));
+
+TRACE_DECLARE_FLOAT_COUNTER(COUNTER_TedsAssetData_PathAddBatchLastDuration, TEXT("TedsAssetData/PathAdd/Batch/LastDuration"));
+TRACE_DECLARE_FLOAT_COUNTER(COUNTER_TedsAssetData_AssetAddBatchLastDuration, TEXT("TedsAssetData/AssetAdd/Batch/LastDuration"));
+TRACE_DECLARE_FLOAT_COUNTER(COUNTER_TedsAssetData_AssetUpdateBatchLastDuration, TEXT("TedsAssetData/AssetUpdate/Batch/LastDuration"));
+
+DECLARE_FLOAT_COUNTER_STAT(TEXT("Total Asset Data processing time"), STAT_TedsAssetData_TotalProcessingTime, STATGROUP_TedsAssetData);
+
+TRACE_DECLARE_FLOAT_COUNTER(COUNTER_TedsAssetData_TotalProcessingTime, TEXT("TedsAssetData/TotalProcessingTime"));
+
+#endif
+
 namespace UE::Editor::AssetData::Private
 {
+
+std::atomic<uint32> FTedsAssetData::ProcessedAddedPathBatchSize = 0;
+std::atomic<uint32> FTedsAssetData::ProcessedAddedAssetBatchSize = 0;
+std::atomic<uint32> FTedsAssetData::ProcessedUpdatedAssetBatchSize = 0;
+
+std::atomic<uint64> FTedsAssetData::ProcessPathAddStartTime = std::numeric_limits<uint64>::max();
+std::atomic<uint64> FTedsAssetData::ProcessPathAddEndTime = 0;
+
+std::atomic<uint64> FTedsAssetData::ProcessAssetAddStartTime = std::numeric_limits<uint64>::max();
+std::atomic<uint64> FTedsAssetData::ProcessAssetAddEndTime = 0;
+
+std::atomic<uint64> FTedsAssetData::ProcessAssetUpdateStartTime = std::numeric_limits<uint64>::max();
+std::atomic<uint64> FTedsAssetData::ProcessAssetUpdateEndTime = 0;
+
 using namespace DataStorage;
 
 constexpr int32 ParallelForMinBatchSize = 1024 * 4;
@@ -47,7 +118,6 @@ FPopulateAssetDataRowArgs ThreadSafe_PopulateAssetDataTableRow(TAssetData&& InAs
 	return Output;
 }
 
-
 void PopulateAssetDataTableRow(FPopulateAssetDataRowArgs&& InAssetDataRowArgs, ICoreProvider& Database, RowHandle RowHandle)
 {
 	Database.GetColumn<FAssetDataColumn_Experimental>(RowHandle)->AssetData = MoveTemp(InAssetDataRowArgs.AssetData);
@@ -125,6 +195,8 @@ FTedsAssetData::FTedsAssetData(UE::Editor::DataStorage::ICoreProvider& InDatabas
 	LLM_SCOPE_BYNAME(TEXT("FTedsAssetData"))
 #endif
 
+	ProcessingStartTime = FPlatformTime::Cycles64();
+
 	// Register to events from asset registry
 	IAssetRegistry& AssetRegistry = IAssetRegistry::GetChecked();
 
@@ -136,6 +208,14 @@ FTedsAssetData::FTedsAssetData(UE::Editor::DataStorage::ICoreProvider& InDatabas
 	AssetRegistry.OnPathsAdded().AddRaw(this, &FTedsAssetData::OnPathsAdded);
 	AssetRegistry.OnPathsRemoved().AddRaw(this, &FTedsAssetData::OnPathsRemoved);
 
+	Database.OnUpdate().AddRaw(this, &FTedsAssetData::OnDatabaseUpdate);
+
+	Database.RegisterCooperativeUpdate(TEXT("ProcessAssetTasksCooperatively"), DataStorage::ICoreProvider::ECooperativeTaskPriority::High,
+		[this](FTimespan TimeAllowance)
+		{
+			ProcessTaskQueues(TimeAllowance);
+		});
+
 	// Register data types to TEDS
 	PathsTable = Database.FindTable(FName(TEXT("Editor_AssetRegistryPathsTable")));
 	if (PathsTable == InvalidTableHandle)
@@ -149,7 +229,6 @@ FTedsAssetData::FTedsAssetData(UE::Editor::DataStorage::ICoreProvider& InDatabas
 		AssetsDataTable = Database.RegisterTable<FAssetDataColumn_Experimental, FUpdatedPathTag, FUpdatedAssetDataTag>(FName(TEXT("Editor_AssetRegistryAssetDataTable")));
 	}
 
-
 	RemoveUpdatedPathTagQuery = Database.RegisterQuery(
 		Select(
 			TEXT("FTedsAssetData: Remove Updated Path Tag"),
@@ -175,8 +254,18 @@ FTedsAssetData::FTedsAssetData(UE::Editor::DataStorage::ICoreProvider& InDatabas
 		.Compile());
 
 	// Init with the data existing at moment in asset registry
-	TArray<FAssetData> AssetsData;
-	AssetRegistry.GetAllAssets(AssetsData);
+	UE::AssetRegistry::FFiltering::InitializeShouldSkipAsset();
+	TChunkedArray<FAssetData> CachedAssets;
+	AssetRegistry.EnumerateAllAssets([&CachedAssets](const FAssetData& AssetData)
+		{
+			if (!UE::AssetRegistry::FFiltering::ShouldSkipAsset(AssetData.AssetClassPath, AssetData.PackageFlags) &&
+				ContentBrowserDataUtils::IsPrimaryAsset(AssetData))
+			{
+				CachedAssets.AddElement(AssetData);
+			}
+			return true;
+		});
+	CachedAssets.MoveToLinearArray(AddedAssets);
 
 	TChunkedArray<FName> CachedPaths;
 	AssetRegistry.EnumerateAllCachedPaths([&CachedPaths](FName Name)
@@ -184,65 +273,7 @@ FTedsAssetData::FTedsAssetData(UE::Editor::DataStorage::ICoreProvider& InDatabas
 			CachedPaths.AddElement(Name);
 			return true;
 		});
-
-	TArray<FNameBuilder> NameBuilders;
-	TArray<FPopulatePathRowArgs> PopulatePathRowArgs;
-	PopulatePathRowArgs.AddDefaulted(CachedPaths.Num());
-
-	// Prepare Path Rows
-	ParallelForWithTaskContext(TEXT("Populating TEDS Asset Registry Path"), NameBuilders, CachedPaths.Num(),ParallelForMinBatchSize, [&PopulatePathRowArgs, &CachedPaths, this](FNameBuilder& NameBuilder, int32 Index)
-		{
-			FName Path = CachedPaths[Index];
-			Path.ToString(NameBuilder);
-			PopulatePathRowArgs[Index] = ThreadSafe_PopulatePathRowArgs(FMapKey(Path), Path, NameBuilder);
-		});
-
-	TArray<RowHandle> ReservedRows;
-	ReservedRows.AddUninitialized(PopulatePathRowArgs.Num() + AssetsData.Num());
-	Database.BatchReserveRows(ReservedRows);
-
-	TArray<TPair<FMapKey, RowHandle>> KeysToReservedRows;
-	KeysToReservedRows.AddDefaulted(ReservedRows.Num());
-
-	// Index Reserved Path Rows
-	TConstArrayView<RowHandle> ReservedPopulatePathRows(ReservedRows.GetData(), PopulatePathRowArgs.Num());
-	TArrayView<TPair<FMapKey, RowHandle>> KeysToReservedPathRows(KeysToReservedRows.GetData(), PopulatePathRowArgs.Num());
-	ParallelFor(TEXT("Populating TEDS Asset Registry Path Data Keys"), ReservedPopulatePathRows.Num(), ParallelForMinBatchSize, [&KeysToReservedPathRows, &ReservedPopulatePathRows, &PopulatePathRowArgs](int32 Index)
-		{
-			KeysToReservedPathRows[Index] = TPair<FMapKey, RowHandle>(PopulatePathRowArgs[Index].AssetRegistryPathKey, ReservedPopulatePathRows[Index]);
-		});
-	Database.BatchMapRows(MappingDomain, KeysToReservedPathRows);
-
-	// Populate Path Rows
-	Database.BatchAddRow(PathsTable, ReservedPopulatePathRows, [PathRowArgs = MoveTemp(PopulatePathRowArgs), Index = 0, this](RowHandle InRowHandle) mutable
-		{
-			PopulatePathDataTableRow(MoveTemp(PathRowArgs[Index]), Database, InRowHandle);
-			++Index;
-		});
-
-	// Prepare Asset Data Rows
-	TArray<FPopulateAssetDataRowArgs> PopulateAssetDataRowArgs;
-	PopulateAssetDataRowArgs.AddDefaulted(AssetsData.Num());
-	ParallelFor(TEXT("Populating TEDS Asset Registry Asset Data"), PopulateAssetDataRowArgs.Num(), ParallelForMinBatchSize, [&PopulateAssetDataRowArgs, &AssetsData, this](int32 Index)
-		{
-			PopulateAssetDataRowArgs[Index] = ThreadSafe_PopulateAssetDataTableRow(MoveTemp(AssetsData[Index]), Database);
-		});
-
-	// Index Reserved Asset Data Rows
-	TConstArrayView<RowHandle> ReservedAssetDataRows(ReservedRows.GetData() + ReservedPopulatePathRows.Num(), AssetsData.Num());
-	TArrayView<TPair<FMapKey, RowHandle>> KeysToResevedAssetRows(KeysToReservedRows.GetData() + ReservedPopulatePathRows.Num(), ReservedAssetDataRows.Num());
-	ParallelFor(TEXT("Populating TEDS Asset Registry Asset Data Keys"), PopulateAssetDataRowArgs.Num(), ParallelForMinBatchSize, [&PopulateAssetDataRowArgs, &ReservedAssetDataRows, &KeysToResevedAssetRows](int32 Index)
-		{
-			KeysToResevedAssetRows[Index] = TPair<FMapKey, RowHandle>(MoveTemp(PopulateAssetDataRowArgs[Index].ObjectPathKey), ReservedAssetDataRows[Index]);
-		});
-	Database.BatchMapRows(MappingDomain, KeysToResevedAssetRows);
-
-	// Populate Asset Rows
-	Database.BatchAddRow(AssetsDataTable, ReservedAssetDataRows,  [AssetDataRowArgs = MoveTemp(PopulateAssetDataRowArgs), Index = 0, this](RowHandle InRowHandle) mutable
-		{
-			PopulateAssetDataTableRow(MoveTemp(AssetDataRowArgs[Index]), Database, InRowHandle);
-			++Index;
-		});
+	CachedPaths.MoveToLinearArray(AddedPaths);
 }
 
 FTedsAssetData::~FTedsAssetData()
@@ -256,12 +287,15 @@ FTedsAssetData::~FTedsAssetData()
 	// Not needed on a editor shut down
 	if (!IsEngineExitRequested())
 	{
+		Database.UnregisterCooperativeUpdate(TEXT("ProcessAssetTasksCooperatively"));
+
+		Database.OnUpdate().RemoveAll(this);
+
 		if (IAssetRegistry* AssetRegistry = IAssetRegistry::Get())
 		{
 			Database.UnregisterQuery(RemoveUpdatedAssetDataTagQuery);
 			Database.UnregisterQuery(RemoveUpdatedPathTagQuery);
 
-
 			AssetRegistry->OnAssetsAdded().RemoveAll(this);
 			AssetRegistry->OnAssetsRemoved().RemoveAll(this);
 			AssetRegistry->OnAssetsUpdated().RemoveAll(this);
@@ -301,28 +335,55 @@ void FTedsAssetData::OnAssetsAdded(TConstArrayView<FAssetData> InAssetsAdded)
 {
 	TRACE_CPUPROFILER_EVENT_SCOPE(FTedsAssetData::OnAssetsAdded);
 
+	AddedAssets.Reserve(AddedAssets.Num() + InAssetsAdded.Num());
+
+	UE::AssetRegistry::FFiltering::InitializeShouldSkipAsset();
+	for (const FAssetData& AssetData : InAssetsAdded)
+	{
+		if (!UE::AssetRegistry::FFiltering::ShouldSkipAsset(AssetData.AssetClassPath, AssetData.PackageFlags) &&
+			ContentBrowserDataUtils::IsPrimaryAsset(AssetData))
+		{
+			AddedAssets.Emplace(AssetData);
+		}
+	}
+}
+
+FTimespan FTedsAssetData::ProcessAddedAssets(const FTimespan& AllowedTime)
+{
+	TRACE_CPUPROFILER_EVENT_SCOPE(FTedsAssetData::ProcessAddedAssets);
+
 #if TRACK_TEDSASSETDATA_MEMORY
 	LLM_SCOPE_BYNAME(TEXT("FTedsAssetData"))
 #endif
 
+	if(AddedAssets.IsEmpty())
+	{
+		LastAddedAssetBatchSize = 0;
+		return AllowedTime;
+	}
+
+	int32 BatchSize = static_cast<int32>((AllowedTime / AvgTimePerProcessedAssetAdd).GetTicks());
+
+	// Regardless of AvgTime always process at least the minimum batch size
+	uint32 BatchMin = FMath::Min(MinimumBatchSize, AddedAssets.Num());
+	BatchSize = static_cast<int32>(FMath::Clamp(BatchSize, BatchMin, AddedAssets.Num()));
+
+	uint32 BatchStartIndex = AddedAssets.Num() - BatchSize;
+
+	TConstArrayView<FAssetData> AssetBatch(&AddedAssets[BatchStartIndex], BatchSize);
 	TArray<int32> Contexts;
 	TArray<FPopulateAssetDataRowArgs> PopulateRowArgs;
-	PopulateRowArgs.AddDefaulted(InAssetsAdded.Num());
-
-	UE::AssetRegistry::FFiltering::InitializeShouldSkipAsset();
+	PopulateRowArgs.AddDefaulted(AssetBatch.Num());
 
-	ParallelForWithTaskContext(TEXT("Populating TEDS Asset Registry Asset Data"), Contexts, InAssetsAdded.Num(), ParallelForMinBatchSize, [&PopulateRowArgs, &InAssetsAdded, this](int32& ValidAssetCount, int32 Index)
+	ParallelForWithTaskContext(TEXT("Populating TEDS Asset Registry Asset Data"), Contexts, AssetBatch.Num(), ParallelForMinBatchSize, [&PopulateRowArgs, &AssetBatch, this](int32& ValidAssetCount, int32 Index)
 		{
 			FPopulateAssetDataRowArgs RowArgs;
-			const FAssetData& AssetData = InAssetsAdded[Index];
+			const FAssetData& AssetData = AssetBatch[Index];
 
-			if (!UE::AssetRegistry::FFiltering::ShouldSkipAsset(AssetData.AssetClassPath, AssetData.PackageFlags))
+			RowArgs = ThreadSafe_PopulateAssetDataTableRow(AssetData, Database);
+			if (RowArgs.AssetData.IsValid())
 			{
-				RowArgs = ThreadSafe_PopulateAssetDataTableRow(AssetData, Database);
-				if (RowArgs.AssetData.IsValid())
-				{
-					++ValidAssetCount;
-				}
+				++ValidAssetCount;
 			}
 			PopulateRowArgs[Index] = MoveTemp(RowArgs);
 		});
@@ -355,6 +416,13 @@ void FTedsAssetData::OnAssetsAdded(TConstArrayView<FAssetData> InAssetsAdded)
 
 		Database.BatchMapRows(MappingDomain, KeyToRow);
 	}
+
+	// Cut the consumed batch off of the back of our queue
+	AddedAssets.RemoveAt(BatchStartIndex, BatchSize);
+
+	FTimespan ConsumedTime(static_cast<int64>(AvgTimePerProcessedAssetAdd * BatchSize));
+	LastAddedAssetBatchSize += BatchSize; // Accumulate total batch size since we can be called multiple times per tick
+	return AllowedTime - ConsumedTime;
 }
 
 void FTedsAssetData::OnAssetsRemoved(TConstArrayView<FAssetData> InAssetsRemoved)
@@ -376,6 +444,8 @@ void FTedsAssetData::OnAssetsRemoved(TConstArrayView<FAssetData> InAssetsRemoved
 		{
 			RowsToRemove.Add(AssetRow);
 		}
+		AddedAssets.Remove(Asset);
+		UpdatedAssets.Remove(Asset);
 	}
 
 	Database.BatchRemoveRows(RowsToRemove);
@@ -385,31 +455,48 @@ void FTedsAssetData::OnAssetsUpdated(TConstArrayView<FAssetData> InAssetsUpdated
 {
 	TRACE_CPUPROFILER_EVENT_SCOPE(FTedsAssetData::OnAssetsUpdated);
 
-#if TRACK_TEDSASSETDATA_MEMORY
-	LLM_SCOPE_BYNAME(TEXT("FTedsAssetData"))
-#endif
-
-	for (const FAssetData& Asset : InAssetsUpdated)
-	{
-		const FMapKey AssetKey = FMapKey(Asset.GetSoftObjectPath());
-		const RowHandle Row = Database.LookupMappedRow(MappingDomain, AssetKey);
-		if (Database.IsRowAssigned(Row))
-		{
-			Database.GetColumn<FAssetDataColumn_Experimental>(Row)->AssetData = Asset;
-			Database.AddColumn<FUpdatedAssetDataTag>(Row);
-		}
-	}
+	OnAssetsUpdated_Impl(InAssetsUpdated);
 }
 
 void FTedsAssetData::OnAssetsUpdatedOnDisk(TConstArrayView<FAssetData> InAssetsUpdated)
 {
 	TRACE_CPUPROFILER_EVENT_SCOPE(FTedsAssetData::OnAssetsUpdatedOnDisk);
 
+	OnAssetsUpdated_Impl(InAssetsUpdated);
+}
+
+void FTedsAssetData::OnAssetsUpdated_Impl(TConstArrayView<FAssetData> InAssetsUpdated)
+{
+	TRACE_CPUPROFILER_EVENT_SCOPE(FTedsAssetData::OnAssetsUpdated_Impl);
+
+	UpdatedAssets.Append(InAssetsUpdated.GetData(), InAssetsUpdated.Num());
+}
+
+FTimespan FTedsAssetData::ProcessUpdatedAssets(const FTimespan& AllowedTime)
+{
+	TRACE_CPUPROFILER_EVENT_SCOPE(FTedsAssetData::ProcessUpdatedAssets);
+
 #if TRACK_TEDSASSETDATA_MEMORY
 	LLM_SCOPE_BYNAME(TEXT("FTedsAssetData"))
 #endif
 
-	for (const FAssetData& Asset : InAssetsUpdated)
+	if (UpdatedAssets.IsEmpty())
+	{
+		LastUpdatedAssetBatchSize = 0;
+		return AllowedTime;
+	}
+
+	int32 BatchSize = static_cast<int32>((AllowedTime / AvgTimePerProcessedAssetUpdate).GetTicks());
+
+	// Regardless of AvgTime always process at least the minimum batch size
+	uint32 BatchMin = FMath::Min(MinimumBatchSize, UpdatedAssets.Num());
+	BatchSize = static_cast<int32>(FMath::Clamp(BatchSize, BatchMin, UpdatedAssets.Num()));
+
+	uint32 BatchStartIndex = UpdatedAssets.Num() - BatchSize;
+
+	TConstArrayView<FAssetData> AssetBatch(&UpdatedAssets[BatchStartIndex], BatchSize);
+
+	for (const FAssetData& Asset : AssetBatch)
 	{
 		const FMapKey AssetKey = FMapKey(Asset.GetSoftObjectPath());
 		const RowHandle Row = Database.LookupMappedRow(MappingDomain, AssetKey);
@@ -419,6 +506,13 @@ void FTedsAssetData::OnAssetsUpdatedOnDisk(TConstArrayView<FAssetData> InAssetsU
 			Database.AddColumn<FUpdatedAssetDataTag>(Row);
 		}
 	}
+
+	// Cut the consumed batch off of the back of our queue
+	UpdatedAssets.RemoveAt(BatchStartIndex, BatchSize);
+
+	FTimespan ConsumedTime(static_cast<int64>(AvgTimePerProcessedAssetUpdate * BatchSize));
+	LastUpdatedAssetBatchSize += BatchSize; // Accumulate total batch size since we can be called multiple times per tick
+	return AllowedTime - ConsumedTime;
 }
 
 void FTedsAssetData::OnAssetRenamed(const FAssetData& InAsset, const FString& InOldObjectPath)
@@ -459,18 +553,50 @@ void FTedsAssetData::OnPathsAdded(TConstArrayView<FStringView> InPathsAdded)
 {
 	TRACE_CPUPROFILER_EVENT_SCOPE(FTedsAssetData::OnPathsAdded);
 
+	int32 AddedPathsSize = AddedPaths.Num();
+	AddedPaths.AddDefaulted(InPathsAdded.Num());
+
+	// Would normally append InPathsAdded to AddedPaths, but we need to convert the FStringView to FName
+	ParallelFor(TEXT("Staging TEDS Asset Registry Path"), InPathsAdded.Num(), ParallelForMinBatchSize, [&InPathsAdded, AddedPathsSize, this](int32 Index)
+		{
+			AddedPaths[AddedPathsSize + Index] = FName(InPathsAdded[Index]);
+		});
+}
+
+FTimespan FTedsAssetData::ProcessAddedPaths(const FTimespan& AllowedTime)
+{
+	TRACE_CPUPROFILER_EVENT_SCOPE(FTedsAssetData::ProcessAddedPaths);
+
 #if TRACK_TEDSASSETDATA_MEMORY
 	LLM_SCOPE_BYNAME(TEXT("FTedsAssetData"));
 #endif
 
+	if (AddedPaths.IsEmpty())
+	{
+		LastAddedPathBatchSize = 0;
+		return AllowedTime;
+	}
+
+	int32 BatchSize = static_cast<int32>((AllowedTime / AvgTimePerProcessedPathAdd).GetTicks());
+
+	// Regardless of AvgTime always process at least the minimum batch size
+	uint32 BatchMin = FMath::Min(MinimumBatchSize, AddedPaths.Num());
+	BatchSize = static_cast<int32>(FMath::Clamp(BatchSize, BatchMin, AddedPaths.Num()));
+
+	uint32 BatchStartIndex = AddedPaths.Num() - BatchSize;
+
+	TConstArrayView<FName> PathBatch(&AddedPaths[BatchStartIndex], BatchSize);
+
 	TArray<int32> PreparePathRowContexts;
 	TArray<FPopulatePathRowArgs> PopulateRowArgs;
-	PopulateRowArgs.AddDefaulted(InPathsAdded.Num());
+	PopulateRowArgs.AddDefaulted(PathBatch.Num());
 
-	ParallelForWithTaskContext(TEXT("Populating TEDS Asset Registry Path"), PreparePathRowContexts, InPathsAdded.Num(), ParallelForMinBatchSize, [&PopulateRowArgs, &InPathsAdded, this](int32& WorkerValidCount, int32 Index)
+	ParallelForWithTaskContext(TEXT("Populating TEDS Asset Registry Path"), PreparePathRowContexts, PathBatch.Num(), ParallelForMinBatchSize, [&PopulateRowArgs, &PathBatch, this](int32& WorkerValidCount, int32 Index)
 		{
-			FStringView Path = InPathsAdded[Index];
-			FName PathName(Path);
+			FName PathName = PathBatch[Index];
+			FNameBuilder Path;
+			PathName.ToString(Path);
+
 			FMapKey AssetRegistryPathKey = FMapKey(PathName);
 			FPopulatePathRowArgs RowArgs;
 
@@ -493,7 +619,6 @@ void FTedsAssetData::OnPathsAdded(TConstArrayView<FStringView> InPathsAdded)
 		NewRowsCount += Context;
 	}
 
-
 	if (NewRowsCount > 0)
 	{
 		TArray<RowHandle> ReservedRow;
@@ -533,6 +658,13 @@ void FTedsAssetData::OnPathsAdded(TConstArrayView<FStringView> InPathsAdded)
 				++Index;
 			});
 	}
+
+	// Cut the consumed batch off of the back of our queue
+	AddedPaths.RemoveAt(BatchStartIndex, BatchSize);
+
+	FTimespan ConsumedTime(static_cast<int64>(AvgTimePerProcessedPathAdd * BatchSize));
+	LastAddedPathBatchSize += BatchSize; // Accumulate total batch size since we can be called multiple times per tick
+	return AllowedTime - ConsumedTime;
 }
  
 void FTedsAssetData::OnPathsRemoved(TConstArrayView<FStringView> InPathsRemoved)
@@ -545,10 +677,152 @@ void FTedsAssetData::OnPathsRemoved(TConstArrayView<FStringView> InPathsRemoved)
 
 	for (const FStringView Path : InPathsRemoved)
 	{
-		FMapKey PathKey = FMapKey(FName(Path));
+		FName PathName(Path);
+		FMapKey PathKey = FMapKey(PathName);
 		Database.RemoveRow(Database.LookupMappedRow(MappingDomain, PathKey));
+		AddedPaths.Remove(PathName);
+	}
+}
+
+void FTedsAssetData::OnDatabaseUpdate()
+{
+	TRACE_CPUPROFILER_EVENT_SCOPE(FTedsAssetData::OnDatabaseUpdate);
+
+	UpdateAvgProcessedTime(AvgTimePerProcessedPathAdd, ProcessPathAddStartTime, ProcessPathAddEndTime, LastAddedPathBatchSize);
+	UpdateAvgProcessedTime(AvgTimePerProcessedAssetAdd, ProcessAssetAddStartTime, ProcessAssetAddEndTime, LastAddedAssetBatchSize);
+	UpdateAvgProcessedTime(AvgTimePerProcessedAssetUpdate, ProcessAssetUpdateStartTime, ProcessAssetUpdateEndTime, LastUpdatedAssetBatchSize);
+
+#if COUNTERSTRACE_ENABLED
+	RecordAssetDataStats();
+#endif
+
+	LastAddedPathBatchSize = 0;
+	LastAddedAssetBatchSize = 0;
+	LastUpdatedAssetBatchSize = 0;
+
+	ProcessedAddedPathBatchSize = 0;
+	ProcessedAddedAssetBatchSize = 0;
+	ProcessedUpdatedAssetBatchSize = 0;
+
+	ProcessTaskQueues(MaxAllowedTimePerTick);
+
+	ProcessPathAddStartTime = std::numeric_limits<uint64>::max();
+	ProcessPathAddEndTime = 0;
+
+	ProcessAssetAddStartTime = std::numeric_limits<uint64>::max();
+	ProcessAssetAddEndTime = 0;
+
+	ProcessAssetUpdateStartTime = std::numeric_limits<uint64>::max();
+	ProcessAssetUpdateEndTime = 0;
+}
+
+void FTedsAssetData::UpdateAvgProcessedTime(double& OutAverageTime, uint64 LatestStartTime, uint64 LatestEndTime, int32 LastBatchSize)
+{
+	constexpr double PreviousTicksWeight = 0.7;
+	constexpr double LatestTickWeight = 0.3;
+
+	static_assert(PreviousTicksWeight + LatestTickWeight == 1.0);
+
+	if (LastBatchSize != 0 && LatestEndTime != 0)
+	{
+		FTimespan ProcessorDuration = FTimespan::FromSeconds(FPlatformTime::ToSeconds64(LatestEndTime - LatestStartTime));
+		double LatestTimePerUnit = static_cast<double>((ProcessorDuration / LastBatchSize).GetTicks());
+		LatestTimePerUnit = FMath::Max(LatestTimePerUnit, 1); // Timespan can truncate to 0 so give it a min of 1
+
+		OutAverageTime = (OutAverageTime * PreviousTicksWeight) + (LatestTimePerUnit * LatestTickWeight);
+		OutAverageTime = FMath::Max(OutAverageTime, 1);
 	}
 }
+
+void FTedsAssetData::ProcessTaskQueues(const FTimespan AllowedTime)
+{
+	TRACE_CPUPROFILER_EVENT_SCOPE(FTedsAssetData::ProcessTaskQueues);
+
+	if (AllowedTime.IsZero())
+	{
+		return;
+	}
+
+	int32 TaskTypeCount = AddedPaths.IsEmpty() ? 0 : 1;
+	TaskTypeCount += AddedAssets.IsEmpty() ? 0 : 1;
+	TaskTypeCount += UpdatedAssets.IsEmpty() ? 0 : 1;
+
+	// By dividing allowed time by number of filled task queues we allocate time to each task type instead of all to the first type
+	FTimespan AllowedTimePerTaskType = AllowedTime / TaskTypeCount;
+
+	FTimespan UnusedTime = 0;
+	if (!AddedPaths.IsEmpty())
+	{
+		UnusedTime = ProcessAddedPaths(AllowedTimePerTaskType);
+	}
+
+	if (!AddedAssets.IsEmpty())
+	{
+		UnusedTime = ProcessAddedAssets(AllowedTimePerTaskType + UnusedTime);
+	}
+
+	if (!UpdatedAssets.IsEmpty())
+	{
+		ProcessUpdatedAssets(AllowedTimePerTaskType + UnusedTime);
+	}
+}
+
+void FTedsAssetData::RecordAssetDataStats()
+{
+#if COUNTERSTRACE_ENABLED
+
+	TRACE_CPUPROFILER_EVENT_SCOPE(FTedsAssetData::RecordAssetDataStats);
+
+	SET_DWORD_STAT(STAT_TedsAssetData_PathAddsRemaining, AddedPaths.Num());
+	SET_DWORD_STAT(STAT_TedsAssetData_AssetAddsRemaining, AddedAssets.Num());
+	SET_DWORD_STAT(STAT_TedsAssetData_AssetUpdatesRemaining, UpdatedAssets.Num());
+
+	TRACE_COUNTER_SET(COUNTER_TedsAssetData_PathAddsRemaining, AddedPaths.Num());
+	TRACE_COUNTER_SET(COUNTER_TedsAssetData_AssetAddsRemaining, AddedAssets.Num());
+	TRACE_COUNTER_SET(COUNTER_TedsAssetData_AssetUpdatesRemaining, UpdatedAssets.Num());
+
+	SET_DWORD_STAT(STAT_TedsAssetData_PathAddBatchSize, LastAddedPathBatchSize);
+	SET_DWORD_STAT(STAT_TedsAssetData_PathAddBatchSize, LastAddedAssetBatchSize);
+	SET_DWORD_STAT(STAT_TedsAssetData_PathAddBatchSize, LastUpdatedAssetBatchSize);
+
+	TRACE_COUNTER_SET(COUNTER_TedsAssetData_PathAddBatchSize, LastAddedPathBatchSize);
+	TRACE_COUNTER_SET(COUNTER_TedsAssetData_AssetAddBatchSize, LastAddedAssetBatchSize);
+	TRACE_COUNTER_SET(COUNTER_TedsAssetData_AssetUpdateBatchSize, LastUpdatedAssetBatchSize);
+
+	SET_DWORD_STAT(STAT_TedsAssetData_PathAddBatchSizeReceived, ProcessedAddedPathBatchSize.load());
+	SET_DWORD_STAT(STAT_TedsAssetData_AssetAddBatchSizeReceived, ProcessedAddedAssetBatchSize.load());
+	SET_DWORD_STAT(STAT_TedsAssetData_AssetUpdateBatchSizeReceived, ProcessedUpdatedAssetBatchSize.load());
+
+	TRACE_COUNTER_SET(COUNTER_TedsAssetData_PathAddBatchSizeReceived, ProcessedAddedPathBatchSize.load());
+	TRACE_COUNTER_SET(COUNTER_TedsAssetData_AssetAddBatchSizeReceived, ProcessedAddedAssetBatchSize.load());
+	TRACE_COUNTER_SET(COUNTER_TedsAssetData_AssetUpdateBatchSizeReceived, ProcessedUpdatedAssetBatchSize.load());
+
+	SET_FLOAT_STAT(STAT_TedsAssetData_PathAddAvgTime, AvgTimePerProcessedPathAdd / ETimespan::TicksPerMillisecond);
+	SET_FLOAT_STAT(STAT_TedsAssetData_AssetAddAvgTime, AvgTimePerProcessedAssetAdd / ETimespan::TicksPerMillisecond);
+	SET_FLOAT_STAT(STAT_TedsAssetData_AssetUpdateAvgTime, AvgTimePerProcessedAssetUpdate / ETimespan::TicksPerMillisecond);
+
+	TRACE_COUNTER_SET(COUNTER_TedsAssetData_PathAddAvgTime, AvgTimePerProcessedPathAdd / ETimespan::TicksPerMillisecond);
+	TRACE_COUNTER_SET(COUNTER_TedsAssetData_AssetAddAvgTime, AvgTimePerProcessedAssetAdd / ETimespan::TicksPerMillisecond);
+	TRACE_COUNTER_SET(COUNTER_TedsAssetData_AssetUpdateAvgTime, AvgTimePerProcessedAssetUpdate / ETimespan::TicksPerMillisecond);
+
+	SET_FLOAT_STAT(STAT_TedsAssetData_PathAddBatchLastDuration, FTimespan::FromSeconds(FPlatformTime::ToSeconds64(ProcessPathAddEndTime - ProcessPathAddStartTime)).GetTotalMilliseconds());
+	SET_FLOAT_STAT(STAT_TedsAssetData_AssetAddBatchLastDuration, FTimespan::FromSeconds(FPlatformTime::ToSeconds64(ProcessAssetAddEndTime - ProcessAssetAddStartTime)).GetTotalMilliseconds());
+	SET_FLOAT_STAT(STAT_TedsAssetData_AssetUpdateBatchLastDuration, FTimespan::FromSeconds(FPlatformTime::ToSeconds64(ProcessAssetUpdateEndTime - ProcessAssetUpdateStartTime)).GetTotalMilliseconds());
+
+	TRACE_COUNTER_SET(COUNTER_TedsAssetData_PathAddBatchLastDuration, FTimespan::FromSeconds(FPlatformTime::ToSeconds64(ProcessPathAddEndTime - ProcessPathAddStartTime)).GetTotalMilliseconds());
+	TRACE_COUNTER_SET(COUNTER_TedsAssetData_AssetAddBatchLastDuration, FTimespan::FromSeconds(FPlatformTime::ToSeconds64(ProcessAssetAddEndTime - ProcessAssetAddStartTime)).GetTotalMilliseconds());
+	TRACE_COUNTER_SET(COUNTER_TedsAssetData_AssetUpdateBatchLastDuration, FTimespan::FromSeconds(FPlatformTime::ToSeconds64(ProcessAssetUpdateEndTime - ProcessAssetUpdateStartTime)).GetTotalMilliseconds());
+
+	if ((AddedPaths.Num() + AddedAssets.Num() + UpdatedAssets.Num()) > 0)
+	{
+		SET_FLOAT_STAT(STAT_TedsAssetData_TotalProcessingTime, FPlatformTime::ToSeconds64(FPlatformTime::Cycles64() - ProcessingStartTime));
+
+		TRACE_COUNTER_SET(COUNTER_TedsAssetData_TotalProcessingTime, FPlatformTime::ToSeconds64(FPlatformTime::Cycles64() - ProcessingStartTime));
+	}
+
+#endif
+}
+
 } // namespace UE::Editor::AssetData::Private
 
 #undef TRACK_TEDSASSETDATA_MEMORY
diff --git a/Engine/Plugins/Experimental/EditorDataStorageFeatures/Source/TedsAssetData/Private/TedsAssetData.h b/Engine/Plugins/Experimental/EditorDataStorageFeatures/Source/TedsAssetData/Private/TedsAssetData.h
--- a/Engine/Plugins/Experimental/EditorDataStorageFeatures/Source/TedsAssetData/Private/TedsAssetData.h
+++ b/Engine/Plugins/Experimental/EditorDataStorageFeatures/Source/TedsAssetData/Private/TedsAssetData.h
@@ -20,6 +20,8 @@ namespace UE::Editor::AssetData::Private
 	 */
 	class FTedsAssetData
 	{
+		friend class FTedsAssetDataCBDataSource;
+
 	public:
 		FTedsAssetData(UE::Editor::DataStorage::ICoreProvider& InDatabase);
 
@@ -29,15 +31,62 @@ namespace UE::Editor::AssetData::Private
 
 	private:
 		void OnAssetsAdded(TConstArrayView<FAssetData> InAssetsAdded);
+		FTimespan ProcessAddedAssets(const FTimespan& AllowedTime);
+
 		void OnAssetsRemoved(TConstArrayView<FAssetData> InAssetsRemoved);
+
 		void OnAssetsUpdated(TConstArrayView<FAssetData> InAssetsUpdated);
 		void OnAssetsUpdatedOnDisk(TConstArrayView<FAssetData> InAssetsUpdated);
+		void OnAssetsUpdated_Impl(TConstArrayView<FAssetData> InAssetsUpdated);
+		FTimespan ProcessUpdatedAssets(const FTimespan& AllowedTime);
 
 		void OnAssetRenamed(const FAssetData& InAsset, const FString& InOldObjectPath);
 
 		void OnPathsAdded(TConstArrayView<FStringView> InPathsAdded);
+		FTimespan ProcessAddedPaths(const FTimespan& AllowedTIme);
+
 		void OnPathsRemoved(TConstArrayView<FStringView> InPathsRemoved);
 
+		void OnDatabaseUpdate();
+
+		void UpdateAvgProcessedTime(double& OutAverageTime, uint64 LatestStartTime, uint64 LatestEndTime, int32 LastBatchSize);
+
+		void ProcessTaskQueues(const FTimespan AllowedTime);
+
+		void RecordAssetDataStats();
+
+		TArray<FName> AddedPaths;
+		TArray<FAssetData> AddedAssets;
+		TArray<FAssetData> UpdatedAssets;
+
+		const FTimespan MaxAllowedTimePerTick = FTimespan::FromMilliseconds(6);
+		const FTimespan DefaultMaxTimePerProcessCategory = MaxAllowedTimePerTick / 3;
+		const uint32 DefaultNumberOfProcessedItems = 500;
+
+		double AvgTimePerProcessedPathAdd = static_cast<double>((DefaultMaxTimePerProcessCategory / DefaultNumberOfProcessedItems).GetTicks());
+		double AvgTimePerProcessedAssetAdd = static_cast<double>((DefaultMaxTimePerProcessCategory / DefaultNumberOfProcessedItems).GetTicks());
+		double AvgTimePerProcessedAssetUpdate = static_cast<double>((DefaultMaxTimePerProcessCategory / DefaultNumberOfProcessedItems).GetTicks());
+
+		const int32 MinimumBatchSize = 250;
+		int32 LastAddedPathBatchSize = 0;
+		int32 LastAddedAssetBatchSize = 0;
+		int32 LastUpdatedAssetBatchSize = 0;
+
+		static std::atomic<uint32> ProcessedAddedPathBatchSize;
+		static std::atomic<uint32> ProcessedAddedAssetBatchSize;
+		static std::atomic<uint32> ProcessedUpdatedAssetBatchSize;
+
+		uint64 ProcessingStartTime = 0;
+
+		static std::atomic<uint64> ProcessPathAddStartTime;
+		static std::atomic<uint64> ProcessPathAddEndTime;
+
+		static std::atomic<uint64> ProcessAssetAddStartTime;
+		static std::atomic<uint64> ProcessAssetAddEndTime;
+
+		static std::atomic<uint64> ProcessAssetUpdateStartTime;
+		static std::atomic<uint64> ProcessAssetUpdateEndTime;
+
 		UE::Editor::DataStorage::ICoreProvider& Database;
 		DataStorage::TableHandle PathsTable = DataStorage::InvalidTableHandle;
 		DataStorage::TableHandle AssetsDataTable = DataStorage::InvalidTableHandle;
