Fix detected code issues with new TedsAssetData behavior

diff --git a/Engine/Plugins/Experimental/EditorDataStorageFeatures/Source/TedsAssetData/Private/CB/TedsAssetDataCBDataSource.cpp b/Engine/Plugins/Experimental/EditorDataStorageFeatures/Source/TedsAssetData/Private/CB/TedsAssetDataCBDataSource.cpp
--- a/Engine/Plugins/Experimental/EditorDataStorageFeatures/Source/TedsAssetData/Private/CB/TedsAssetDataCBDataSource.cpp
+++ b/Engine/Plugins/Experimental/EditorDataStorageFeatures/Source/TedsAssetData/Private/CB/TedsAssetDataCBDataSource.cpp
@@ -365,9 +365,15 @@ void FTedsAssetDataCBDataSource::ProcessPathQueryCallback(DataStorage::IQueryCon
 	int32 NumOfRowToProcess = Context.GetRowCount();
 	FNameBuilder InternalPath;
 	FNameBuilder VirtualPath;
+	FString Path;
 
 	for (int32 Index = 0; Index < NumOfRowToProcess; ++Index)
 	{
+		Path = PathColumn[Index].Path.ToString();
+
+		InternalPath.Reset();
+		InternalPath.Append(PathColumn[Index].Path.ToString());
+
 		if (GenerateVirtualPath(InternalPath, VirtualPath))
 		{
 			FVirtualPathColumn_Experimental VirtualPathColumn;
@@ -375,8 +381,6 @@ void FTedsAssetDataCBDataSource::ProcessPathQueryCallback(DataStorage::IQueryCon
 			// Todo investigate for a batch add Maybe?
 			Context.AddColumn(Rows[Index], MoveTemp(VirtualPathColumn));
 
-
-			FString Path = PathColumn->Path.ToString();
 			FFolderTypeColumn_Experimental FolderType;
 
 			if (AssetViewUtils::IsPluginFolder(Path))
diff --git a/Engine/Plugins/Experimental/EditorDataStorageFeatures/Source/TedsAssetData/Private/TedsAssetData.cpp b/Engine/Plugins/Experimental/EditorDataStorageFeatures/Source/TedsAssetData/Private/TedsAssetData.cpp
--- a/Engine/Plugins/Experimental/EditorDataStorageFeatures/Source/TedsAssetData/Private/TedsAssetData.cpp
+++ b/Engine/Plugins/Experimental/EditorDataStorageFeatures/Source/TedsAssetData/Private/TedsAssetData.cpp
@@ -632,7 +632,7 @@ FTimespan FTedsAssetData::ProcessAddedPaths(const FTimespan& AllowedTime)
 		for (int32 Index = 0; Index < NewRowsCount; ++Index)
 		{
 			const FPopulatePathRowArgs* RowArgs = &PopulateRowArgs[RowCount];
-			while (!RowArgs)
+			while (!(*RowArgs))
 			{
 				++RowCount;
 				RowArgs = &PopulateRowArgs[RowCount];
@@ -747,6 +747,11 @@ void FTedsAssetData::ProcessTaskQueues(const FTimespan AllowedTime)
 	TaskTypeCount += AddedAssets.IsEmpty() ? 0 : 1;
 	TaskTypeCount += UpdatedAssets.IsEmpty() ? 0 : 1;
 
+	if (TaskTypeCount == 0)
+	{
+		return;
+	}
+
 	// By dividing allowed time by number of filled task queues we allocate time to each task type instead of all to the first type
 	FTimespan AllowedTimePerTaskType = AllowedTime / TaskTypeCount;
 
