Expose QueryStack Arg to EverythingPicker Contexts for Advanced customization

diff --git a/Engine/Plugins/Experimental/EditorDataStorageFeatures/Source/TedsEverythingPicker/Private/Context/TedsPickerContextUtil.cpp b/Engine/Plugins/Experimental/EditorDataStorageFeatures/Source/TedsEverythingPicker/Private/Context/TedsPickerContextUtil.cpp
--- a/Engine/Plugins/Experimental/EditorDataStorageFeatures/Source/TedsEverythingPicker/Private/Context/TedsPickerContextUtil.cpp
+++ b/Engine/Plugins/Experimental/EditorDataStorageFeatures/Source/TedsEverythingPicker/Private/Context/TedsPickerContextUtil.cpp
@@ -43,27 +43,41 @@ namespace UE::Editor::DataStorage::Picker
 			return;
 		}
 
-		TSharedRef<QueryStack::FQueryNode> ReferenceQueryNode =
-			MakeShared<QueryStack::FQueryNode>(*Storage, InArgs._Query);
+		TSharedPtr<QueryStack::IRowNode> FinalNode = nullptr;
+		if (!InArgs._QueryStack.IsValid())
+		{
+			TSharedRef<QueryStack::FQueryNode> ReferenceQueryNode =
+				MakeShared<QueryStack::FQueryNode>(*Storage, InArgs._Query);
 
-		TSharedRef<QueryStack::FRowQueryResultsNode> ReferenceQueryResultsNode =
-			MakeShared<QueryStack::FRowQueryResultsNode>(*Storage,
-				ReferenceQueryNode, QueryStack::FRowQueryResultsNode::ESyncFlags::RefreshOnUpdate);
+			if (ReferenceQueryNode->GetQuery() != InvalidQueryHandle)
+			{
+				FinalNode =
+					MakeShared<QueryStack::FRowQueryResultsNode>(*Storage,
+						ReferenceQueryNode, QueryStack::FRowQueryResultsNode::ESyncFlags::RefreshOnUpdate);
+			}
+			else
+			{
+				FinalNode = MakeShared<QueryStack::FRowArrayNode>(FRowHandleArray());
+			}
+		}
+		else
+		{
+			FinalNode = InArgs._QueryStack;
+		}
 
+		// Append a search node on the end
 		TSharedPtr<QueryStack::IRowNode> SearchNode;
 		TSharedPtr<SBox> SearchBoxContainer = SNew(SBox);
 		if(InArgs._SearchingEnabled)
 		{
 			SearchBoxContainer->SetContent(SNew(STedsSearchBox)
-				.InSearchableRowNode(ReferenceQueryResultsNode.ToSharedPtr())
+				.InSearchableRowNode(FinalNode)
 				.OutSearchNode(&SearchNode));
 		}
 
-		// If there is no SearchNode available, either disabled or there is an internal error, assign it to be the 
-		// ReferenceQueryResultsNode so that it won't be null
-		if(!SearchNode)
+		if(SearchNode)
 		{
-			SearchNode = ReferenceQueryResultsNode.ToSharedPtr();
+			FinalNode = SearchNode;
 		}
 		
 		// Using the Teds Outliner Purpose for matching customization for now
@@ -79,7 +93,7 @@ namespace UE::Editor::DataStorage::Picker
 		TSharedPtr<FHierarchyViewerLegacyData> HierarchyData = MakeShared<FHierarchyViewerLegacyData>(FTableRowParentColumn::StaticStruct(), MoveTemp(ParentExtractionFunction));
 
 		TSharedPtr<SHierarchyViewer> ReferenceView = SNew(SHierarchyViewer, HierarchyData)
-			.AllNodeProvider(SearchNode)
+			.AllNodeProvider(FinalNode)
 			.Columns({ FTypedElementLabelColumn::StaticStruct(), FTypedElementClassTypeInfoColumn::StaticStruct() })
 			.CellWidgetPurpose(PurposeId)
 			.OnSelectionChanged(InArgs._OnSelectionChanged)
@@ -120,21 +134,28 @@ namespace UE::Editor::DataStorage::Picker
 		// Default FinalNode to an empty row list
 		TSharedPtr<QueryStack::IRowNode> FinalNode = nullptr;
 
-		TSharedRef<QueryStack::FQueryNode> TypeQueryNode =
-			MakeShared<QueryStack::FQueryNode>(*Storage, InArgs._Query);
-
-		bool bHasValidQuery = TypeQueryNode->GetQuery() != InvalidQueryHandle;
-		if (bHasValidQuery)
+		bool bHasValidQuery = false;
+		if (!InArgs._QueryStack.IsValid())
 		{
-			TSharedRef<QueryStack::FRowQueryResultsNode> TypeResultsNode =
-				MakeShared<QueryStack::FRowQueryResultsNode>(*Storage,
-					TypeQueryNode, QueryStack::FRowQueryResultsNode::ESyncFlags::RefreshOnUpdate);
+			TSharedRef<QueryStack::FQueryNode> TypeQueryNode =
+				MakeShared<QueryStack::FQueryNode>(*Storage, InArgs._Query);
 
-			FinalNode = TypeResultsNode.ToSharedPtr();
+			bHasValidQuery = TypeQueryNode->GetQuery() != InvalidQueryHandle;
+			if (bHasValidQuery)
+			{
+				FinalNode =
+					MakeShared<QueryStack::FRowQueryResultsNode>(*Storage,
+						TypeQueryNode, QueryStack::FRowQueryResultsNode::ESyncFlags::RefreshOnUpdate);
+			}
+			else
+			{
+				FinalNode = MakeShared<QueryStack::FRowArrayNode>(FRowHandleArray());
+			}
 		}
 		else
 		{
-			FinalNode = MakeShared<QueryStack::FRowArrayNode>(FRowHandleArray());
+			FinalNode = InArgs._QueryStack;
+			bHasValidQuery = true;
 		}
 
 		if (InArgs._BaseType)
diff --git a/Engine/Plugins/Experimental/EditorDataStorageFeatures/Source/TedsEverythingPicker/Public/Context/TedsPickerContextUtil.h b/Engine/Plugins/Experimental/EditorDataStorageFeatures/Source/TedsEverythingPicker/Public/Context/TedsPickerContextUtil.h
--- a/Engine/Plugins/Experimental/EditorDataStorageFeatures/Source/TedsEverythingPicker/Public/Context/TedsPickerContextUtil.h
+++ b/Engine/Plugins/Experimental/EditorDataStorageFeatures/Source/TedsEverythingPicker/Public/Context/TedsPickerContextUtil.h
@@ -13,6 +13,7 @@
 namespace UE::Editor::DataStorage::QueryStack
 {
 	class FRowFilterNode;
+	class IRowNode;
 }
 
 namespace UE::Editor::DataStorage::Picker
@@ -26,6 +27,7 @@ namespace UE::Editor::DataStorage::Picker
 			: _SearchingEnabled(false)
 			{}
 			SLATE_ARGUMENT(FQueryDescription, Query)
+			SLATE_ARGUMENT(TSharedPtr<QueryStack::IRowNode>, QueryStack) // Advanced: QueryStack allows finer control over displayed data. Will be used instead of Query if specified
 			SLATE_ARGUMENT(FOnSelectionChanged, OnSelectionChanged)
 			SLATE_ARGUMENT(bool, SearchingEnabled)
 		SLATE_END_ARGS()
@@ -40,6 +42,7 @@ namespace UE::Editor::DataStorage::Picker
 			: _SearchingEnabled(false)
 			{}
 			SLATE_ARGUMENT(FQueryDescription, Query) // If set will perform the requested query. Ex: Select().Where().All<FClassTypeInfoTag>().Compile(). Results Intersect with BaseType
+			SLATE_ARGUMENT(TSharedPtr<QueryStack::IRowNode>, QueryStack) // Advanced: QueryStack allows finer control over displayed data. Will be used instead of Query if specified
 			SLATE_ARGUMENT(UStruct*, BaseType) // If set, will collect a list of BaseType + SubTypes. Results intersect with Query
 			SLATE_ARGUMENT(FOnSelectionChanged, OnSelectionChanged)
 			SLATE_ARGUMENT(bool, SearchingEnabled)
