Enable TedsAssetData by default

diff --git a/Engine/Plugins/Experimental/EditorDataStorageFeatures/Source/TedsAssetData/Private/TedsAssetDataModule.cpp b/Engine/Plugins/Experimental/EditorDataStorageFeatures/Source/TedsAssetData/Private/TedsAssetDataModule.cpp
--- a/Engine/Plugins/Experimental/EditorDataStorageFeatures/Source/TedsAssetData/Private/TedsAssetDataModule.cpp
+++ b/Engine/Plugins/Experimental/EditorDataStorageFeatures/Source/TedsAssetData/Private/TedsAssetDataModule.cpp
@@ -16,7 +16,7 @@ namespace UE::Editor::AssetData
 
 namespace Private
 {
-TAutoConsoleVariable<bool> CVarTEDSAssetDataStorage(TEXT("TEDS.AssetDataStorage"), false, TEXT("When true we will activate a wrapper that store the a copy of the asset data including the in memory change from the asset registry into TEDS.")
+TAutoConsoleVariable<bool> CVarTEDSAssetDataStorage(TEXT("TEDS.AssetDataStorage"), true, TEXT("When true we will activate a wrapper that store the a copy of the asset data including the in memory change from the asset registry into TEDS.")
 	, FConsoleVariableDelegate::CreateLambda([](IConsoleVariable* Variable)
 	{
 		const bool bIsEnabled = Variable->GetBool();
