Enable Teds TypeInfo by default + Fixup type display names

diff --git a/Engine/Plugins/Experimental/EditorDataStorageFeatures/Source/TedsTypeInfo/Private/TedsTypeInfoFactory.cpp b/Engine/Plugins/Experimental/EditorDataStorageFeatures/Source/TedsTypeInfo/Private/TedsTypeInfoFactory.cpp
--- a/Engine/Plugins/Experimental/EditorDataStorageFeatures/Source/TedsTypeInfo/Private/TedsTypeInfoFactory.cpp
+++ b/Engine/Plugins/Experimental/EditorDataStorageFeatures/Source/TedsTypeInfo/Private/TedsTypeInfoFactory.cpp
@@ -50,7 +50,7 @@ void UTypeInfoFactory::RegisterQueries(UE::Editor::DataStorage::ICoreProvider& D
 
 	DataStorage.RegisterQuery(
 		Select(
-			TEXT("Setup Class Hierachy Info"),
+			TEXT("Setup Class Hierarchy Info"),
 			FProcessor(EQueryTickPhase::DuringPhysics, DataStorage.GetQueryTickGroupName(EQueryTickGroups::Default))
 				.MakeActivatable(ClassSetupActionName),
 			[this](IQueryContext& Context, RowHandle Row, const FTypedElementClassTypeInfoColumn& ClassTypeColumn)
@@ -194,7 +194,7 @@ void UTypeInfoFactory::AddCommonColumns(UE::Editor::DataStorage::ICoreProvider&
 	using namespace UE::Editor::DataStorage;
 	checkf(InTypeInfo, TEXT("Attempted to add type info columns with null TypeInfo"));
 
-	DataStorage.AddColumn(TypeRowHandle, FTypedElementLabelColumn{ .Label = InTypeInfo->GetName() });
+	DataStorage.AddColumn(TypeRowHandle, FTypedElementLabelColumn{ .Label = InTypeInfo->GetDisplayNameText().ToString() });
 	DataStorage.AddColumns<FTypeInfoTag, FTypeInfoRequiresHierarchyUpdateTag>(TypeRowHandle);
 }
 
diff --git a/Engine/Plugins/Experimental/EditorDataStorageFeatures/Source/TedsTypeInfo/Private/TedsTypeInfoModule.cpp b/Engine/Plugins/Experimental/EditorDataStorageFeatures/Source/TedsTypeInfo/Private/TedsTypeInfoModule.cpp
--- a/Engine/Plugins/Experimental/EditorDataStorageFeatures/Source/TedsTypeInfo/Private/TedsTypeInfoModule.cpp
+++ b/Engine/Plugins/Experimental/EditorDataStorageFeatures/Source/TedsTypeInfo/Private/TedsTypeInfoModule.cpp
@@ -14,7 +14,7 @@ namespace UE::Editor::DataStorage::TypeInfo
 	namespace Private
 	{
 		TAutoConsoleVariable<bool> CVarTEDSTypeInfoEnabled(TEXT("TEDS.Feature.TypeInfoIntegration.Enable"),
-			false,
+			true,
 			TEXT("When true we will store type info into TEDS and propagate type updates when they occur"),
 			FConsoleVariableDelegate::CreateLambda([](IConsoleVariable* Variable)
 				{
