Disable Teds TypeInfo due to automated test failure

diff --git a/Engine/Plugins/Experimental/EditorDataStorageFeatures/Source/TedsTypeInfo/Private/TedsTypeInfoModule.cpp b/Engine/Plugins/Experimental/EditorDataStorageFeatures/Source/TedsTypeInfo/Private/TedsTypeInfoModule.cpp
--- a/Engine/Plugins/Experimental/EditorDataStorageFeatures/Source/TedsTypeInfo/Private/TedsTypeInfoModule.cpp
+++ b/Engine/Plugins/Experimental/EditorDataStorageFeatures/Source/TedsTypeInfo/Private/TedsTypeInfoModule.cpp
@@ -14,7 +14,7 @@ namespace UE::Editor::DataStorage::TypeInfo
 	namespace Private
 	{
 		TAutoConsoleVariable<bool> CVarTEDSTypeInfoEnabled(TEXT("TEDS.Feature.TypeInfoIntegration.Enable"),
-			true,
+			false,
 			TEXT("When true we will store type info into TEDS and propagate type updates when they occur"),
 			FConsoleVariableDelegate::CreateLambda([](IConsoleVariable* Variable)
 				{
