Add a new TEDSFilterBar Widget and add TedsFilters to the TableViewer module (TedsOutlinerFilter to be removed after)

diff --git a/Engine/Plugins/Experimental/EditorDataStorageFeatures/Source/TedsTableViewer/Private/TedsFilter.cpp b/Engine/Plugins/Experimental/EditorDataStorageFeatures/Source/TedsTableViewer/Private/TedsFilter.cpp
new file mode 100644
--- /dev/null
+++ b/Engine/Plugins/Experimental/EditorDataStorageFeatures/Source/TedsTableViewer/Private/TedsFilter.cpp
@@ -0,0 +1,148 @@
+// Copyright Epic Games, Inc. All Rights Reserved.
+
+#include "TedsFilter.h"
+
+#include "Widgets/STedsFilterBar.h"
+
+namespace UE::Editor::DataStorage
+{
+	class STedsFilterBar;
+
+	FTedsFilterData::FTedsFilterData(const FName& InFilterName, const FText& InFilterDisplayName, const FText& InFilterToolTip, const FName& InFilterIconName,
+    	const TSharedPtr<FFilterCategory>& InCategory, const QueryHandle& InFilterQuery)
+    	: FilterName(InFilterName)
+    	, FilterDisplayName(InFilterDisplayName)
+    	, FilterToolTip(InFilterToolTip)
+    	, FilterIconName(InFilterIconName)
+    	, FilterCategory(InCategory)
+    {
+    	FilterQuery.Set<QueryHandle>(CheckValidFilterQueryHandle(InFilterQuery) ? InFilterQuery : InvalidQueryHandle);
+    }
+    
+    FTedsFilterData::FTedsFilterData(const FName& InFilterName, const FText& InFilterDisplayName, const QueryHandle& InFilterQuery)
+    	: FTedsFilterData(
+    		InFilterName, 
+    		InFilterDisplayName, 
+    		FText::FromString(InFilterName.ToString()), 
+    		FName(), 
+    		nullptr, 
+    		InFilterQuery)
+    {}
+    
+    FTedsFilterData::FTedsFilterData(const FName& InFilterName, const FText& InFilterDisplayName, const FText& InFilterToolTip, const FName& InFilterIconName,
+    	const TSharedPtr<FFilterCategory>& InCategory, const Queries::TQueryFunction<bool>& InFilterQuery)
+    	: FilterName(InFilterName)
+    	, FilterDisplayName(InFilterDisplayName)
+    	, FilterToolTip(InFilterToolTip)
+    	, FilterIconName(InFilterIconName)
+    	, FilterCategory(InCategory)
+    {
+    	FilterQuery.Set<Queries::TQueryFunction<bool>>(InFilterQuery);
+    }
+    
+    FTedsFilterData::FTedsFilterData(const FName& InFilterName, const FText& InFilterDisplayName, const Queries::TQueryFunction<bool>& InFilterQuery)
+    	: FTedsFilterData(
+    		InFilterName, 
+    		InFilterDisplayName, 
+    		FText::FromString(InFilterName.ToString()), 
+    		FName(), 
+    		nullptr, 
+    		InFilterQuery)
+    {}
+    
+    FTedsFilterData::FTedsFilterData(const FTedsFilterData&) = default;
+    FTedsFilterData::~FTedsFilterData() = default;
+		
+	FString FTedsFilter::GetName() const
+	{
+		return FilterName.ToString();
+	}
+
+	FText FTedsFilter::GetDisplayName() const
+	{
+		return FilterDisplayName;
+	}
+
+	FText FTedsFilter::GetToolTipText() const
+	{
+		return FilterToolTip;
+	}
+
+	FLinearColor FTedsFilter::GetColor() const
+	{
+		return FLinearColor();	
+	}
+
+	FName FTedsFilter::GetIconName() const
+	{
+		return FilterIconName;
+	}
+
+	bool FTedsFilter::IsInverseFilter() const
+	{
+		return false;
+	}
+
+	void FTedsFilter::ActiveStateChanged(bool bActive)
+	{
+		using namespace UE::Editor::DataStorage::Queries;
+		
+		const TSharedPtr<STedsFilterBar> TedsFilterBarPin = TedsFilterBar.Pin();
+		if (ensureMsgf(TedsFilterBarPin, TEXT("No TedsFilterBar Context was set for the %s filter."), *FilterDisplayName.ToString()))
+		{
+			if(bActive)
+			{
+				if (bIsClassFilter)
+				{
+					// Class filters are separated since we want to OR them with each other
+					TedsFilterBarPin->AddClassQueryFunction(FilterName, FilterQuery.Get<TQueryFunction<bool>>());
+				}
+				else if (FilterQuery.IsType<QueryHandle>())
+				{
+					TedsFilterBarPin->AddExternalQuery(FilterName, FilterQuery.Get<QueryHandle>());
+				}
+				else /* if (FilterQuery.IsType<TQueryFunction<bool>>()) */
+				{
+					TedsFilterBarPin->AddExternalQueryFunction(FilterName, FilterQuery.Get<TQueryFunction<bool>>());
+				}
+			}
+			else
+			{
+				if (bIsClassFilter)
+				{
+					TedsFilterBarPin->RemoveClassQueryFunction(FilterName);
+				}
+				else if (FilterQuery.IsType<QueryHandle>())
+				{
+					TedsFilterBarPin->RemoveExternalQuery(FilterName);
+				}
+				else /* if (FilterQuery.IsType<TQueryFunction<bool>>()) */
+				{
+					TedsFilterBarPin->RemoveExternalQueryFunction(FilterName);
+				}
+			}
+		}
+	}
+
+	void FTedsFilter::ModifyContextMenu(FMenuBuilder& MenuBuilder)
+	{
+		
+	}
+
+	void FTedsFilter::SaveSettings(const FString& IniFilename, const FString& IniSection, const FString& SettingsString) const
+	{
+		
+	}
+
+	void FTedsFilter::LoadSettings(const FString& IniFilename, const FString& IniSection, const FString& SettingsString)
+	{
+		
+	}
+
+	bool FTedsFilter::PassesFilter(FTedsRowHandle& InItem) const
+	{
+		// The filter is applied through a TEDS query and this is just a dummy to activate it, so we can just run a simple valid check on if the
+		// TEDSFilterBar is still valid.
+		return TedsFilterBar.IsValid();
+	}
+} // namespace UE::Editor::DataStorage
diff --git a/Engine/Plugins/Experimental/EditorDataStorageFeatures/Source/TedsTableViewer/Private/Widgets/STedsFilterBar.cpp b/Engine/Plugins/Experimental/EditorDataStorageFeatures/Source/TedsTableViewer/Private/Widgets/STedsFilterBar.cpp
new file mode 100644
--- /dev/null
+++ b/Engine/Plugins/Experimental/EditorDataStorageFeatures/Source/TedsTableViewer/Private/Widgets/STedsFilterBar.cpp
@@ -0,0 +1,253 @@
+// Copyright Epic Games, Inc. All Rights Reserved.
+
+#include "Widgets/STedsFilterBar.h"
+
+#include "TedsQueryHandleNode.h"
+#include "TedsQueryMergeNode.h"
+#include "TedsRowFilterNode.h"
+#include "TedsRowMergeNode.h"
+#include "TedsRowQueryResultsNode.h"
+
+#define LOCTEXT_NAMESPACE "STedsFilterBar"
+
+namespace UE::Editor::DataStorage
+{
+	void STedsFilterBar::Construct(const FArguments& InArgs)
+	{
+		using namespace DataStorage::QueryStack;
+
+		OnPostFiltersChanged = InArgs._OnPostFiltersChanged;
+
+		typename SBasicFilterBar<FTedsRowHandle&>::FArguments Args;
+		Args._OnFilterChanged = InArgs._OnFilterChanged;
+		Args._FilterBarLayout = InArgs._FilterBarLayout;
+		Args._FilterPillStyle = InArgs._FilterPillStyle;
+		Args._UseSectionsForCategories = InArgs._UseSectionsForCategories;
+
+		TSharedRef<STedsFilterBar> STedsFilterBarRef = SharedThis(this);
+		Storage = GetMutableDataStorageFeature<DataStorage::ICoreProvider>(StorageFeatureName);
+
+		TSharedRef<FFilterCategory> OtherFiltersCategory = MakeShared<FFilterCategory>(LOCTEXT("OtherFilters", "Other Filters"), LOCTEXT("OtherFiltersToolTip", "Other Filters"));
+		for (FTedsFilterData FilterData : InArgs._Filters)
+		{
+			if (!FilterData.FilterCategory)
+			{
+				FilterData.FilterCategory = OtherFiltersCategory;
+			}
+			Args._CustomFilters.Add(MakeShared<DataStorage::FTedsFilter>(FilterData, STedsFilterBarRef));
+		}
+		for (UClass* ClassFilter : InArgs._ClassFilters)
+		{
+			Args._CustomFilters.Add(MakeShared<DataStorage::FTedsFilter>(ClassFilter, OtherFiltersCategory, STedsFilterBarRef));
+		}
+		
+		if(ensureMsgf(InArgs._OutFilteredNode && InArgs._InFilterableRowNode, TEXT("In the TEDSFilterBar, a valid Row Node must be passed in to search on and a valid Out Node pointer must be given")))
+		{
+			InFilterableRowNode = InArgs._InFilterableRowNode;
+			OutFilteredNode = InArgs._OutFilteredNode;
+			check(OutFilteredNode);
+			// Immediately compile the Query so the OutFilteredNode can be set
+			RecompileQueries();
+		}
+		
+		SBasicFilterBar<FTedsRowHandle&>::Construct(Args);
+	}
+
+	void STedsFilterBar::Tick(const FGeometry& AllottedGeometry, const double InCurrentTime, const float InDeltaTime)
+	{
+		SBasicFilterBar<FTedsRowHandle&>::Tick(AllottedGeometry, InCurrentTime, InDeltaTime);
+		if (bMarkDirty)
+		{
+			RecompileQueries();
+		}
+	}
+	
+	void STedsFilterBar::PopulateAddFilterMenu(UToolMenu* InMenu)
+	{
+		PopulateCommonFilterSections(InMenu);
+		PopulateCustomFilters(InMenu);
+	}
+
+	void STedsFilterBar::PopulateCommonFilterSections(UToolMenu* InMenu)
+	{
+		FToolMenuSection& Section = InMenu->AddSection("FilterBarResetFilters");
+		Section.AddMenuEntry(
+			"ResetFilters",
+			LOCTEXT("FilterListResetFilters", "Reset Filters"),
+			LOCTEXT("FilterListResetToolTip", "Resets current pinned filters"),
+			FSlateIcon(FAppStyle::Get().GetStyleSetName(), "PropertyWindow.DiffersFromDefault"),
+			FUIAction(
+				FExecuteAction::CreateSP(this, &STedsFilterBar::OnResetFilters),
+				FCanExecuteAction::CreateLambda([this]() { return HasAnyFilters(); })
+			)
+		);
+
+		Section.AddSeparator("TedsCommonFiltersSeparator");
+	}
+
+	void STedsFilterBar::PopulateCustomFilters(UToolMenu* InMenu)
+	{
+		if (bUseSectionsForCategories)
+		{
+			// Add all the filters as sections
+			for (const TSharedPtr<FFilterCategory>& Category : AllFilterCategories)
+			{
+				FToolMenuSection& Section = InMenu->AddSection(*Category->Title.ToString(), Category->Title);
+				CreateOtherFiltersMenuCategory(Section, Category);
+			}
+		}
+		else
+		{
+			FToolMenuSection& Section = InMenu->AddSection("TEDSFilterBarMainSection");
+
+			// Add all the filters as submenus
+			for (const TSharedPtr<FFilterCategory>& Category : AllFilterCategories)
+			{
+				Section.AddSubMenu(
+					NAME_None,
+					Category->Title,
+					Category->Tooltip,
+					FNewToolMenuDelegate::CreateSP(this, &STedsFilterBar::CreateOtherFiltersMenuCategory, Category),
+					FUIAction(
+						FExecuteAction::CreateSP(this, &STedsFilterBar::FrontendFilterCategoryClicked, Category),
+						FCanExecuteAction(),
+						FGetActionCheckState::CreateSP(this, &STedsFilterBar::IsFrontendFilterCategoryChecked, Category)
+					),
+					EUserInterfaceActionType::ToggleButton
+				);
+			}
+		}
+	}
+
+	TSharedRef<SWidget> STedsFilterBar::MakeAddFilterMenu()
+	{
+		static const FName FilterMenuName = "TedsFilterMenu.FilterMenu";
+		if (!UToolMenus::Get()->IsMenuRegistered(FilterMenuName))
+		{
+			UToolMenu* Menu = UToolMenus::Get()->RegisterMenu(FilterMenuName);
+			Menu->bShouldCloseWindowAfterMenuSelection = true;
+			Menu->bCloseSelfOnly = true;
+            
+			Menu->AddDynamicSection(NAME_None, FNewToolMenuDelegate::CreateLambda([](UToolMenu* InMenu)
+			{
+				if (UFilterBarContext* Context = InMenu->FindContext<UFilterBarContext>())
+				{
+					Context->PopulateFilterMenu.ExecuteIfBound(InMenu);
+				}
+			}));
+		}
+
+		UFilterBarContext* FilterBarContext = NewObject<UFilterBarContext>();
+		FilterBarContext->PopulateFilterMenu = FOnPopulateAddFilterMenu::CreateSP(this, &STedsFilterBar::PopulateAddFilterMenu);
+		FToolMenuContext ToolMenuContext(FilterBarContext);
+		return UToolMenus::Get()->GenerateWidget(FilterMenuName, ToolMenuContext);
+	}
+
+	void STedsFilterBar::RecompileQueries()
+	{
+		using namespace DataStorage::Queries;
+		using namespace DataStorage::QueryStack;
+
+		if (!ExternalQueries.IsEmpty())
+		{
+			TArray<TSharedPtr<IQueryNode>> QueryNodes;
+			for (const TPair<FName, TSharedPtr<IQueryNode>>& ExternalQuery : ExternalQueries)
+			{
+				QueryNodes.Add(ExternalQuery.Value);
+			}
+			const TSharedPtr<IQueryNode> CombinedQueryNode = MakeShared<FQueryMergeNode>(*Storage, QueryNodes);
+			CombinedRowCollectorNode = MakeShared<FRowQueryResultsNode>(*Storage, CombinedQueryNode, FRowQueryResultsNode::ESyncFlags::RefreshOnUpdate);
+			const TArray<TSharedPtr<IRowNode>> FilterNodesArray({CombinedRowCollectorNode, InFilterableRowNode});
+			CombinedQueryFilterRowNode = MakeShared<FRowMergeNode>(FilterNodesArray, FRowMergeNode::EMergeApproach::Repeating);
+		}
+		else
+		{
+			CombinedQueryFilterRowNode = InFilterableRowNode;
+		}
+
+		FilterNodes.Empty();
+		CompositeFilterNode = CombinedQueryFilterRowNode;
+		// Combines all active external filters using AND
+		for (const TPair<FName, TQueryFunction<bool>>& ExternalQueryFunction : ExternalQueryFunctions)
+		{
+			const TSharedPtr<FRowFilterNode> FilterNode = MakeShared<FRowFilterNode>(Storage, CompositeFilterNode, ExternalQueryFunction.Value);
+			FilterNodes.Add(FilterNode);
+			CompositeFilterNode = FilterNode;
+		}
+
+		ClassFilterNodes.Empty();
+		// Combines all active class filters using OR
+		if (!ClassFilters.IsEmpty())
+		{
+			for (const TPair<FName, TQueryFunction<bool>>& ClassFilter : ClassFilters)
+			{
+				ClassFilterNodes.Add(MakeShared<FRowFilterNode>(Storage, InFilterableRowNode, ClassFilter.Value));
+			}
+		
+			const TArray<TSharedPtr<IRowNode>> ClassBaseFilterNodes(ClassFilterNodes);
+			CombinedClassFilterRowNode = MakeShared<FRowMergeNode>(ClassBaseFilterNodes, FRowMergeNode::EMergeApproach::Unique);
+			
+			*OutFilteredNode = MakeShared<FRowMergeNode>(MakeConstArrayView({CompositeFilterNode, CombinedClassFilterRowNode}), FRowMergeNode::EMergeApproach::Repeating);
+		}
+		else
+		{
+			*OutFilteredNode = MoveTemp(CompositeFilterNode);
+		}
+		
+		OnPostFiltersChanged.ExecuteIfBound();
+		bMarkDirty = false;
+	}
+
+	void STedsFilterBar::AddExternalQuery(const FName& QueryName, const QueryHandle& InQueryHandle)
+	{
+		if (ensureMsgf(InQueryHandle != InvalidQueryHandle, TEXT("An Invalid Query Handle cannot be used for a TEDS Filter")))
+		{
+			ExternalQueries.Emplace(QueryName, MakeShared<DataStorage::QueryStack::FQueryHandleNode>(InQueryHandle));
+		}
+		bMarkDirty = true;
+	}
+
+	void STedsFilterBar::RemoveExternalQuery(const FName& QueryName)
+	{
+		ExternalQueries.Remove(QueryName);
+		bMarkDirty = true;
+	}
+
+	void STedsFilterBar::AppendExternalQueries(FQueryDescription& OutQuery)
+	{
+		for(const TPair<FName, TSharedPtr<QueryStack::IQueryNode>>& ExternalQuery : ExternalQueries)
+		{
+			DataStorage::Queries::MergeQueries(OutQuery, Storage->GetQueryDescription(ExternalQuery.Value->GetQuery()));
+		}
+		bMarkDirty = true;
+	}
+
+	void STedsFilterBar::AddExternalQueryFunction(const FName& QueryName, const DataStorage::Queries::TQueryFunction<bool>& InQueryFunction)
+	{
+		// Store the external query functions as functions instead of RowNodes since the row node to filter on has not been initialized yet
+		ExternalQueryFunctions.Emplace(QueryName, InQueryFunction);
+		bMarkDirty = true;
+	}
+
+	void STedsFilterBar::RemoveExternalQueryFunction(const FName& QueryName)
+	{
+		ExternalQueryFunctions.Remove(QueryName);
+		bMarkDirty = true;
+	}
+
+	void STedsFilterBar::AddClassQueryFunction(const FName& ClassName, const Queries::TQueryFunction<bool>& InClassQueryFunction)
+	{
+		ClassFilters.Emplace(ClassName, InClassQueryFunction);
+		bMarkDirty = true;
+	}
+
+	void STedsFilterBar::RemoveClassQueryFunction(const FName& ClassName)
+	{
+		ClassFilters.Remove(ClassName);
+		bMarkDirty = true;
+	}
+	
+} // namespace UE::Editor::DataStorage
+
+#undef LOCTEXT_NAMESPACE
+
diff --git a/Engine/Plugins/Experimental/EditorDataStorageFeatures/Source/TedsTableViewer/Public/TedsFilter.h b/Engine/Plugins/Experimental/EditorDataStorageFeatures/Source/TedsTableViewer/Public/TedsFilter.h
new file mode 100644
--- /dev/null
+++ b/Engine/Plugins/Experimental/EditorDataStorageFeatures/Source/TedsTableViewer/Public/TedsFilter.h
@@ -0,0 +1,207 @@
+// Copyright Epic Games, Inc. All Rights Reserved.
+
+#pragma once
+
+#include <type_traits>
+
+#include "Filters/FilterBase.h"
+#include "DataStorage/Features.h"
+#include "DataStorage/Handles.h"
+#include "DataStorage/Queries/Description.h"
+#include "Elements/Columns/TypedElementTypeInfoColumns.h"
+#include "Elements/Framework/TypedElementQueryBuilder.h"
+#include "Elements/Framework/TypedElementQueryContext.h"
+#include "Styling/SlateIconFinder.h"
+
+#define LOCTEXT_NAMESPACE "TEDSFilter"
+
+namespace UE::Editor::DataStorage
+{
+	class STedsFilterBar;
+
+	// Helper function to check if the given QueryHandle is valid to be used in a filter (is not an observer).
+	inline bool CheckValidFilterQueryHandle(const QueryHandle& InQueryHandle)
+	{
+		// Check if the given Query Handle is valid for a filter (Not an Observer Query Handle)
+		DataStorage::ICoreProvider* Storage = GetMutableDataStorageFeature<DataStorage::ICoreProvider>(StorageFeatureName);
+		if (ensureMsgf(Storage, TEXT("TEDS must be initialized before TEDS Filters")))
+		{
+			const EQueryCallbackType QueryHandleCallbackType = Storage->GetQueryDescription(InQueryHandle).Callback.Type;
+			return ensureMsgf(QueryHandleCallbackType == EQueryCallbackType::None || QueryHandleCallbackType == EQueryCallbackType::Processor,
+				TEXT("TEDS Filters cannot accept Observer Query Handles."));
+		}
+		return false;
+	}
+
+	struct TEDSTABLEVIEWER_API FTedsFilterData
+	{
+		FTedsFilterData(const FName& InFilterName, const FText& InFilterDisplayName, const FText& InFilterToolTip, const FName& InFilterIconName,
+			const TSharedPtr<FFilterCategory>& InCategory, const QueryHandle& InFilterQuery);
+
+		FTedsFilterData(const FName& InFilterName, const FText& InFilterDisplayName, const QueryHandle& InFilterQuery);
+		
+		FTedsFilterData(const FName& InFilterName, const FText& InFilterDisplayName, const FText& InFilterToolTip, const FName& InFilterIconName,
+			const TSharedPtr<FFilterCategory>& InCategory, const Queries::TQueryFunction<bool>& InFilterQuery);
+
+		FTedsFilterData(const FName& InFilterName, const FText& InFilterDisplayName, const Queries::TQueryFunction<bool>& InFilterQuery);
+
+		/*
+		 * Require that the template type cannot be one of the pre-defined FilterQuery types so it will pass to the defined
+		 * constructor instead of this one.
+		 */
+		template< typename FilterFunction UE_REQUIRES(
+			!std::is_same_v<std::remove_cvref_t<FilterFunction>, QueryHandle> &&
+			!std::is_same_v<std::remove_cvref_t<FilterFunction>, Queries::TQueryFunction<bool>>) >
+		FTedsFilterData(const FName& InFilterName, const FText& InFilterDisplayName, const FText& InFilterToolTip, const FName& InFilterIconName,
+			const TSharedPtr<FFilterCategory>& InCategory, FilterFunction&& InFilterQuery)
+			: FTedsFilterData(
+				InFilterName, 
+				InFilterDisplayName, 
+				InFilterToolTip, 
+				InFilterIconName, 
+				InCategory,
+				Queries::BuildQueryFunction<bool>(Forward<FilterFunction>(InFilterQuery)))
+		{}
+
+		/*
+		 * Require that the template type cannot be one of the pre-defined FilterQuery types so it will pass to the defined
+		 * constructor instead of this one.
+		 */
+		template< typename FilterFunction UE_REQUIRES(
+			!std::is_same_v<std::remove_cvref_t<FilterFunction>, QueryHandle> &&
+			!std::is_same_v<std::remove_cvref_t<FilterFunction>, Queries::TQueryFunction<bool>>) >
+	    FTedsFilterData(const FName& InFilterName, const FText& InFilterDisplayName, FilterFunction&& InFilterQuery)
+			: FTedsFilterData(
+				InFilterName, 
+				InFilterDisplayName, 
+				FText::FromString(InFilterName.ToString()), 
+				FName(), 
+				nullptr, 
+				Queries::BuildQueryFunction<bool>(Forward<FilterFunction>(InFilterQuery)))
+		{}
+
+		FTedsFilterData(const FTedsFilterData&);
+		~FTedsFilterData(); 
+		
+		FName FilterName;
+		FText FilterDisplayName;
+		FText FilterToolTip;
+		FName FilterIconName;
+		
+		// The Category can be set as a nullptr in FTedsFilterData to go to the default 'Other Filters' category when adding,
+		// otherwise the filter will not display if it's not given a category
+		TSharedPtr<FFilterCategory> FilterCategory;
+		TVariant<QueryHandle, Queries::TQueryFunction<bool>> FilterQuery;
+	};
+		
+	class FTedsFilter : public FFilterBase<FTedsRowHandle&>
+	{
+	public:
+		FTedsFilter(const FTedsFilterData& InTedsFilterData, const TSharedPtr<STedsFilterBar>& InTedsFilterBar )
+			: FFilterBase(InTedsFilterData.FilterCategory)
+			, FilterName(InTedsFilterData.FilterName)
+			, FilterDisplayName(InTedsFilterData.FilterDisplayName)
+			, FilterToolTip(InTedsFilterData.FilterToolTip)
+			, FilterIconName(InTedsFilterData.FilterIconName)
+			, bIsClassFilter(false)
+			, FilterQuery(InTedsFilterData.FilterQuery)
+			, TedsFilterBar(InTedsFilterBar)
+		{}
+
+		FTedsFilter(const FName& InFilterName, const FText& InFilterDisplayName, const FText& InFilterToolTip, const FName& InFilterIconName, 
+			const TSharedPtr<FFilterCategory>& InCategory, const QueryHandle& InFilterQuery, const TSharedPtr<STedsFilterBar>& InTedsFilterBar)
+			: FFilterBase(InCategory)
+			, FilterName(InFilterName)
+			, FilterDisplayName(InFilterDisplayName)
+			, FilterToolTip(InFilterToolTip)
+			, FilterIconName(InFilterIconName)
+			, bIsClassFilter(false)
+			, TedsFilterBar(InTedsFilterBar)
+		{
+			FilterQuery.Set<QueryHandle>(CheckValidFilterQueryHandle(InFilterQuery) ? InFilterQuery : InvalidQueryHandle);
+		}
+		
+		FTedsFilter(const FName& InFilterName, const FText& InFilterDisplayName, const FText& InFilterToolTip, 
+			const FName& InFilterIconName, const TSharedPtr<FFilterCategory>& InCategory, const Queries::TQueryFunction<bool>& InFilterQuery, 
+			const TSharedPtr<STedsFilterBar>& InTedsFilterBar)
+			: FFilterBase(InCategory)
+			, FilterName(InFilterName)
+			, FilterDisplayName(InFilterDisplayName)
+			, FilterToolTip(InFilterToolTip)
+			, FilterIconName(InFilterIconName)
+			, bIsClassFilter(false)
+			, TedsFilterBar(InTedsFilterBar)
+		{
+			FilterQuery.Set<Queries::TQueryFunction<bool>>(InFilterQuery);
+		}
+
+		FTedsFilter(const UClass* InClass, const TSharedPtr<FFilterCategory>& InCategory, const TSharedPtr<STedsFilterBar>& InTedsFilterBar)
+			: FFilterBase(InCategory)
+			, FilterName(InClass->GetFName())
+			, FilterDisplayName(InClass->GetDisplayNameText())
+			, FilterToolTip(FText::Format(LOCTEXT("FilterClassTooltip", "Filter by {0}"), InClass->GetDisplayNameText()))
+			, FilterIconName(FSlateIconFinder::FindIconForClass(InClass).GetStyleName())
+			, bIsClassFilter(true)
+			, TedsFilterBar(InTedsFilterBar)
+		{
+			FilterQuery.Set<Queries::TQueryFunction<bool>>(
+				Queries::BuildQueryFunction<bool>([InClass](
+					Queries::TQueryContext<Queries::RowBatchInfo> Context,
+					Queries::TResult<bool>& Result,
+					Queries::TConstBatch<FTypedElementClassTypeInfoColumn> TypeInfoColumns)
+				{
+					Context.ForEachRow([&Result, InClass](const FTypedElementClassTypeInfoColumn& TypeInfoColumn)
+					{
+						Result.Add(TypeInfoColumn.TypeInfo->IsChildOf(InClass));
+					}, TypeInfoColumns);
+				}));
+		}
+
+		/** Returns the system name for this filter */
+		virtual FString GetName() const override;
+
+		/** Returns the human-readable name for this filter */
+		virtual FText GetDisplayName() const override;
+
+		/** Returns the tooltip for this filter, shown in the filters menu */
+		virtual FText GetToolTipText() const override;
+
+		/** Returns the color this filter button will be when displayed as a button */
+		virtual FLinearColor GetColor() const override;
+
+		/** Returns the name of the icon to use in menu entries */
+		virtual FName GetIconName() const override;
+
+		/** If true, the filter will be active in the FilterBar when it is inactive in the UI (i.e the filter pill is grayed out) */
+		virtual bool IsInverseFilter() const override;
+
+		/** Notification that the filter became active or inactive */
+		virtual void ActiveStateChanged(bool bActive) override;
+
+		/** Called when the right-click context menu is being built for this filter */
+		virtual void ModifyContextMenu(FMenuBuilder& MenuBuilder) override;
+		
+		/** Can be overriden for custom FilterBar subclasses to save settings, currently not implemented in any generic Filter Bar */
+		virtual void SaveSettings(const FString& IniFilename, const FString& IniSection, const FString& SettingsString) const override;
+
+		/** Can be overridden for custom FilterBar subclasses to load settings, currently not implemented in any generic Filter Bar */
+		virtual void LoadSettings(const FString& IniFilename, const FString& IniSection, const FString& SettingsString) override;
+		
+		/** Returns whether the specified Item passes the Filter's restrictions */
+		virtual bool PassesFilter(FTedsRowHandle& InItem ) const override;
+
+	protected:
+
+		FName FilterName;
+		FText FilterDisplayName;
+		FText FilterToolTip;
+		FName FilterIconName;
+
+		const bool bIsClassFilter;
+
+		TVariant<QueryHandle, Queries::TQueryFunction<bool>> FilterQuery;
+		TWeakPtr<STedsFilterBar> TedsFilterBar;
+	};
+} // namespace UE::Editor::DataStorage
+
+#undef LOCTEXT_NAMESPACE
\ No newline at end of file
diff --git a/Engine/Plugins/Experimental/EditorDataStorageFeatures/Source/TedsTableViewer/Public/Widgets/STedsFilterBar.h b/Engine/Plugins/Experimental/EditorDataStorageFeatures/Source/TedsTableViewer/Public/Widgets/STedsFilterBar.h
new file mode 100644
--- /dev/null
+++ b/Engine/Plugins/Experimental/EditorDataStorageFeatures/Source/TedsTableViewer/Public/Widgets/STedsFilterBar.h
@@ -0,0 +1,146 @@
+// Copyright Epic Games, Inc. All Rights Reserved.
+#pragma once
+
+#include "Filters/SBasicFilterBar.h"
+#include "TedsFilter.h"
+#include "TedsQueryStackInterfaces.h"
+#include "Elements/Framework/TypedElementQueryBuilder.h"
+
+namespace UE::Editor::DataStorage
+{
+	namespace QueryStack
+	{
+		class FRowQueryResultsNode;
+		class FRowFilterNode;
+	}
+
+	/**
+	* A Basic TEDS Filter Bar widget, which can be used to filter items of type FTedsRowHandle given a list of custom filters
+	* 
+	* NOTE: The filter functions create copies, so you want to use a reference or pointer as the template type when possible
+	* Sample Usage:
+	*		SAssignNew(MyFilterBar, STedsFilterBar)
+	*		.InFilterableRowNode(QueryStack)
+	*		.OutFilteredNode(&MyOutNode)
+	*		.Filters(MyCustomFilters) // An array of filters available to this FilterBar
+	*
+	* Use MakeAddFilterButton() to make the button that summons the dropdown showing all the filters
+	*/
+	class TEDSTABLEVIEWER_API STedsFilterBar : public SBasicFilterBar<FTedsRowHandle&>
+	{
+		using FOnFilterChanged = typename SBasicFilterBar<FTedsRowHandle&>::FOnFilterChanged;
+		DECLARE_DELEGATE(FOnPostFiltersChanged);
+		
+	public:
+		SLATE_BEGIN_ARGS(STedsFilterBar)
+			: _FilterBarIdentifier(NAME_None)
+			, _FilterBarLayout(EFilterBarLayout::Horizontal)
+			, _FilterPillStyle(EFilterPillStyle::Basic)
+			, _UseSectionsForCategories(false)
+			{}
+
+			/** A unique identifier for this filter bar needed to enable saving settings in the config file */
+			SLATE_ARGUMENT(FName, FilterBarIdentifier)
+
+			/** The node whose rows we'll filter on */
+			SLATE_ARGUMENT(TSharedPtr<QueryStack::IRowNode>, InFilterableRowNode)
+
+			/** The filter results node that will contain the filtered rows */
+			SLATE_ARGUMENT(TSharedPtr<QueryStack::IRowNode>*, OutFilteredNode)
+
+			/** Array of filters to add to the filter bar *AND Filters */
+			SLATE_ARGUMENT(TArray<FTedsFilterData>, Filters)
+		
+			/** Array of class filters to add to the filter bar *OR Filters */
+			SLATE_ARGUMENT(TArray<UClass*>, ClassFilters)
+			
+			/** Delegate for when filters have changed */
+			SLATE_EVENT(FOnFilterChanged, OnFilterChanged)
+
+			/** Delegate for after the filters have changed and the query stack has been updated */
+			SLATE_EVENT(FOnPostFiltersChanged, OnPostFiltersChanged)
+			
+			/** The layout that determines how the filters are laid out */
+			SLATE_ARGUMENT(EFilterBarLayout, FilterBarLayout)
+			
+			/** Determines how each individual filter pill looks like */
+			SLATE_ARGUMENT(EFilterPillStyle, FilterPillStyle)
+		
+			/** Whether to use submenus or sections for categories in the filter menu */
+			SLATE_ARGUMENT(bool, UseSectionsForCategories)
+			
+		SLATE_END_ARGS()
+
+		/** Constructs this widget with InArgs */
+		void Construct(const FArguments& InArgs);
+
+		/** Override tick so that filters are only recompiled once if multiple were added at once */
+		virtual void Tick(const FGeometry& AllottedGeometry, const double InCurrentTime, const float InDeltaTime) override;
+		
+		// Add an external query description to the FilterBar
+		void AddExternalQuery(const FName& QueryName, const QueryHandle& InQueryHandle);
+		void RemoveExternalQuery(const FName& QueryName);
+
+		// Append all external queries into the given query description
+		void AppendExternalQueries(DataStorage::FQueryDescription& OutQuery);
+
+		// Add an external query function to the FilterBar
+		void AddExternalQueryFunction(const FName& QueryName, const Queries::TQueryFunction<bool>& InQueryFunction);
+		void RemoveExternalQueryFunction(const FName& QueryName);
+
+		// Add an external class query function to the FilterBar (uses OR) 
+		void AddClassQueryFunction(const FName& ClassName, const Queries::TQueryFunction<bool>& InClassQueryFunction);
+		void RemoveClassQueryFunction(const FName& ClassName);
+
+	protected:
+		void RecompileQueries();
+		
+		virtual void PopulateAddFilterMenu(UToolMenu* InMenu);
+		virtual void PopulateCommonFilterSections(UToolMenu* InMenu);
+		virtual void PopulateCustomFilters(UToolMenu* InMenu);
+	
+	private:
+		virtual TSharedRef<SWidget> MakeAddFilterMenu() override;
+
+		ICoreProvider* Storage = nullptr;
+
+		// The row node passed in as the data to filter on
+		TSharedPtr<QueryStack::IRowNode> InFilterableRowNode;
+
+		// Pointer to the Node that is filtered and updated by the FilterBar
+		TSharedPtr<QueryStack::IRowNode>* OutFilteredNode = nullptr;
+		
+		// The Node responsible for compositing the final node to be given to the OutFilteredNode
+		TSharedPtr<QueryStack::IRowNode> CompositeFilterNode;
+
+		// External query descriptions that are currently active (e.g Filters)
+		TMap<FName, TSharedPtr<QueryStack::IQueryNode>> ExternalQueries;
+	
+		// External query functions that are currently active (e.g Filters)
+		TMap<FName, Queries::TQueryFunction<bool>> ExternalQueryFunctions;
+
+		// External query functions used to filter by class, these are the only filters that use OR instead of AND
+		TMap<FName, Queries::TQueryFunction<bool>> ClassFilters;
+
+		// Array to store all filter nodes created from the active query functions
+		TArray<TSharedPtr<QueryStack::FRowFilterNode>> FilterNodes;
+
+		// Array to store all filter nodes created from the active class filters
+		TArray<TSharedPtr<QueryStack::FRowFilterNode>> ClassFilterNodes;
+		
+		// The query stack node responsible for collecting all rows that match the composite query on FullRefresh()
+		TSharedPtr<QueryStack::FRowQueryResultsNode> CombinedRowCollectorNode;
+
+		// The query stack node responsible for combining query filters - utilizes the repeated merge approach (AND)
+		TSharedPtr<QueryStack::IRowNode> CombinedQueryFilterRowNode;
+
+		// The query stack node responsible for combining class filters - utilizes the unique merge approach (OR)
+		TSharedPtr<QueryStack::IRowNode> CombinedClassFilterRowNode;
+		
+		// Delegate for additional functionality to perform after the filters have been updated, performed at the end of the RecompileQueries function
+		FOnPostFiltersChanged OnPostFiltersChanged;
+
+		// Responsible for seeing if a filter was added or removed so multiple recompiles aren't performed every frame
+		bool bMarkDirty = false;
+	};
+}
diff --git a/Engine/Plugins/Experimental/EditorDataStorageFeatures/Source/TedsTableViewer/TedsTableViewer.Build.cs b/Engine/Plugins/Experimental/EditorDataStorageFeatures/Source/TedsTableViewer/TedsTableViewer.Build.cs
--- a/Engine/Plugins/Experimental/EditorDataStorageFeatures/Source/TedsTableViewer/TedsTableViewer.Build.cs
+++ b/Engine/Plugins/Experimental/EditorDataStorageFeatures/Source/TedsTableViewer/TedsTableViewer.Build.cs
@@ -26,11 +26,13 @@ public class TedsTableViewer : ModuleRules
 				{
 					"Core",
 					"CoreUObject",
+					"EditorWidgets",
 					"Engine",
 					"InputCore",
 					"Slate",
 					"SlateCore",
 					"WorkspaceMenuStructure",
+					"ToolMenus",
 					"ToolWidgets"
 				});
 		}
