[TEDS Viewers] Create a new composite hierarchy viewer (include searching and filtering by default)

diff --git a/Engine/Plugins/Experimental/EditorDataStorageFeatures/Source/TedsTableViewer/Private/TedsTableViewerModel.cpp b/Engine/Plugins/Experimental/EditorDataStorageFeatures/Source/TedsTableViewer/Private/TedsTableViewerModel.cpp
--- a/Engine/Plugins/Experimental/EditorDataStorageFeatures/Source/TedsTableViewer/Private/TedsTableViewerModel.cpp
+++ b/Engine/Plugins/Experimental/EditorDataStorageFeatures/Source/TedsTableViewer/Private/TedsTableViewerModel.cpp
@@ -26,17 +26,11 @@ namespace UE::Editor::DataStorage
 		Storage = GetMutableDataStorageFeature<ICoreProvider>(StorageFeatureName);
 		StorageUi = GetMutableDataStorageFeature<IUiProvider>(UiFeatureName);
 		
-		checkf(InRowQueryStack, TEXT("The TEDS Table Viewer requires a valid query stack to provide it rows to process."));
-		RowSecondarySortingNode = MakeShared<QueryStack::FRowSortNode>(*Storage, InRowQueryStack);
-		RowSecondaryInversionNode = MakeShared<QueryStack::FRowOrderInversionNode>(RowSecondarySortingNode, false);
-		RowPrimarySortingNode = MakeShared<QueryStack::FRowSortNode>(*Storage, RowSecondaryInversionNode);
-		RowPrimaryInversionNode = MakeShared<QueryStack::FRowOrderInversionNode>(RowPrimarySortingNode, false);
-		RowQueryStack = RowPrimaryInversionNode;
+		SetQueryStack(InRowQueryStack);
 
 		TickerHandle = FTSTicker::GetCoreTicker().AddTicker(FTickerDelegate::CreateRaw(this, &FTedsTableViewerModel::Tick), 0);
 
 		GenerateColumns();
-		Refresh();
 	}
 
 	FTedsTableViewerModel::~FTedsTableViewerModel()
@@ -161,6 +155,21 @@ namespace UE::Editor::DataStorage
 		return OnModelChanged;
 	}
 
+	void FTedsTableViewerModel::SetQueryStack(const TSharedPtr<QueryStack::IRowNode>& InRowQueryStack)
+	{
+		checkf(InRowQueryStack, TEXT("The TEDS Table Viewer requires a valid query stack to provide it with rows to process."));
+		RowSecondarySortingNode = MakeShared<QueryStack::FRowSortNode>(*Storage, InRowQueryStack);
+		RowSecondaryInversionNode = MakeShared<QueryStack::FRowOrderInversionNode>(RowSecondarySortingNode, false);
+		RowPrimarySortingNode = MakeShared<QueryStack::FRowSortNode>(*Storage, RowSecondaryInversionNode);
+		RowPrimaryInversionNode = MakeShared<QueryStack::FRowOrderInversionNode>(RowPrimarySortingNode, false);
+		RowQueryStack = RowPrimaryInversionNode;
+
+		PrimarySortColumnName = NAME_None;
+		SecondarySortColumnName = NAME_None;
+		
+		Refresh();
+	}
+
 	void FTedsTableViewerModel::SetColumns(const TArray<TWeakObjectPtr<const UScriptStruct>>& InColumns)
 	{
 		RequestedTedsColumns = InColumns;
diff --git a/Engine/Plugins/Experimental/EditorDataStorageFeatures/Source/TedsTableViewer/Private/Widgets/Composite/STedsCompositeHierarchyViewer.cpp b/Engine/Plugins/Experimental/EditorDataStorageFeatures/Source/TedsTableViewer/Private/Widgets/Composite/STedsCompositeHierarchyViewer.cpp
new file mode 100644
--- /dev/null
+++ b/Engine/Plugins/Experimental/EditorDataStorageFeatures/Source/TedsTableViewer/Private/Widgets/Composite/STedsCompositeHierarchyViewer.cpp
@@ -0,0 +1,171 @@
+// Copyright Epic Games, Inc. All Rights Reserved.
+
+#include "Widgets/Composite/STedsCompositeHierarchyViewer.h"
+
+#include "TedsQueryStackInterfaces.h"
+#include "Widgets/STedsFilterBar.h"
+#include "TedsTableViewerModel.h"
+#include "Widgets/STedsSearchBox.h"
+
+namespace UE::Editor::DataStorage
+{
+	void STedsCompositeHierarchyViewer::Construct(const FArguments& InArgs, TSharedPtr<IHierarchyViewerDataInterface> InHierarchyInterface)
+	{
+		TSharedPtr<SVerticalBox> ContentWidget = SNew(SVerticalBox);
+
+		const bool bFilteringEnabled = InArgs._EnableFiltering;
+		const bool bSearchingEnabled = InArgs._EnableSearching;
+
+		SHierarchyViewer::FArguments HierarchyViewerArgs = InArgs._HierarchyViewerArgs;
+		TSharedPtr<QueryStack::IRowNode> QueryStackNode = HierarchyViewerArgs._AllNodeProvider;
+		
+		if(bFilteringEnabled || bSearchingEnabled)
+		{
+			// Box containing the default search + filter widgets
+			TSharedPtr<SHorizontalBox> FilterSearchBar = SNew(SHorizontalBox);
+
+			if (bSearchingEnabled)
+			{
+				SearchBox = SNew(STedsSearchBox)
+					.InSearchableRowNode(QueryStackNode)
+					.OutSearchNode(&SearchNode);
+
+				if (SearchNode.IsValid())
+				{
+					QueryStackNode = SearchNode;
+				}
+			}
+			if (bFilteringEnabled)
+			{
+				FilterBar = SNew(STedsFilterBar)
+					.InFilterableRowNode(QueryStackNode)
+					.OutFilteredNode(&FilterNode)
+					.Filters(InArgs._Filters)
+					.ClassFilters(InArgs._ClassFilters)
+					.UseSectionsForCategories(InArgs._UseSectionsForCategories)
+					.OnPostFiltersChanged_Lambda([this]()
+					{
+						if (Model && FilterNode.IsValid())
+						{
+							// Forward the updated FilterNode to the Model QueryStack since it holds a static copy to the old filter node
+							// Only called when the filters have changed, meaning the query stack has to be recomputed
+							Model->SetQueryStack(FilterNode);
+						}
+					});
+
+				if (FilterNode.IsValid())
+				{
+					QueryStackNode = FilterNode;
+				}
+			}
+
+			// Add Filter Menu Button before Search Box if it was successfully created
+			if (FilterBar)
+			{
+				FilterSearchBar->AddSlot()
+				.AutoWidth()
+				.VAlign(VAlign_Center)
+				.Padding(2.0f)
+				[
+					FilterBar->MakeAddFilterButton(FilterBar.ToSharedRef())
+				];
+			}
+			// Add Search Box if it was successfully created
+			if (SearchBox)
+			{
+				FilterSearchBar->AddSlot()
+				.FillWidth(1.0f)
+				.VAlign(VAlign_Center)
+				.Padding(2.0f)
+				[
+					SearchBox.ToSharedRef()
+				];
+			}
+
+			// Add Filter menu dropdown button + search bar above the viewer
+			// TODO: Make it so that different default styles can be set to show the Filter/Search bar in different ways
+			ContentWidget->AddSlot()
+			.AutoHeight()
+			[
+				FilterSearchBar.ToSharedRef()
+			];
+
+			// Add Filter shelf
+			if (bFilteringEnabled)
+			{
+				ContentWidget->AddSlot()
+				.AutoHeight()
+				.Padding(2.0f)
+				[
+					FilterBar.ToSharedRef()
+				];
+			}
+		}
+
+		HierarchyViewerArgs.AllNodeProvider(QueryStackNode);
+		HierarchyViewer->Construct(HierarchyViewerArgs, MoveTemp(InHierarchyInterface));
+
+		Model = HierarchyViewer->GetModel();
+		
+		ContentWidget->AddSlot()
+		[
+			HierarchyViewer->AsWidget()
+		];
+		
+		ChildSlot
+		[
+			ContentWidget.ToSharedRef()
+		];
+	}
+
+	void STedsCompositeHierarchyViewer::SetColumns(const TArray<TWeakObjectPtr<const UScriptStruct>>& InColumns)
+	{
+		HierarchyViewer->SetColumns(InColumns);
+	}
+
+	void STedsCompositeHierarchyViewer::AddCustomRowWidget(const TSharedRef<FTedsTableViewerColumn>& InColumn)
+	{
+		HierarchyViewer->AddCustomRowWidget(InColumn);
+	}
+
+	void STedsCompositeHierarchyViewer::ForEachSelectedRow(TFunctionRef<void(RowHandle)> InCallback) const
+	{
+		HierarchyViewer->ForEachSelectedRow(InCallback);
+	}
+
+	RowHandle STedsCompositeHierarchyViewer::GetWidgetRowHandle() const
+	{
+		return HierarchyViewer->GetWidgetRowHandle();
+	}
+
+	void STedsCompositeHierarchyViewer::SetSelection(RowHandle Row, bool bSelected, const ESelectInfo::Type SelectInfo) const
+	{
+		HierarchyViewer->SetSelection(Row, bSelected, SelectInfo);
+	}
+
+	void STedsCompositeHierarchyViewer::ScrollIntoView(RowHandle Row) const
+	{
+		HierarchyViewer->ScrollIntoView(Row);
+	}
+
+	void STedsCompositeHierarchyViewer::ClearSelection() const
+	{
+		HierarchyViewer->ClearSelection();
+	}
+
+	TSharedRef<SWidget> STedsCompositeHierarchyViewer::AsWidget()
+	{
+		return AsShared();
+	}
+
+	bool STedsCompositeHierarchyViewer::IsSelected(RowHandle InRow) const
+	{
+		return HierarchyViewer->IsSelected(InRow);
+	}
+
+	bool STedsCompositeHierarchyViewer::IsSelectedExclusively(RowHandle InRow) const
+	{
+		return HierarchyViewer->IsSelectedExclusively(InRow);
+	}
+
+} // namespace UE::Editor::DataStorage
diff --git a/Engine/Plugins/Experimental/EditorDataStorageFeatures/Source/TedsTableViewer/Private/Widgets/STedsFilterBar.cpp b/Engine/Plugins/Experimental/EditorDataStorageFeatures/Source/TedsTableViewer/Private/Widgets/STedsFilterBar.cpp
--- a/Engine/Plugins/Experimental/EditorDataStorageFeatures/Source/TedsTableViewer/Private/Widgets/STedsFilterBar.cpp
+++ b/Engine/Plugins/Experimental/EditorDataStorageFeatures/Source/TedsTableViewer/Private/Widgets/STedsFilterBar.cpp
@@ -157,8 +157,9 @@ namespace UE::Editor::DataStorage
 			}
 			const TSharedPtr<IQueryNode> CombinedQueryNode = MakeShared<FQueryMergeNode>(*Storage, QueryNodes);
 			CombinedRowCollectorNode = MakeShared<FRowQueryResultsNode>(*Storage, CombinedQueryNode, FRowQueryResultsNode::ESyncFlags::RefreshOnUpdate);
-			const TArray<TSharedPtr<IRowNode>> FilterNodesArray({CombinedRowCollectorNode, InFilterableRowNode});
-			CombinedQueryFilterRowNode = MakeShared<FRowMergeNode>(FilterNodesArray, FRowMergeNode::EMergeApproach::Repeating);
+			const TSharedPtr<QueryStack::IRowNode> TemporaryFilterNodesArray[2] = {CombinedRowCollectorNode, InFilterableRowNode};
+			
+			CombinedQueryFilterRowNode = MakeShared<FRowMergeNode>(TemporaryFilterNodesArray, FRowMergeNode::EMergeApproach::Repeating);
 		}
 		else
 		{
@@ -183,11 +184,11 @@ namespace UE::Editor::DataStorage
 			{
 				ClassFilterNodes.Add(MakeShared<FRowFilterNode>(Storage, InFilterableRowNode, ClassFilter.Value));
 			}
-		
-			const TArray<TSharedPtr<IRowNode>> ClassBaseFilterNodes(ClassFilterNodes);
-			CombinedClassFilterRowNode = MakeShared<FRowMergeNode>(ClassBaseFilterNodes, FRowMergeNode::EMergeApproach::Unique);
+
+			CombinedClassFilterRowNode = MakeShared<FRowMergeNode>(ClassFilterNodes, FRowMergeNode::EMergeApproach::Unique);
+			const TSharedPtr<QueryStack::IRowNode> TemporaryCombinedClassFilterNodes[2] = {CombinedClassFilterRowNode, InFilterableRowNode};
 			
-			*OutFilteredNode = MakeShared<FRowMergeNode>(MakeConstArrayView({CompositeFilterNode, CombinedClassFilterRowNode}), FRowMergeNode::EMergeApproach::Repeating);
+			*OutFilteredNode = MakeShared<FRowMergeNode>(TemporaryCombinedClassFilterNodes, FRowMergeNode::EMergeApproach::Repeating);
 		}
 		else
 		{
diff --git a/Engine/Plugins/Experimental/EditorDataStorageFeatures/Source/TedsTableViewer/Private/Widgets/STedsHierarchyViewer.cpp b/Engine/Plugins/Experimental/EditorDataStorageFeatures/Source/TedsTableViewer/Private/Widgets/STedsHierarchyViewer.cpp
--- a/Engine/Plugins/Experimental/EditorDataStorageFeatures/Source/TedsTableViewer/Private/Widgets/STedsHierarchyViewer.cpp
+++ b/Engine/Plugins/Experimental/EditorDataStorageFeatures/Source/TedsTableViewer/Private/Widgets/STedsHierarchyViewer.cpp
@@ -110,6 +110,11 @@ namespace UE::Editor::DataStorage
 		CreateInternalWidget();
 	}
 
+	TSharedPtr<FTedsTableViewerModel> SHierarchyViewer::GetModel() const
+	{
+		return Model;
+	}
+
 	void SHierarchyViewer::AddWidgetColumns()
 	{
 		if(ICoreProvider* DataStorage = Model->GetDataStorageInterface())
diff --git a/Engine/Plugins/Experimental/EditorDataStorageFeatures/Source/TedsTableViewer/Private/Widgets/STedsSearchBox.cpp b/Engine/Plugins/Experimental/EditorDataStorageFeatures/Source/TedsTableViewer/Private/Widgets/STedsSearchBox.cpp
--- a/Engine/Plugins/Experimental/EditorDataStorageFeatures/Source/TedsTableViewer/Private/Widgets/STedsSearchBox.cpp
+++ b/Engine/Plugins/Experimental/EditorDataStorageFeatures/Source/TedsTableViewer/Private/Widgets/STedsSearchBox.cpp
@@ -33,9 +33,9 @@ namespace UE::Editor::DataStorage
 			return;
 		}
 
-		if(InArgs._OutSearchNode)
+		if(InArgs._OutSearchNode && InArgs._InSearchableRowNode)
 		{
-			SearchNode = MakeShared<QueryStack::FRowFilterNode>(Storage, InArgs._InSearchableRowNode,
+			OutSearchNode = MakeShared<QueryStack::FRowFilterNode>(Storage, InArgs._InSearchableRowNode,
 				[this](TQueryContext<SingleRowInfo> Context, const FTypedElementLabelColumn& LabelColumn)
 				{
 					if(SearchText.IsEmpty())
@@ -57,7 +57,7 @@ namespace UE::Editor::DataStorage
 					return true;
 				});
 				
-			*InArgs._OutSearchNode = SearchNode;
+			*InArgs._OutSearchNode = OutSearchNode;
 		}
 		
 		ChildSlot
@@ -70,7 +70,11 @@ namespace UE::Editor::DataStorage
 	void STedsSearchBox::OnTextChanged(const FText& InSearchText)
 	{
 		SearchText = InSearchText;
-		SearchNode->ForceRefresh();
+		if(OutSearchNode)
+		{
+			OutSearchNode->ForceRefresh();
+		}
+		
 	}
 	
 } // namespace UE::Editor::DataStorage
diff --git a/Engine/Plugins/Experimental/EditorDataStorageFeatures/Source/TedsTableViewer/Public/TedsTableViewerModel.h b/Engine/Plugins/Experimental/EditorDataStorageFeatures/Source/TedsTableViewer/Public/TedsTableViewerModel.h
--- a/Engine/Plugins/Experimental/EditorDataStorageFeatures/Source/TedsTableViewer/Public/TedsTableViewerModel.h
+++ b/Engine/Plugins/Experimental/EditorDataStorageFeatures/Source/TedsTableViewer/Public/TedsTableViewerModel.h
@@ -66,6 +66,9 @@ namespace UE::Editor::DataStorage
 
 		// Delegate when the item list changes
 		TEDSTABLEVIEWER_API FOnModelChanged& GetOnModelChanged();
+
+		// Clear the current QueryStack being displayed, set it to the given node, and recreate the sorting nodes 
+		TEDSTABLEVIEWER_API void SetQueryStack(const TSharedPtr<QueryStack::IRowNode>& InRowQueryStack);
 		
 		// Clear the current list of columns being displayed and set it to the given list
 		TEDSTABLEVIEWER_API void SetColumns(const TArray<TWeakObjectPtr<const UScriptStruct>>& InColumns);
diff --git a/Engine/Plugins/Experimental/EditorDataStorageFeatures/Source/TedsTableViewer/Public/Widgets/Composite/STedsCompositeHierarchyViewer.h b/Engine/Plugins/Experimental/EditorDataStorageFeatures/Source/TedsTableViewer/Public/Widgets/Composite/STedsCompositeHierarchyViewer.h
new file mode 100644
--- /dev/null
+++ b/Engine/Plugins/Experimental/EditorDataStorageFeatures/Source/TedsTableViewer/Public/Widgets/Composite/STedsCompositeHierarchyViewer.h
@@ -0,0 +1,127 @@
+// Copyright Epic Games, Inc. All Rights Reserved.
+#pragma once
+
+#include "DataStorage/Handles.h"
+#include "TedsFilter.h"
+#include "Widgets/STedsHierarchyViewer.h"
+
+namespace UE::Editor::DataStorage
+{
+	class FTedsTableViewerModel;
+	class IRowNode;
+	class STedsFilterBar;
+	class STedsSearchBox;
+}
+
+namespace UE::Editor::DataStorage
+{
+	class FTedsTableViewerColumn;
+
+	/*
+	 * A hierarchy viewer widget can be used to show a visual representation of data in TEDS. This composite widget adds features in a 'default' layout
+	 * such as searching and filtering. The rows to display can be specified using a RowQueryStack, and the columns to display are directly input 
+	 * into the widget
+	 * Example usage:
+	 *
+	 *	SNew(STedsCompositeHierarchyViewer, HierarchyData) // Filtering and Searching enabled by default
+     *		.HierarchyViewerArgs(SHierarchyViewer::FArguments()
+     *			.AllNodeProvider(FilterNode)
+     *			.Columns({ FTypedElementLabelColumn::StaticStruct(), FTypedElementClassTypeInfoColumn::StaticStruct() })
+     *			.CellWidgetPurpose(PurposeId)
+	 *		.Filters(MyCustomFilterArray);
+	 */
+	class STedsCompositeHierarchyViewer : public SHierarchyViewer
+	{
+	public:
+
+		SLATE_BEGIN_ARGS(STedsCompositeHierarchyViewer)
+			: _HierarchyViewerArgs()
+			, _EnableSearching(true)
+			, _EnableFiltering(true)
+			, _UseSectionsForCategories(false)
+		{
+		}
+
+		// Arguments used to create the Hierarchy Viewer
+		SLATE_ARGUMENT(SHierarchyViewer::FArguments, HierarchyViewerArgs)
+		
+		// If we want to enable the search bar on the table
+        SLATE_ARGUMENT(bool, EnableSearching)
+
+		// If we want to enable the filter bar on the table
+		SLATE_ARGUMENT(bool, EnableFiltering)
+
+		// Array of filters to add to the filter bar *AND Filters
+		SLATE_ARGUMENT(TArray<FTedsFilterData>, Filters)
+
+		// Array of class filters to add to the filter bar *OR Filters
+		SLATE_ARGUMENT(TArray<UClass*>, ClassFilters)
+		
+		// Whether to use submenus or sections for categories in the filter menu
+		SLATE_ARGUMENT(bool, UseSectionsForCategories)
+
+		SLATE_END_ARGS()
+
+	public:
+		
+		TEDSTABLEVIEWER_API void Construct(const FArguments& InArgs, TSharedPtr<IHierarchyViewerDataInterface> InHierarchyInterface);
+
+		// Clear the current list of columns being displayed and set it to the given list
+		TEDSTABLEVIEWER_API virtual void SetColumns(const TArray<TWeakObjectPtr<const UScriptStruct>>& InColumns) override;
+
+		// Add a custom per-row widget to the table viewer (that doesn't necessarily map to a TEDS column)
+		// This means a new column for the table viewer
+		TEDSTABLEVIEWER_API virtual void AddCustomRowWidget(const TSharedRef<FTedsTableViewerColumn>& InColumn) override;
+
+		// Execute the given callback for each row that is selected in the table viewer
+		TEDSTABLEVIEWER_API virtual void ForEachSelectedRow(TFunctionRef<void(RowHandle)> InCallback) const override;
+
+		// Get the row handle for the widget row the table viewer's contents are stored in
+		TEDSTABLEVIEWER_API virtual RowHandle GetWidgetRowHandle() const override;
+		
+		// Select the given row in the table viewer
+		TEDSTABLEVIEWER_API virtual void SetSelection(RowHandle Row, bool bSelected, const ESelectInfo::Type SelectInfo) const override;
+
+		// Scroll the given row into view in the table viewer
+		TEDSTABLEVIEWER_API virtual void ScrollIntoView(RowHandle Row) const override;
+
+		TEDSTABLEVIEWER_API virtual void ClearSelection() const override;
+		
+		TEDSTABLEVIEWER_API virtual TSharedRef<SWidget> AsWidget() override;
+
+		TEDSTABLEVIEWER_API virtual bool IsSelected(RowHandle InRow) const override;
+
+		TEDSTABLEVIEWER_API virtual bool IsSelectedExclusively(RowHandle InRow) const override;
+
+	private:
+		// The actual height of each item
+		TAttribute<float> ItemHeight;
+
+		// The actual padding between each item
+		TAttribute<FMargin> ItemPadding;
+
+		// Our model class
+		TSharedPtr<FTedsTableViewerModel> Model;
+
+		// Row Node passed to the SearchBox that receives the searched result
+		TSharedPtr<QueryStack::IRowNode> SearchNode;
+
+		// Row Node passed to the FilterBar that receives the filtered result
+		TSharedPtr<QueryStack::IRowNode> FilterNode;
+		
+		// The search box widget
+		TSharedPtr<STedsSearchBox> SearchBox;
+
+		// The filter bar widget (uses MakeAddFilterButton to create the menu dropdown)
+		TSharedPtr<STedsFilterBar> FilterBar;
+		
+		// The actual HierarchyViewer widget that displays the rows
+		TSharedPtr<SHierarchyViewer> HierarchyViewer = MakeShared<SHierarchyViewer>();
+
+		// Delegate fired when the selection changes
+		FOnSelectionChanged OnSelectionChanged;
+
+		// The message to show in place of the table viewer when there are no rows provided by the current query stack
+		TAttribute<FText> EmptyRowsMessage;
+	};
+}
\ No newline at end of file
diff --git a/Engine/Plugins/Experimental/EditorDataStorageFeatures/Source/TedsTableViewer/Public/Widgets/STedsFilterBar.h b/Engine/Plugins/Experimental/EditorDataStorageFeatures/Source/TedsTableViewer/Public/Widgets/STedsFilterBar.h
--- a/Engine/Plugins/Experimental/EditorDataStorageFeatures/Source/TedsTableViewer/Public/Widgets/STedsFilterBar.h
+++ b/Engine/Plugins/Experimental/EditorDataStorageFeatures/Source/TedsTableViewer/Public/Widgets/STedsFilterBar.h
@@ -123,13 +123,13 @@ namespace UE::Editor::DataStorage
 		TMap<FName, Queries::TQueryFunction<bool>> ClassFilters;
 
 		// Array to store all filter nodes created from the active query functions
-		TArray<TSharedPtr<QueryStack::FRowFilterNode>> FilterNodes;
+		TArray<TSharedPtr<QueryStack::IRowNode>> FilterNodes;
 
 		// Array to store all filter nodes created from the active class filters
-		TArray<TSharedPtr<QueryStack::FRowFilterNode>> ClassFilterNodes;
+		TArray<TSharedPtr<QueryStack::IRowNode>> ClassFilterNodes;
 		
 		// The query stack node responsible for collecting all rows that match the composite query on FullRefresh()
-		TSharedPtr<QueryStack::FRowQueryResultsNode> CombinedRowCollectorNode;
+		TSharedPtr<QueryStack::IRowNode> CombinedRowCollectorNode;
 
 		// The query stack node responsible for combining query filters - utilizes the repeated merge approach (AND)
 		TSharedPtr<QueryStack::IRowNode> CombinedQueryFilterRowNode;
diff --git a/Engine/Plugins/Experimental/EditorDataStorageFeatures/Source/TedsTableViewer/Public/Widgets/STedsHierarchyViewer.h b/Engine/Plugins/Experimental/EditorDataStorageFeatures/Source/TedsTableViewer/Public/Widgets/STedsHierarchyViewer.h
--- a/Engine/Plugins/Experimental/EditorDataStorageFeatures/Source/TedsTableViewer/Public/Widgets/STedsHierarchyViewer.h
+++ b/Engine/Plugins/Experimental/EditorDataStorageFeatures/Source/TedsTableViewer/Public/Widgets/STedsHierarchyViewer.h
@@ -119,6 +119,7 @@ namespace UE::Editor::DataStorage
 		// Get if a specific row is selected exclusively
 		TEDSTABLEVIEWER_API virtual bool IsSelectedExclusively(RowHandle InRow) const override;
 
+		TEDSTABLEVIEWER_API TSharedPtr<FTedsTableViewerModel> GetModel() const;
 
 	protected:
 		
diff --git a/Engine/Plugins/Experimental/EditorDataStorageFeatures/Source/TedsTableViewer/Public/Widgets/STedsSearchBox.h b/Engine/Plugins/Experimental/EditorDataStorageFeatures/Source/TedsTableViewer/Public/Widgets/STedsSearchBox.h
--- a/Engine/Plugins/Experimental/EditorDataStorageFeatures/Source/TedsTableViewer/Public/Widgets/STedsSearchBox.h
+++ b/Engine/Plugins/Experimental/EditorDataStorageFeatures/Source/TedsTableViewer/Public/Widgets/STedsSearchBox.h
@@ -17,17 +17,22 @@ namespace UE::Editor::DataStorage
 	 * A search bar widget that can be used to search across a TEDS table viewer widget using a QueryStack Search node.
 	 * Example future usage:
 	 * 
-	 *	SNew(STedsSearchBar)
-     *		.OutSearchNode(&OutSearchNode) // TSharedPtr<QueryStack::IRowNode> OutSearchNode
-     *		.InSearchableRowNode(ReferenceQueryResultsNode)
+	 *	SNew(STedsSearchBox)
+	 *	.OutSearchNode(&OutSearchNode) // TSharedPtr<QueryStack::IRowNode> OutSearchNode
+	 *	.InSearchableRowNode(ReferenceQueryResultsNode)
 	 */
 	class STedsSearchBox : public SCompoundWidget
 	{
 	public:
 		
 		SLATE_BEGIN_ARGS(STedsSearchBox) {}
-			SLATE_ARGUMENT(TSharedPtr<QueryStack::IRowNode>, InSearchableRowNode) // The node whose rows we'll perform a search on
-			SLATE_ARGUMENT(TSharedPtr<QueryStack::IRowNode>*, OutSearchNode)      // The search results node we receive back
+
+		// The node whose rows we'll perform a search on
+		SLATE_ARGUMENT(TSharedPtr<QueryStack::IRowNode>, InSearchableRowNode)
+
+		// The search results node we create and receive back
+		SLATE_ARGUMENT(TSharedPtr<QueryStack::IRowNode>*, OutSearchNode)
+			
 		SLATE_END_ARGS()
 
 	public:
@@ -38,7 +43,7 @@ namespace UE::Editor::DataStorage
 		void OnTextChanged(const FText& InSearchText);
 
 		/** Internal SharedPtr of the OutSearchNode so that when text is updated, it can ForceRefresh the node */
-		TSharedPtr<QueryStack::FRowFilterNode> SearchNode;
+		TSharedPtr<QueryStack::FRowFilterNode> OutSearchNode;
 		FText SearchText;
 	};
 }
\ No newline at end of file
