[TEDS Viewers] Allow Sorting to be persisted via IsEnabled and GetColumnSort functions on the TEDS Sort + Inversion Nodes

diff --git a/Engine/Plugins/Experimental/EditorDataStorageFeatures/Source/TedsOutliner/Private/TedsOutlinerImpl.cpp b/Engine/Plugins/Experimental/EditorDataStorageFeatures/Source/TedsOutliner/Private/TedsOutlinerImpl.cpp
--- a/Engine/Plugins/Experimental/EditorDataStorageFeatures/Source/TedsOutliner/Private/TedsOutlinerImpl.cpp
+++ b/Engine/Plugins/Experimental/EditorDataStorageFeatures/Source/TedsOutliner/Private/TedsOutlinerImpl.cpp
@@ -868,8 +868,9 @@ void FTedsOutlinerImpl::RecompileQueries()
 	RowHandleSortNode = MakeShared<FRowHandleSortNode>(CopyNode);
 
 	// Nodes to actually sort the rows by the column the user requested
-	RowPrimarySortingNode = MakeShared<FRowSortNode>(*Storage, RowHandleSortNode);
-	RowPrimaryInversionNode = MakeShared<FRowOrderInversionNode>(RowPrimarySortingNode, false);
+	// Persist the sort order if the sorting nodes are being reset
+	RowPrimarySortingNode = MakeShared<FRowSortNode>(*Storage, RowHandleSortNode, RowPrimarySortingNode ? RowPrimarySortingNode->GetColumnSorter() : nullptr);
+	RowPrimaryInversionNode = MakeShared<FRowOrderInversionNode>(RowPrimarySortingNode, RowPrimaryInversionNode ? RowPrimaryInversionNode->IsEnabled() : false);
 
 	// A custom node to store row handle -> sort index mapping.
 	RowSortMapNode = MakeShared<FRowMapNode>(RowPrimaryInversionNode);
diff --git a/Engine/Plugins/Experimental/EditorDataStorageFeatures/Source/TedsQueryStack/Private/TedsRowOrderInversionNode.cpp b/Engine/Plugins/Experimental/EditorDataStorageFeatures/Source/TedsQueryStack/Private/TedsRowOrderInversionNode.cpp
--- a/Engine/Plugins/Experimental/EditorDataStorageFeatures/Source/TedsQueryStack/Private/TedsRowOrderInversionNode.cpp
+++ b/Engine/Plugins/Experimental/EditorDataStorageFeatures/Source/TedsQueryStack/Private/TedsRowOrderInversionNode.cpp
@@ -22,6 +22,11 @@ namespace UE::Editor::DataStorage::QueryStack
 	{
 		bIsEnabled = bInIsEnabled;
 	}
+	
+	bool FRowOrderInversionNode::IsEnabled() const
+    {
+    	return bIsEnabled;
+    }
 
 	INode::RevisionId FRowOrderInversionNode::GetRevision() const
 	{
diff --git a/Engine/Plugins/Experimental/EditorDataStorageFeatures/Source/TedsQueryStack/Private/TedsRowSortNode.cpp b/Engine/Plugins/Experimental/EditorDataStorageFeatures/Source/TedsQueryStack/Private/TedsRowSortNode.cpp
--- a/Engine/Plugins/Experimental/EditorDataStorageFeatures/Source/TedsQueryStack/Private/TedsRowSortNode.cpp
+++ b/Engine/Plugins/Experimental/EditorDataStorageFeatures/Source/TedsQueryStack/Private/TedsRowSortNode.cpp
@@ -51,6 +51,11 @@ namespace UE::Editor::DataStorage::QueryStack
 		// Restart sort on next update.
 		LastUpdatedRevision = TNumericLimits<RevisionId>::Max();
 	}
+	
+	TSharedPtr<const FColumnSorterInterface> FRowSortNode::GetColumnSorter() const
+    {
+    	return ColumnSorter;
+    }
 
 	bool FRowSortNode::IsSorting() const
 	{
diff --git a/Engine/Plugins/Experimental/EditorDataStorageFeatures/Source/TedsQueryStack/Public/TedsRowOrderInversionNode.h b/Engine/Plugins/Experimental/EditorDataStorageFeatures/Source/TedsQueryStack/Public/TedsRowOrderInversionNode.h
--- a/Engine/Plugins/Experimental/EditorDataStorageFeatures/Source/TedsQueryStack/Public/TedsRowOrderInversionNode.h
+++ b/Engine/Plugins/Experimental/EditorDataStorageFeatures/Source/TedsQueryStack/Public/TedsRowOrderInversionNode.h
@@ -20,6 +20,7 @@ namespace UE::Editor::DataStorage::QueryStack
 
 		/** Clears out the current row list and gets a fresh copy from the parent. */
 		TEDSQUERYSTACK_API void Enable(bool bInIsEnabled);
+		TEDSQUERYSTACK_API bool IsEnabled() const;
 
 		TEDSQUERYSTACK_API virtual RevisionId GetRevision() const override;
 		TEDSQUERYSTACK_API virtual void Update() override;
diff --git a/Engine/Plugins/Experimental/EditorDataStorageFeatures/Source/TedsQueryStack/Public/TedsRowSortNode.h b/Engine/Plugins/Experimental/EditorDataStorageFeatures/Source/TedsQueryStack/Public/TedsRowSortNode.h
--- a/Engine/Plugins/Experimental/EditorDataStorageFeatures/Source/TedsQueryStack/Public/TedsRowSortNode.h
+++ b/Engine/Plugins/Experimental/EditorDataStorageFeatures/Source/TedsQueryStack/Public/TedsRowSortNode.h
@@ -38,6 +38,7 @@ namespace UE::Editor::DataStorage::QueryStack
 		virtual ~FRowSortNode() override = default;
 
 		TEDSQUERYSTACK_API void SetColumnSorter(TSharedPtr<const FColumnSorterInterface> InColumnSorter);
+		TEDSQUERYSTACK_API TSharedPtr<const FColumnSorterInterface> GetColumnSorter() const;
 		TEDSQUERYSTACK_API bool IsSorting() const;
 
 		TEDSQUERYSTACK_API virtual RevisionId GetRevision() const override;
diff --git a/Engine/Plugins/Experimental/EditorDataStorageFeatures/Source/TedsTableViewer/Private/TedsTableViewerModel.cpp b/Engine/Plugins/Experimental/EditorDataStorageFeatures/Source/TedsTableViewer/Private/TedsTableViewerModel.cpp
--- a/Engine/Plugins/Experimental/EditorDataStorageFeatures/Source/TedsTableViewer/Private/TedsTableViewerModel.cpp
+++ b/Engine/Plugins/Experimental/EditorDataStorageFeatures/Source/TedsTableViewer/Private/TedsTableViewerModel.cpp
@@ -158,14 +158,12 @@ namespace UE::Editor::DataStorage
 	void FTedsTableViewerModel::SetQueryStack(const TSharedPtr<QueryStack::IRowNode>& InRowQueryStack)
 	{
 		checkf(InRowQueryStack, TEXT("The TEDS Table Viewer requires a valid query stack to provide it with rows to process."));
-		RowSecondarySortingNode = MakeShared<QueryStack::FRowSortNode>(*Storage, InRowQueryStack);
-		RowSecondaryInversionNode = MakeShared<QueryStack::FRowOrderInversionNode>(RowSecondarySortingNode, false);
-		RowPrimarySortingNode = MakeShared<QueryStack::FRowSortNode>(*Storage, RowSecondaryInversionNode);
-		RowPrimaryInversionNode = MakeShared<QueryStack::FRowOrderInversionNode>(RowPrimarySortingNode, false);
+		// If the sorting nodes have been made and are just being reset to a new query stack, persist the sorting state
+		RowSecondarySortingNode = MakeShared<QueryStack::FRowSortNode>(*Storage, InRowQueryStack, RowSecondarySortingNode ? RowSecondarySortingNode->GetColumnSorter() : nullptr);
+		RowSecondaryInversionNode = MakeShared<QueryStack::FRowOrderInversionNode>(RowSecondarySortingNode, RowSecondaryInversionNode ? RowSecondaryInversionNode->IsEnabled() : false);
+		RowPrimarySortingNode = MakeShared<QueryStack::FRowSortNode>(*Storage, RowSecondaryInversionNode, RowPrimarySortingNode ? RowPrimarySortingNode->GetColumnSorter() : nullptr);
+		RowPrimaryInversionNode = MakeShared<QueryStack::FRowOrderInversionNode>(RowPrimarySortingNode, RowPrimaryInversionNode ? RowPrimaryInversionNode->IsEnabled() : false);
 		RowQueryStack = RowPrimaryInversionNode;
-
-		PrimarySortColumnName = NAME_None;
-		SecondarySortColumnName = NAME_None;
 		
 		Refresh();
 	}
