Fix TedsType info assert when running certain Verse automated tests

diff --git a/Engine/Plugins/Experimental/EditorDataStorageFeatures/Source/TedsTypeInfo/Private/TedsTypeInfoFactory.cpp b/Engine/Plugins/Experimental/EditorDataStorageFeatures/Source/TedsTypeInfo/Private/TedsTypeInfoFactory.cpp
--- a/Engine/Plugins/Experimental/EditorDataStorageFeatures/Source/TedsTypeInfo/Private/TedsTypeInfoFactory.cpp
+++ b/Engine/Plugins/Experimental/EditorDataStorageFeatures/Source/TedsTypeInfo/Private/TedsTypeInfoFactory.cpp
@@ -23,6 +23,11 @@
 const FName UTypeInfoFactory::TypeTableName = "Type Information";
 const FName UTypeInfoFactory::ClassSetupActionName = "ClassSetupAction";
 
+void UTypeInfoFactory::PreRegister(UE::Editor::DataStorage::ICoreProvider& DataStorage)
+{
+	DataStorage.OnUpdateCompleted().AddUObject(this, &UTypeInfoFactory::OnDatabaseUpdateCompleted);
+}
+
 void UTypeInfoFactory::RegisterTables(UE::Editor::DataStorage::ICoreProvider& DataStorage, UE::Editor::DataStorage::ICompatibilityProvider& DataStorageCompatability)
 {
 	using namespace UE::Editor::DataStorage;
@@ -81,6 +86,20 @@ void UTypeInfoFactory::RegisterQueries(UE::Editor::DataStorage::ICoreProvider& D
 	TypeInfo::FTedsTypeInfoModule::GetChecked().SetTypeInfoFactoryEnabled();
 }
 
+void UTypeInfoFactory::PreShutdown(UE::Editor::DataStorage::ICoreProvider& DataStorage)
+{
+	DataStorage.OnUpdateCompleted().RemoveAll(this);
+}
+
+void UTypeInfoFactory::OnDatabaseUpdateCompleted()
+{
+	if (bRefreshTypeInfoQueued)
+	{
+		PopulateTypeInfo<UClass>();
+		bRefreshTypeInfoQueued = false;
+	}
+}
+
 UE::Editor::DataStorage::TableHandle UTypeInfoFactory::GetTypeTable() const
 {
 	return TypeTable;
@@ -95,7 +114,7 @@ void UTypeInfoFactory::ClearAllTypeInfo()
 void UTypeInfoFactory::RefreshAllTypeInfo()
 {
 	ClearAllTypeInfo();
-	PopulateTypeInfo<UClass>();
+	bRefreshTypeInfoQueued = true;
 }
 
 template <class Type>
diff --git a/Engine/Plugins/Experimental/EditorDataStorageFeatures/Source/TedsTypeInfo/Private/TedsTypeInfoFactory.h b/Engine/Plugins/Experimental/EditorDataStorageFeatures/Source/TedsTypeInfo/Private/TedsTypeInfoFactory.h
--- a/Engine/Plugins/Experimental/EditorDataStorageFeatures/Source/TedsTypeInfo/Private/TedsTypeInfoFactory.h
+++ b/Engine/Plugins/Experimental/EditorDataStorageFeatures/Source/TedsTypeInfo/Private/TedsTypeInfoFactory.h
@@ -41,8 +41,10 @@ public:
 	virtual ~UTypeInfoFactory() override = default;
 
 	//~ Begin UEditorDataStorageFactory interface
+	void PreRegister(UE::Editor::DataStorage::ICoreProvider& DataStorage) override;
 	void RegisterTables(UE::Editor::DataStorage::ICoreProvider& DataStorage, UE::Editor::DataStorage::ICompatibilityProvider& DataStorageCompatability) override;
 	void RegisterQueries(UE::Editor::DataStorage::ICoreProvider& DataStorage) override;
+	void PreShutdown(UE::Editor::DataStorage::ICoreProvider& DataStorage) override;
 	//~ End UEditorDataStorageFactory interface
 
 	UE::Editor::DataStorage::TableHandle GetTypeTable() const;
@@ -52,6 +54,8 @@ public:
 
 private:
 
+	void OnDatabaseUpdateCompleted();
+
 	template <class Type>
 	void PopulateTypeInfo();
 
@@ -68,6 +72,8 @@ private:
 	void AddClassColumns(UE::Editor::DataStorage::ICoreProvider& DataStorage, UE::Editor::DataStorage::RowHandle TypeRowHandle, const UClass* InClassInfo);
 	void AddVerseColumns(UE::Editor::DataStorage::ICoreProvider& DataStorage, UE::Editor::DataStorage::RowHandle TypeRowHandle, const UVerseClass* InVerseTypeInfo);
 
+	bool bRefreshTypeInfoQueued = false;
+
 	const UE::Editor::DataStorage::IHierarchyAccessInterface* ClassHierarchyAccessInterface = nullptr;
 
 	UE::Editor::DataStorage::TableHandle TypeTable = UE::Editor::DataStorage::InvalidTableHandle;
diff --git a/Engine/Plugins/Experimental/EditorDataStorageFeatures/Source/TedsTypeInfo/Private/TedsTypeInfoModule.cpp b/Engine/Plugins/Experimental/EditorDataStorageFeatures/Source/TedsTypeInfo/Private/TedsTypeInfoModule.cpp
--- a/Engine/Plugins/Experimental/EditorDataStorageFeatures/Source/TedsTypeInfo/Private/TedsTypeInfoModule.cpp
+++ b/Engine/Plugins/Experimental/EditorDataStorageFeatures/Source/TedsTypeInfo/Private/TedsTypeInfoModule.cpp
@@ -14,7 +14,7 @@ namespace UE::Editor::DataStorage::TypeInfo
 	namespace Private
 	{
 		TAutoConsoleVariable<bool> CVarTEDSTypeInfoEnabled(TEXT("TEDS.Feature.TypeInfoIntegration.Enable"),
-			false,
+			true,
 			TEXT("When true we will store type info into TEDS and propagate type updates when they occur"),
 			FConsoleVariableDelegate::CreateLambda([](IConsoleVariable* Variable)
 				{
