TEDS: Fix typo in IsHierarchyHandleValid function name

diff --git a/Engine/Plugins/Experimental/EditorDataStorage/Source/TedsCore/Private/TypedElementDatabase.cpp b/Engine/Plugins/Experimental/EditorDataStorage/Source/TedsCore/Private/TypedElementDatabase.cpp
--- a/Engine/Plugins/Experimental/EditorDataStorage/Source/TedsCore/Private/TypedElementDatabase.cpp
+++ b/Engine/Plugins/Experimental/EditorDataStorage/Source/TedsCore/Private/TypedElementDatabase.cpp
@@ -1095,7 +1095,7 @@ UE::Editor::DataStorage::FHierarchyHandle UEditorDataStorage::FindHierarchyByNam
 	return Environment->GetHierarchyRegistrar().FindHierarchyByName(Name);
 }
 
-bool UEditorDataStorage::IsValidHierachyHandle(UE::Editor::DataStorage::FHierarchyHandle Handle) const
+bool UEditorDataStorage::IsValidHierarchyHandle(UE::Editor::DataStorage::FHierarchyHandle Handle) const
 {
 	return Environment->GetHierarchyRegistrar().GetAccessInterface(Handle) != nullptr;
 }
diff --git a/Engine/Plugins/Experimental/EditorDataStorage/Source/TedsCore/Private/TypedElementDatabase.h b/Engine/Plugins/Experimental/EditorDataStorage/Source/TedsCore/Private/TypedElementDatabase.h
--- a/Engine/Plugins/Experimental/EditorDataStorage/Source/TedsCore/Private/TypedElementDatabase.h
+++ b/Engine/Plugins/Experimental/EditorDataStorage/Source/TedsCore/Private/TypedElementDatabase.h
@@ -142,7 +142,7 @@ public:
 
 	UE_API virtual UE::Editor::DataStorage::FHierarchyHandle RegisterHierarchy(const UE::Editor::DataStorage::FHierarchyRegistrationParams& Params) override;
 	UE_API virtual UE::Editor::DataStorage::FHierarchyHandle FindHierarchyByName(const FName& Name) const override;
-	UE_API virtual bool IsValidHierachyHandle(UE::Editor::DataStorage::FHierarchyHandle) const override;
+	UE_API virtual bool IsValidHierarchyHandle(UE::Editor::DataStorage::FHierarchyHandle) const override;
 	UE_API virtual const UScriptStruct* GetChildTagType(
 		UE::Editor::DataStorage::FHierarchyHandle) const override;
 	UE_API virtual const UScriptStruct* GetParentTagType(
diff --git a/Engine/Plugins/Experimental/EditorDataStorageFeatures/Source/TedsActorCompatibility/Private/ActorComponentDebugHierarchyWidget/ActorComponentTree.cpp b/Engine/Plugins/Experimental/EditorDataStorageFeatures/Source/TedsActorCompatibility/Private/ActorComponentDebugHierarchyWidget/ActorComponentTree.cpp
--- a/Engine/Plugins/Experimental/EditorDataStorageFeatures/Source/TedsActorCompatibility/Private/ActorComponentDebugHierarchyWidget/ActorComponentTree.cpp
+++ b/Engine/Plugins/Experimental/EditorDataStorageFeatures/Source/TedsActorCompatibility/Private/ActorComponentDebugHierarchyWidget/ActorComponentTree.cpp
@@ -41,7 +41,7 @@ namespace UE::Editor::DataStorage::ActorCompatibility
 				
 				FHierarchyHandle HierarchyHandle = DataStorageInterface->FindHierarchyByName(TEXT("ActorComponent"));
 
-				if (DataStorageInterface->IsValidHierachyHandle(HierarchyHandle))
+				if (DataStorageInterface->IsValidHierarchyHandle(HierarchyHandle))
 				{
 					// Query node which will enumerate all selected actor rows and provide access to their hierarchy
 					TSharedRef<QueryStack::FQueryNode> QueryNode = MakeShared<QueryStack::FQueryNode>(*DataStorageInterface);
diff --git a/Engine/Source/Runtime/TypedElementFramework/Public/Elements/Interfaces/TypedElementDataStorageInterface.h b/Engine/Source/Runtime/TypedElementFramework/Public/Elements/Interfaces/TypedElementDataStorageInterface.h
--- a/Engine/Source/Runtime/TypedElementFramework/Public/Elements/Interfaces/TypedElementDataStorageInterface.h
+++ b/Engine/Source/Runtime/TypedElementFramework/Public/Elements/Interfaces/TypedElementDataStorageInterface.h
@@ -381,7 +381,7 @@ public:
 	
 	virtual UE::Editor::DataStorage::FHierarchyHandle RegisterHierarchy(const UE::Editor::DataStorage::FHierarchyRegistrationParams& Params) = 0;
 	virtual UE::Editor::DataStorage::FHierarchyHandle FindHierarchyByName(const FName& Name) const = 0;
-	virtual bool IsValidHierachyHandle(FHierarchyHandle) const = 0;
+	virtual bool IsValidHierarchyHandle(FHierarchyHandle) const = 0;
 
 	// Gets the tag type indicating the row is a child in this hierarchy and thus has a parent
 	virtual const UScriptStruct* GetChildTagType(FHierarchyHandle InHierarchyHandle) const = 0;
