TEDS: Update all factories that were registering a hierarchy in PreRegister to use the new RegisterHierarchies method

diff --git a/Engine/Plugins/Experimental/EditorDataStorageFeatures/Source/TedsActorCompatibility/Private/Compatibility/TedsActorParentQueries.cpp b/Engine/Plugins/Experimental/EditorDataStorageFeatures/Source/TedsActorCompatibility/Private/Compatibility/TedsActorParentQueries.cpp
--- a/Engine/Plugins/Experimental/EditorDataStorageFeatures/Source/TedsActorCompatibility/Private/Compatibility/TedsActorParentQueries.cpp
+++ b/Engine/Plugins/Experimental/EditorDataStorageFeatures/Source/TedsActorCompatibility/Private/Compatibility/TedsActorParentQueries.cpp
@@ -35,15 +35,18 @@ void UActorParentDataStorageFactory::RegisterQueries(UE::Editor::DataStorage::IC
 
 void UActorParentDataStorageFactory::PreRegister(UE::Editor::DataStorage::ICoreProvider& InDataStorage)
 {
-	UE::Editor::DataStorage::FHierarchyRegistrationParams ActorHierarchyParams{.Name = UE::Editor::DataStorage::ActorParentQueries::Private::ActorEditorHierarchyName};
-	EditorActorHierarchyHandle = InDataStorage.RegisterHierarchy(ActorHierarchyParams);
-
 	GEngine->OnLevelActorDetached().AddUObject(this, &UActorParentDataStorageFactory::OnLevelActorDetached);
 
 	// Store a pointer to data storage so we don't have to grab it from the global instance of the registry every time
 	DataStorage = &InDataStorage;
 }
 
+void UActorParentDataStorageFactory::RegisterHierarchies(UE::Editor::DataStorage::ICoreProvider& InDataStorage)
+{
+	UE::Editor::DataStorage::FHierarchyRegistrationParams ActorHierarchyParams{.Name = UE::Editor::DataStorage::ActorParentQueries::Private::ActorEditorHierarchyName};
+	EditorActorHierarchyHandle = InDataStorage.RegisterHierarchy(ActorHierarchyParams);
+}
+
 void UActorParentDataStorageFactory::PreShutdown(UE::Editor::DataStorage::ICoreProvider& InDataStorage)
 {
 	GEngine->OnLevelActorDetached().RemoveAll(this);
diff --git a/Engine/Plugins/Experimental/EditorDataStorageFeatures/Source/TedsActorCompatibility/Private/Compatibility/TedsActorParentQueries.h b/Engine/Plugins/Experimental/EditorDataStorageFeatures/Source/TedsActorCompatibility/Private/Compatibility/TedsActorParentQueries.h
--- a/Engine/Plugins/Experimental/EditorDataStorageFeatures/Source/TedsActorCompatibility/Private/Compatibility/TedsActorParentQueries.h
+++ b/Engine/Plugins/Experimental/EditorDataStorageFeatures/Source/TedsActorCompatibility/Private/Compatibility/TedsActorParentQueries.h
@@ -18,6 +18,7 @@ public:
 
 	virtual void RegisterQueries(UE::Editor::DataStorage::ICoreProvider& DataStorage) override;
 	virtual void PreRegister(UE::Editor::DataStorage::ICoreProvider& DataStorage) override;
+	virtual void RegisterHierarchies(UE::Editor::DataStorage::ICoreProvider& DataStorage) override;
 	virtual void PreShutdown(UE::Editor::DataStorage::ICoreProvider& DataStorage) override;
 
 private:
diff --git a/Engine/Plugins/Experimental/EditorDataStorageFeatures/Source/TedsTypeInfo/Private/TedsTypeInfoFactory.cpp b/Engine/Plugins/Experimental/EditorDataStorageFeatures/Source/TedsTypeInfo/Private/TedsTypeInfoFactory.cpp
--- a/Engine/Plugins/Experimental/EditorDataStorageFeatures/Source/TedsTypeInfo/Private/TedsTypeInfoFactory.cpp
+++ b/Engine/Plugins/Experimental/EditorDataStorageFeatures/Source/TedsTypeInfo/Private/TedsTypeInfoFactory.cpp
@@ -29,18 +29,23 @@ void UTypeInfoFactory::RegisterTables(UE::Editor::DataStorage::ICoreProvider& Da
 
 	Super::RegisterTables(DataStorage, DataStorageCompatability);
 
+	TypeTable = DataStorage.RegisterTable({
+		FTypedElementLabelColumn::StaticStruct(),
+		FTypeInfoTag::StaticStruct(),
+		FTypeInfoRequiresHierarchyUpdateTag::StaticStruct()
+		}, TypeTableName);
+}
+
+void UTypeInfoFactory::RegisterHierarchies(UE::Editor::DataStorage::ICoreProvider& DataStorage)
+{
+	using namespace UE::Editor::DataStorage;
+
 	const FHierarchyRegistrationParams Params
 	{
 		.Name = TEXT("ClassHierarchy"),
 	};
 
 	FHierarchyHandle Handle = DataStorage.RegisterHierarchy(Params);
-
-	TypeTable = DataStorage.RegisterTable({
-		FTypedElementLabelColumn::StaticStruct(),
-		FTypeInfoTag::StaticStruct(),
-		FTypeInfoRequiresHierarchyUpdateTag::StaticStruct()
-		}, TypeTableName);
 }
 
 void UTypeInfoFactory::RegisterQueries(UE::Editor::DataStorage::ICoreProvider& DataStorage)
@@ -50,7 +55,7 @@ void UTypeInfoFactory::RegisterQueries(UE::Editor::DataStorage::ICoreProvider& D
 
 	DataStorage.RegisterQuery(
 		Select(
-			TEXT("Setup Class Hierachy Info"),
+			TEXT("Setup Class Hierarchy Info"),
 			FProcessor(EQueryTickPhase::DuringPhysics, DataStorage.GetQueryTickGroupName(EQueryTickGroups::Default))
 				.MakeActivatable(ClassSetupActionName),
 			[this](IQueryContext& Context, RowHandle Row, const FTypedElementClassTypeInfoColumn& ClassTypeColumn)
diff --git a/Engine/Plugins/Experimental/EditorDataStorageFeatures/Source/TedsTypeInfo/Private/TedsTypeInfoFactory.h b/Engine/Plugins/Experimental/EditorDataStorageFeatures/Source/TedsTypeInfo/Private/TedsTypeInfoFactory.h
--- a/Engine/Plugins/Experimental/EditorDataStorageFeatures/Source/TedsTypeInfo/Private/TedsTypeInfoFactory.h
+++ b/Engine/Plugins/Experimental/EditorDataStorageFeatures/Source/TedsTypeInfo/Private/TedsTypeInfoFactory.h
@@ -42,6 +42,7 @@ public:
 
 	//~ Begin UEditorDataStorageFactory interface
 	void RegisterTables(UE::Editor::DataStorage::ICoreProvider& DataStorage, UE::Editor::DataStorage::ICompatibilityProvider& DataStorageCompatability) override;
+	void RegisterHierarchies(UE::Editor::DataStorage::ICoreProvider& DataStorage) override;
 	void RegisterQueries(UE::Editor::DataStorage::ICoreProvider& DataStorage) override;
 	//~ End UEditorDataStorageFactory interface
 
