[TEDS Debugger] TEDS Debugger improvements
- Multiple Tabs to have seperate queries
- Added icons to List/Hierarchy View

diff --git a/Engine/Plugins/Experimental/EditorDataStorageFeatures/Source/TedsDebugger/Private/QueryEditor/TedsQueryEditor.cpp b/Engine/Plugins/Experimental/EditorDataStorageFeatures/Source/TedsDebugger/Private/QueryEditor/TedsQueryEditor.cpp
--- a/Engine/Plugins/Experimental/EditorDataStorageFeatures/Source/TedsDebugger/Private/QueryEditor/TedsQueryEditor.cpp
+++ b/Engine/Plugins/Experimental/EditorDataStorageFeatures/Source/TedsDebugger/Private/QueryEditor/TedsQueryEditor.cpp
@@ -118,7 +118,7 @@ namespace UE::Editor::DataStorage::Debug::QueryEditor
 			FOnGetContent::CreateSP(this, &SQueryEditorWidget::GetViewModesMenuWidget),
 			TAttribute<FText>::CreateSP(this, &SQueryEditorWidget::GetCurrentViewModeText),
 			TAttribute<FText>(),
-			FSlateIcon(),
+			TAttribute<FSlateIcon>::CreateSP(this, &SQueryEditorWidget::GetCurrentViewModeIcon),
 			/*bInSimpleComboBox*/ false);
 
 		if (HierarchyComboWidget)
@@ -147,7 +147,7 @@ namespace UE::Editor::DataStorage::Debug::QueryEditor
 		MenuBuilder.AddMenuEntry(
 			GetViewModeAsText(ETableViewMode::List),
 			FText::GetEmpty(),
-			FSlateIcon(), 
+			GetViewModeAsIcon(ETableViewMode::List), 
 			FUIAction(FExecuteAction::CreateLambda([this]()
 			{
 				if (ResultsView)
@@ -160,7 +160,7 @@ namespace UE::Editor::DataStorage::Debug::QueryEditor
 		MenuBuilder.AddMenuEntry(
 			GetViewModeAsText(ETableViewMode::Tree),
 			FText::GetEmpty(),
-			FSlateIcon(), 
+			GetViewModeAsIcon(ETableViewMode::Tree), 
 			FUIAction(FExecuteAction::CreateLambda([this]()
 			{
 				if (ResultsView)
@@ -178,6 +178,11 @@ namespace UE::Editor::DataStorage::Debug::QueryEditor
 		return ResultsView ? GetViewModeAsText(ResultsView->GetViewMode()) : FText::GetEmpty();
 	}
 
+	FSlateIcon SQueryEditorWidget::GetCurrentViewModeIcon() const
+	{
+		return ResultsView ? GetViewModeAsIcon(ResultsView->GetViewMode()) : FSlateIcon();
+	}
+
 	FText SQueryEditorWidget::GetViewModeAsText(ETableViewMode::Type InViewMode)
 	{
 		switch (InViewMode)
@@ -192,6 +197,21 @@ namespace UE::Editor::DataStorage::Debug::QueryEditor
 			return FText::GetEmpty();
 		}
 	}
+
+	FSlateIcon SQueryEditorWidget::GetViewModeAsIcon(ETableViewMode::Type InViewMode)
+	{
+		switch (InViewMode)
+		{
+		case ETableViewMode::List:
+			return FSlateIcon("UMGStyle", "ClassIcon.ListView");
+		case ETableViewMode::Tile:
+			return FSlateIcon("UMGStyle", "ClassIcon.TileView");
+		case ETableViewMode::Tree:
+			return FSlateIcon("UMGStyle", "ClassIcon.TreeView");
+		default:
+			return FSlateIcon();
+		}
+	}
 } // namespace UE::Editor::DataStorage::Debug::QueryEditor
 
 #undef LOCTEXT_NAMESPACE
\ No newline at end of file
diff --git a/Engine/Plugins/Experimental/EditorDataStorageFeatures/Source/TedsDebugger/Private/QueryEditor/TedsQueryEditor.h b/Engine/Plugins/Experimental/EditorDataStorageFeatures/Source/TedsDebugger/Private/QueryEditor/TedsQueryEditor.h
--- a/Engine/Plugins/Experimental/EditorDataStorageFeatures/Source/TedsDebugger/Private/QueryEditor/TedsQueryEditor.h
+++ b/Engine/Plugins/Experimental/EditorDataStorageFeatures/Source/TedsDebugger/Private/QueryEditor/TedsQueryEditor.h
@@ -3,6 +3,7 @@
 
 #include "Framework/Views/ITypedTableView.h"
 #include "Templates/SharedPointerFwd.h"
+#include "Textures/SlateIcon.h"
 #include "Widgets/SCompoundWidget.h"
 
 class SDockTab;
@@ -35,6 +36,8 @@ namespace UE::Editor::DataStorage::Debug::QueryEditor
 		TSharedRef<SWidget> GetViewModesMenuWidget();
 		FText GetCurrentViewModeText() const;
 		static FText GetViewModeAsText(ETableViewMode::Type InViewMode);
+		FSlateIcon GetCurrentViewModeIcon() const;
+		static FSlateIcon GetViewModeAsIcon(ETableViewMode::Type InViewMode);
 		
 	private:
 		
diff --git a/Engine/Plugins/Experimental/EditorDataStorageFeatures/Source/TedsDebugger/Private/STedsDebugger.cpp b/Engine/Plugins/Experimental/EditorDataStorageFeatures/Source/TedsDebugger/Private/STedsDebugger.cpp
--- a/Engine/Plugins/Experimental/EditorDataStorageFeatures/Source/TedsDebugger/Private/STedsDebugger.cpp
+++ b/Engine/Plugins/Experimental/EditorDataStorageFeatures/Source/TedsDebugger/Private/STedsDebugger.cpp
@@ -21,10 +21,17 @@ namespace UE::Editor::DataStorage::Debug
 {
 	namespace Private
 	{
-		FName QueryEditorToolTabName = TEXT("TEDS Query Editor");
 		FName ToolbarTabName = TEXT("TEDS Debugger Toolbar");
 	}
 
+enum class EQueryEditorTabs
+{
+	QueryEditorOne,
+	QueryEditorTwo,
+	QueryEditorThree,
+	QueryEditorFour,
+};
+
 STedsDebugger::~STedsDebugger()
 {
 	if (AreEditorDataStorageFeaturesEnabled())
@@ -39,11 +46,13 @@ void STedsDebugger::Construct(const FArguments& InArgs, const TSharedRef<SDockTa
 	TabManager = FGlobalTabmanager::Get()->NewTabManager(ConstructUnderMajorTab);
 	TabManager->SetAllowWindowMenuBar(true);
 
+	Models.SetNum(MaxQueryEditorTabs);
+
 	// Register Tab Spawners
 	RegisterTabSpawners();
 
 	// Setup the default layout
-	TSharedRef<FTabManager::FLayout> Layout = FTabManager::NewLayout("TedsDebuggerLayout_v0")
+	TSharedRef<FTabManager::FLayout> Layout = FTabManager::NewLayout("TedsDebuggerLayout_v1.0")
 	->AddArea
 	(
 		FTabManager::NewPrimaryArea()
@@ -61,7 +70,12 @@ void STedsDebugger::Construct(const FArguments& InArgs, const TSharedRef<SDockTa
 			->Split
 			(
 				FTabManager::NewStack()
-					->AddTab(Private::QueryEditorToolTabName, ETabState::OpenedTab)
+					->AddTab(GetQueryTabName(EQueryEditorTabs::QueryEditorOne), ETabState::OpenedTab)
+					->AddTab(GetQueryTabName(EQueryEditorTabs::QueryEditorTwo), ETabState::ClosedTab)
+					->AddTab(GetQueryTabName(EQueryEditorTabs::QueryEditorThree), ETabState::ClosedTab)
+					->AddTab(GetQueryTabName(EQueryEditorTabs::QueryEditorFour), ETabState::ClosedTab)
+					->SetHideTabWell(false)
+					->SetForegroundTab(GetQueryTabName(EQueryEditorTabs::QueryEditorOne))
 			)
 		)
 	);
@@ -109,6 +123,27 @@ void STedsDebugger::FillWindowMenu( FMenuBuilder& MenuBuilder)
 	}
 }
 
+QueryEditor::FTedsQueryEditorModel* STedsDebugger::FindQueryEditorModel(const EQueryEditorTabs QueryTabIdentifier)
+{
+	if (AreEditorDataStorageFeaturesEnabled(); ICoreProvider* DataStorageInterface = GetMutableDataStorageFeature<ICoreProvider>(StorageFeatureName))
+	{
+		TUniquePtr<QueryEditor::FTedsQueryEditorModel>& Model = Models[static_cast<int32>(QueryTabIdentifier)];
+		if (!Model)
+		{
+			Model = MakeUnique<QueryEditor::FTedsQueryEditorModel>(*DataStorageInterface);
+		}
+		return Model.Get();
+	}
+
+	return nullptr; // Return null if TEDS is disabled
+}
+
+FName STedsDebugger::GetQueryTabName(const EQueryEditorTabs QueryTabIdentifier)
+{
+	// Returns a name with the syntax: QueryEditorTab, QueryEditorTab_1, QueryEditorTab_2, etc...
+	return FName(TEXT("QueryEditorTab"), static_cast<int32>(QueryTabIdentifier));
+}
+
 TSharedRef<SDockTab> STedsDebugger::SpawnToolbar(const FSpawnTabArgs& Args)
 {
 	// The toolbar is currently empty but can be used to house tools that are not specific to a specific tab in the debugger
@@ -122,33 +157,27 @@ TSharedRef<SDockTab> STedsDebugger::SpawnToolbar(const FSpawnTabArgs& Args)
 		];
 }
 
-TSharedRef<SDockTab> STedsDebugger::SpawnQueryEditorTab(const FSpawnTabArgs& Args)
+TSharedRef<SDockTab> STedsDebugger::SpawnQueryEditorTab(const FSpawnTabArgs& Args, const EQueryEditorTabs QueryTabIdentifier)
 {
-	TSharedRef<SDockTab> DockTab = SNew(SDockTab).TabRole(ETabRole::PanelTab);
-	if (!QueryEditorModel)
-	{
-		if (AreEditorDataStorageFeaturesEnabled())
-		{
-			ICoreProvider* DataStorageInterface = GetMutableDataStorageFeature<ICoreProvider>(StorageFeatureName);
-			QueryEditorModel = MakeUnique<QueryEditor::FTedsQueryEditorModel>(*DataStorageInterface);
-		}
-	}
-	if (QueryEditorModel)
+	TSharedRef<SDockTab> DockTab = SNew(SDockTab)
+		.TabRole(ETabRole::PanelTab)
+		.Label(FText::Format(LOCTEXT("TedsDebugger_QueryEditorDisplayName", "Query Editor {0}"), FText::AsNumber(static_cast<int32>(QueryTabIdentifier) + 1)));
+
+	if (QueryEditor::FTedsQueryEditorModel* QueryEditorModel = FindQueryEditorModel(QueryTabIdentifier))
 	{
-		QueryEditorModel->Reset();	
+		QueryEditorModel->Reset();
 
-		TSharedRef<QueryEditor::SQueryEditorWidget> QueryEditor =
+		TSharedRef<QueryEditor::SQueryEditorWidget> QueryEditorWidget =
 			SNew(QueryEditor::SQueryEditorWidget, *QueryEditorModel);
-		DockTab->SetContent(QueryEditor);
+		DockTab->SetContent(QueryEditorWidget);
 	}
 	else
 	{
 		TSharedRef<STextBlock> TextBlock = SNew(STextBlock)
-		.Text(LOCTEXT("TedsDebuggerModule_CannotLoadQueryEditor", "Cannot load Query Editor - Invalid Model"));
+			.Text(LOCTEXT("TedsDebuggerModule_CannotLoadQueryEditor", "Cannot load Query Editor - Invalid Model"));
 		DockTab->SetContent(TextBlock);
 	}
 
-
 	return DockTab;
 }
 
@@ -156,7 +185,8 @@ void STedsDebugger::RegisterTabSpawners()
 {
 	const TSharedRef<FWorkspaceItem> AppMenuGroup =
 		TabManager->AddLocalWorkspaceMenuCategory(LOCTEXT("TedsDebuggerGroupName", "Teds Debugger"));
-
+	const FText QueryEditorTooltip = LOCTEXT("TedsDebugger_QueryEditorToolTip", "Open a Query Editor tab. Use this to create and view individual queries.");
+	
 	TabManager->RegisterTabSpawner(
 		Private::ToolbarTabName,
 		FOnSpawnTab::CreateRaw(this, &STedsDebugger::SpawnToolbar))
@@ -165,11 +195,35 @@ void STedsDebugger::RegisterTabSpawners()
 		.SetAutoGenerateMenuEntry(false);
 	
 	TabManager->RegisterTabSpawner(
-		Private::QueryEditorToolTabName,
-		FOnSpawnTab::CreateRaw(this, &STedsDebugger::SpawnQueryEditorTab))
+		GetQueryTabName(EQueryEditorTabs::QueryEditorOne),
+		FOnSpawnTab::CreateRaw(this, &STedsDebugger::SpawnQueryEditorTab, EQueryEditorTabs::QueryEditorOne))
+		.SetGroup(AppMenuGroup)
+		.SetDisplayName(LOCTEXT("TedsDebugger_QueryEditorDisplayNameOne", "Query Editor 1"))
+		.SetTooltipText(QueryEditorTooltip)
+		.SetIcon(FSlateIcon(FAppStyle::GetAppStyleSetName(), "Debug"));
+
+	TabManager->RegisterTabSpawner(
+		GetQueryTabName(EQueryEditorTabs::QueryEditorTwo),
+		FOnSpawnTab::CreateRaw(this, &STedsDebugger::SpawnQueryEditorTab, EQueryEditorTabs::QueryEditorTwo))
+		.SetGroup(AppMenuGroup)
+		.SetDisplayName(LOCTEXT("TedsDebugger_QueryEditorDisplayNameTwo", "Query Editor 2"))
+		.SetTooltipText(QueryEditorTooltip)
+		.SetIcon(FSlateIcon(FAppStyle::GetAppStyleSetName(), "Debug"));
+
+	TabManager->RegisterTabSpawner(
+		GetQueryTabName(EQueryEditorTabs::QueryEditorThree),
+		FOnSpawnTab::CreateRaw(this, &STedsDebugger::SpawnQueryEditorTab, EQueryEditorTabs::QueryEditorThree))
+		.SetGroup(AppMenuGroup)
+		.SetDisplayName(LOCTEXT("TedsDebugger_QueryEditorDisplayNameThree", "Query Editor 3"))
+		.SetTooltipText(QueryEditorTooltip)
+		.SetIcon(FSlateIcon(FAppStyle::GetAppStyleSetName(), "Debug"));
+
+	TabManager->RegisterTabSpawner(
+		GetQueryTabName(EQueryEditorTabs::QueryEditorFour),
+		FOnSpawnTab::CreateRaw(this, &STedsDebugger::SpawnQueryEditorTab, EQueryEditorTabs::QueryEditorFour))
 		.SetGroup(AppMenuGroup)
-		.SetDisplayName(LOCTEXT("TedsDebugger_QueryEditorDisplayName", "Query Editor"))
-		.SetTooltipText(LOCTEXT("TedsDebugger_QueryEditorToolTip", "Opens the TEDS Query Editor"))
+		.SetDisplayName(LOCTEXT("TedsDebugger_QueryEditorDisplayNameFour", "Query Editor 4"))
+		.SetTooltipText(QueryEditorTooltip)
 		.SetIcon(FSlateIcon(FAppStyle::GetAppStyleSetName(), "Debug"));
 }
 
diff --git a/Engine/Plugins/Experimental/EditorDataStorageFeatures/Source/TedsDebugger/Private/STedsDebugger.h b/Engine/Plugins/Experimental/EditorDataStorageFeatures/Source/TedsDebugger/Private/STedsDebugger.h
--- a/Engine/Plugins/Experimental/EditorDataStorageFeatures/Source/TedsDebugger/Private/STedsDebugger.h
+++ b/Engine/Plugins/Experimental/EditorDataStorageFeatures/Source/TedsDebugger/Private/STedsDebugger.h
@@ -12,6 +12,8 @@ namespace UE::Editor::DataStorage
 {
 	namespace Debug
 	{
+		enum class EQueryEditorTabs;
+		
 		class STedsDebugger : public SCompoundWidget
 		{
 		public:
@@ -34,8 +36,10 @@ namespace UE::Editor::DataStorage
 			*/
 			void Construct(const FArguments& InArgs, const TSharedRef<SDockTab>& ConstructUnderMajorTab, const TSharedPtr<SWindow>& ConstructUnderWindow);
 		private:
+			QueryEditor::FTedsQueryEditorModel* FindQueryEditorModel(const EQueryEditorTabs QueryTabIdentifier);
+			static FName GetQueryTabName(const EQueryEditorTabs QueryTabIdentifier);
 			TSharedRef<SDockTab> SpawnToolbar(const FSpawnTabArgs& Args);
-			TSharedRef<SDockTab> SpawnQueryEditorTab(const FSpawnTabArgs& Args);
+			TSharedRef<SDockTab> SpawnQueryEditorTab(const FSpawnTabArgs& Args, const EQueryEditorTabs QueryTabIdentifier);
 			void FillWindowMenu( FMenuBuilder& MenuBuilder);
 
 			void RegisterTabSpawners();
@@ -48,8 +52,8 @@ namespace UE::Editor::DataStorage
 			// Table Viewer
 			QueryHandle TableViewerQuery;
 
-			// Query Editor
-			TUniquePtr<QueryEditor::FTedsQueryEditorModel> QueryEditorModel;
+			static constexpr int32 MaxQueryEditorTabs = 4;
+			TArray<TUniquePtr<QueryEditor::FTedsQueryEditorModel>, TInlineAllocator<MaxQueryEditorTabs>> Models;
 		};
 	} // namespace Debug
 } // namespace UE::Editor::DataStorage
