[TEDS] Changed the Query Stack to using an explicit executor.

The complexity of running a query stack has slowly increased over time. The change to introduce an executor brings back the simplicity while also avoiding common mistakes like forgetting to update a parent. New features include being able to work with cooperative tasks and registering a callback that only gets called if the query stack has had an update.

Original reviews: 45129745, 45129806, 45129812

diff --git a/Engine/Plugins/Experimental/EditorDataStorageFeatures/Source/TedsEverythingPicker/Private/Context/TedsPickerContextUtil.cpp b/Engine/Plugins/Experimental/EditorDataStorageFeatures/Source/TedsEverythingPicker/Private/Context/TedsPickerContextUtil.cpp
--- a/Engine/Plugins/Experimental/EditorDataStorageFeatures/Source/TedsEverythingPicker/Private/Context/TedsPickerContextUtil.cpp
+++ b/Engine/Plugins/Experimental/EditorDataStorageFeatures/Source/TedsEverythingPicker/Private/Context/TedsPickerContextUtil.cpp
@@ -176,7 +176,7 @@ namespace UE::Editor::DataStorage::Picker
 					MakeShared<FHierarchyRowNode>(*Storage, ClassHierarchyHandle, BaseClassRowNode, FHierarchyRowNode::ESyncFlags::IncrementWhenDifferent);
 
 				// Do an initial update to populate the node
-				HierarchyNode->Update();
+				QueryStack::FExplicitUpdateExecutor("STypeListContextView Query Stack one-shot", HierarchyNode).Update();
 
 				// If we had a valid query and a valid base class we will take the intersection of the two lists as our final
 				if (bHasValidQuery)
diff --git a/Engine/Plugins/Experimental/EditorDataStorageFeatures/Source/TedsOutliner/Private/RowMapNode.cpp b/Engine/Plugins/Experimental/EditorDataStorageFeatures/Source/TedsOutliner/Private/RowMapNode.cpp
--- a/Engine/Plugins/Experimental/EditorDataStorageFeatures/Source/TedsOutliner/Private/RowMapNode.cpp
+++ b/Engine/Plugins/Experimental/EditorDataStorageFeatures/Source/TedsOutliner/Private/RowMapNode.cpp
@@ -26,9 +26,8 @@ namespace UE::Editor::Outliner
 		return Revision;
 	}
 
-	void FRowMapNode::Update()
+	void FRowMapNode::Update(FTimespan AllottedTime)
 	{
-		Parent->Update();
 		if (Parent->GetRevision() != CachedParentRevision)
 		{
 			RefreshMap();
@@ -37,6 +36,11 @@ namespace UE::Editor::Outliner
 		}
 	}
 
+	void FRowMapNode::VisitParents(ParentListCallback Callback)
+	{
+		Callback(Parent);
+	}
+
 	int32 FRowMapNode::GetRowIndex(DataStorage::RowHandle Row) const
 	{
 		if (const int32* Index = RowMap.Find(Row))
diff --git a/Engine/Plugins/Experimental/EditorDataStorageFeatures/Source/TedsOutliner/Private/RowMapNode.h b/Engine/Plugins/Experimental/EditorDataStorageFeatures/Source/TedsOutliner/Private/RowMapNode.h
--- a/Engine/Plugins/Experimental/EditorDataStorageFeatures/Source/TedsOutliner/Private/RowMapNode.h
+++ b/Engine/Plugins/Experimental/EditorDataStorageFeatures/Source/TedsOutliner/Private/RowMapNode.h
@@ -19,12 +19,13 @@ namespace UE::Editor::Outliner
 		virtual DataStorage::FRowHandleArrayView GetRows() const override;
 		virtual DataStorage::FRowHandleArray& GetMutableRows() override;
 		virtual RevisionId GetRevision() const override;
-		virtual void Update() override;
+		virtual void VisitParents(ParentListCallback Callback) override;
 
 		// Look up the index of the given row in the parent node. INDEX_NONE if not found
 		int32 GetRowIndex(DataStorage::RowHandle Row) const;
 		
 	protected:
+		virtual void Update(FTimespan AllottedTime) override;
 		void RefreshMap();
 		
 	protected:
diff --git a/Engine/Plugins/Experimental/EditorDataStorageFeatures/Source/TedsOutliner/Private/TedsOutlinerImpl.cpp b/Engine/Plugins/Experimental/EditorDataStorageFeatures/Source/TedsOutliner/Private/TedsOutlinerImpl.cpp
--- a/Engine/Plugins/Experimental/EditorDataStorageFeatures/Source/TedsOutliner/Private/TedsOutlinerImpl.cpp
+++ b/Engine/Plugins/Experimental/EditorDataStorageFeatures/Source/TedsOutliner/Private/TedsOutlinerImpl.cpp
@@ -5,6 +5,7 @@
 #include "SSceneOutliner.h"
 #include "TedsOutlinerColumn.h"
 #include "TedsTableViewerUtils.h"
+#include "TedsQueryStackExecutor.h"
 #include "DataStorage/Features.h"
 #include "Elements/Columns/TypedElementLabelColumns.h"
 #include "Elements/Columns/TypedElementMiscColumns.h"
@@ -789,7 +790,7 @@ void FTedsOutlinerImpl::RefreshQueries() const
 {
 	InitialRowCollectorNode->Refresh();
 	CombinedRowCollectorNode->Refresh();
-	FinalCombinedRowNode->Update();
+	QueryStackExecutor->Update();
 }
 
 void FTedsOutlinerImpl::RecompileQueries()
@@ -874,6 +875,8 @@ void FTedsOutlinerImpl::RecompileQueries()
 
 	// A custom node to store row handle -> sort index mapping.
 	RowSortMapNode = MakeShared<FRowMapNode>(RowPrimaryInversionNode);
+	static int InstanceCounter = 0;
+	QueryStackExecutor = MakeUnique<QueryStack::FExplicitUpdateExecutor>(FName("TedsOutliner", InstanceCounter++), RowSortMapNode);
 
 	// Our final query to collect rows to populate the Outliner - currently the same as the initial query the user provided
 	FQueryDescription FinalQueryDescription(InitialQueryDescription);
@@ -1141,8 +1144,8 @@ void FTedsOutlinerImpl::Tick()
 		bSelectionDirty = false;
 	}
 
-	// The Sort node is the childmost node in all cases, so updating it causes all the nodes in the stack to update
-	RowSortMapNode->Update();
+	// The Sort node is the child most node in all cases, so updating it causes all the nodes in the stack to update
+	QueryStackExecutor->Update();
 }
 
 bool FTedsOutlinerImpl::IsSorting() const
diff --git a/Engine/Plugins/Experimental/EditorDataStorageFeatures/Source/TedsOutliner/Public/TedsOutlinerImpl.h b/Engine/Plugins/Experimental/EditorDataStorageFeatures/Source/TedsOutliner/Public/TedsOutlinerImpl.h
--- a/Engine/Plugins/Experimental/EditorDataStorageFeatures/Source/TedsOutliner/Public/TedsOutlinerImpl.h
+++ b/Engine/Plugins/Experimental/EditorDataStorageFeatures/Source/TedsOutliner/Public/TedsOutlinerImpl.h
@@ -27,6 +27,7 @@ namespace UE::Editor::DataStorage
 		class FRowMergeNode;
 		class FRowFilterNode;
 		class IRowNode;
+		class FExplicitUpdateExecutor;
 	}
 
 	class IUiProvider;
@@ -320,6 +321,9 @@ protected:
 	// A map of the row handle -> index in the sort node for lookup
 	TSharedPtr<FRowMapNode> RowSortMapNode;
 
+	// The executor used to update the query stack stack.
+	TUniquePtr<DataStorage::QueryStack::FExplicitUpdateExecutor> QueryStackExecutor;
+
 	// Query to get all child rows
 	DataStorage::QueryHandle ChildRowHandleQuery = DataStorage::InvalidQueryHandle;
 
diff --git a/Engine/Plugins/Experimental/EditorDataStorageFeatures/Source/TedsQueryStack/Private/TedsHierarchyNode.cpp b/Engine/Plugins/Experimental/EditorDataStorageFeatures/Source/TedsQueryStack/Private/TedsHierarchyNode.cpp
--- a/Engine/Plugins/Experimental/EditorDataStorageFeatures/Source/TedsQueryStack/Private/TedsHierarchyNode.cpp
+++ b/Engine/Plugins/Experimental/EditorDataStorageFeatures/Source/TedsQueryStack/Private/TedsHierarchyNode.cpp
@@ -24,9 +24,8 @@ namespace UE::Editor::DataStorage
 		return Revision;
 	}
 
-	void FHierarchyRowNode::Update()
+	void FHierarchyRowNode::Update(FTimespan AllottedTime)
 	{
-		TopLevelRows->Update();
 		switch (SyncFlags)
 		{
 			case ESyncFlags::ParentRevisionDifferent:
@@ -69,6 +68,11 @@ namespace UE::Editor::DataStorage
 		}
 	}
 
+	void FHierarchyRowNode::VisitParents(ParentListCallback Callback)
+	{
+		Callback(TopLevelRows);
+	}
+
 	FRowHandleArrayView FHierarchyRowNode::GetRows() const
 	{
 		return Rows.GetRows();
diff --git a/Engine/Plugins/Experimental/EditorDataStorageFeatures/Source/TedsQueryStack/Private/TedsQueryHandleNode.cpp b/Engine/Plugins/Experimental/EditorDataStorageFeatures/Source/TedsQueryStack/Private/TedsQueryHandleNode.cpp
--- a/Engine/Plugins/Experimental/EditorDataStorageFeatures/Source/TedsQueryStack/Private/TedsQueryHandleNode.cpp
+++ b/Engine/Plugins/Experimental/EditorDataStorageFeatures/Source/TedsQueryStack/Private/TedsQueryHandleNode.cpp
@@ -20,11 +20,16 @@ namespace UE::Editor::DataStorage::QueryStack
 		return Revision;
 	}
 
-	void FQueryHandleNode::Update()
+	void FQueryHandleNode::Update(FTimespan AllottedTime)
 	{
 		// Nothing to update.
 	}
 
+	void FQueryHandleNode::VisitParents(ParentListCallback Callback)
+	{
+		// No parents to visit.
+	}
+
 	QueryHandle FQueryHandleNode::GetQuery() const
 	{
 		return Query;
diff --git a/Engine/Plugins/Experimental/EditorDataStorageFeatures/Source/TedsQueryStack/Private/TedsQueryMergeNode.cpp b/Engine/Plugins/Experimental/EditorDataStorageFeatures/Source/TedsQueryStack/Private/TedsQueryMergeNode.cpp
--- a/Engine/Plugins/Experimental/EditorDataStorageFeatures/Source/TedsQueryStack/Private/TedsQueryMergeNode.cpp
+++ b/Engine/Plugins/Experimental/EditorDataStorageFeatures/Source/TedsQueryStack/Private/TedsQueryMergeNode.cpp
@@ -31,12 +31,11 @@ namespace UE::Editor::DataStorage::QueryStack
 		return Revision;
 	}
 
-	void FQueryMergeNode::Update()
+	void FQueryMergeNode::Update(FTimespan AllottedTime)
 	{
 		bool bRebuild = false;
 		for (FParentInfo& ParentInfo : Parents)
 		{
-			ParentInfo.Parent->Update();
 			RevisionId ParentRevision = ParentInfo.Parent->GetRevision();
 			if (ParentInfo.Revision != ParentRevision)
 			{
@@ -52,6 +51,14 @@ namespace UE::Editor::DataStorage::QueryStack
 		}
 	}
 
+	void FQueryMergeNode::VisitParents(ParentListCallback Callback)
+	{
+		for (const FParentInfo& ParentInfo : Parents)
+		{
+			Callback(ParentInfo.Parent);
+		}
+	}
+
 	QueryHandle FQueryMergeNode::GetQuery() const
 	{
 		return QueryHandle;
diff --git a/Engine/Plugins/Experimental/EditorDataStorageFeatures/Source/TedsQueryStack/Private/TedsQueryNode.cpp b/Engine/Plugins/Experimental/EditorDataStorageFeatures/Source/TedsQueryStack/Private/TedsQueryNode.cpp
--- a/Engine/Plugins/Experimental/EditorDataStorageFeatures/Source/TedsQueryStack/Private/TedsQueryNode.cpp
+++ b/Engine/Plugins/Experimental/EditorDataStorageFeatures/Source/TedsQueryStack/Private/TedsQueryNode.cpp
@@ -44,11 +44,16 @@ namespace UE::Editor::DataStorage::QueryStack
 		return Revision;
 	}
 
-	void FQueryNode::Update()
+	void FQueryNode::Update(FTimespan AllottedTime)
 	{
 		// Nothing to update.
 	}
 
+	void FQueryNode::VisitParents(ParentListCallback Callback)
+	{
+		// No parents to report.
+	}
+
 	QueryHandle FQueryNode::GetQuery() const
 	{
 		return QueryHandle;
diff --git a/Engine/Plugins/Experimental/EditorDataStorageFeatures/Source/TedsQueryStack/Private/TedsQueryStackExecutor.cpp b/Engine/Plugins/Experimental/EditorDataStorageFeatures/Source/TedsQueryStack/Private/TedsQueryStackExecutor.cpp
new file mode 100644
--- /dev/null
+++ b/Engine/Plugins/Experimental/EditorDataStorageFeatures/Source/TedsQueryStack/Private/TedsQueryStackExecutor.cpp
@@ -0,0 +1,310 @@
+// Copyright Epic Games, Inc. All Rights Reserved.
+
+#include "TedsQueryStackExecutor.h"
+
+#include <utility>
+#include "TedsQueryStackInterfaces.h"
+
+namespace UE::Editor::DataStorage::QueryStack
+{
+	//
+	// FExecutorBase
+	//
+
+	FExecutorBase::~FExecutorBase() = default;
+
+	const FName& FExecutorBase::GetStackName() const
+	{
+		return StackName;
+	}
+
+	TSharedPtr<const INode> FExecutorBase::GetRootNode() const
+	{
+		return RootNode;
+	}
+
+	TSharedPtr<INode> FExecutorBase::GetRootNode()
+	{
+		return RootNode;
+	}
+
+	void FExecutorBase::SetOnChangedCallback(OnChangedFunction Callback)
+	{
+		OnChangedCallback = MoveTemp(Callback);
+	}
+
+	void FExecutorBase::ClearOnChangedCallback()
+	{
+		OnChangedCallback.Reset();
+	}
+
+	FExecutorBase::FExecutorBase(const FName& StackName, TSharedPtr<INode> RootNode)
+		: RootNode(MoveTemp(RootNode))
+		, StackName(StackName)
+	{
+	}
+
+	void FExecutorBase::FullUpdate()
+	{
+		// Expectation is that the implementers check if RootNode is valid. This means that the versions that registers
+		// to callbacks only have to check once.
+		if (OnChangedCallback)
+		{
+			INode::RevisionId Revision = RootNode->GetRevision();
+			FullNodeUpdate(RootNode);
+			if (RootNode->GetRevision() != Revision)
+			{
+				OnChangedCallback(StackName);
+			}
+		}
+		else
+		{
+			FullNodeUpdate(RootNode);
+		}
+	}
+
+	void FExecutorBase::FullNodeUpdate(const TSharedPtr<INode>& Node)
+	{
+		Node->VisitParents(
+			[](TSharedPtr<INode> Parent)
+			{
+				FullNodeUpdate(Parent);
+			});
+		Node->Update(INode::UnlimitedTime);
+	}
+
+	void FExecutorBase::TimedUpdate(FTimespan AllottedTime)
+	{
+		// Expectation is that the implementers check if RootNode is valid. This means that the versions that registers
+		// to callbacks only have to check once.
+		if (OnChangedCallback)
+		{
+			INode::RevisionId Revision = RootNode->GetRevision();
+			TimedNodeUpdate(RootNode, FPlatformTime::Cycles64(), AllottedTime);
+			if (RootNode->GetRevision() != Revision)
+			{
+				OnChangedCallback(StackName);
+			}
+		}
+		else
+		{
+			TimedNodeUpdate(RootNode, FPlatformTime::Cycles64(), AllottedTime);
+		}
+	}
+
+	void FExecutorBase::TimedNodeUpdate(const TSharedPtr<INode>& Node, uint64 StartTime, FTimespan AllottedTime)
+	{
+		Node->VisitParents(
+			[StartTime, AllottedTime](TSharedPtr<INode> Parent)
+			{
+				TimedNodeUpdate(Parent, StartTime, AllottedTime);
+			});
+		FTimespan Duration(FTimespan::FromSeconds(FPlatformTime::ToSeconds64(FPlatformTime::Cycles64() - StartTime)));
+		if (Duration < AllottedTime)
+		{
+			Node->Update(AllottedTime - Duration);
+		}
+	}
+
+
+
+	//
+	// FCooperativeExecutor
+	//
+
+	FCooperativeExecutor::FCooperativeExecutor(const FName& InStackName, ICoreProvider& InDataStorage, TSharedPtr<INode> InRootNode,
+		ICoreProvider::ECooperativeTaskPriority InPriority)
+		: FExecutorBase(InStackName, MoveTemp(InRootNode))
+		, DataStorage(&InDataStorage)
+		, Priority(InPriority)
+	{
+		if (RootNode)
+		{
+			DataStorage->RegisterCooperativeUpdate(StackName, InPriority,
+				[this](FTimespan AllottedTime)
+				{
+					TimedUpdate(AllottedTime);
+				});
+		}
+	}
+
+	FCooperativeExecutor::~FCooperativeExecutor()
+	{
+		if (RootNode && DataStorage)
+		{
+			DataStorage->UnregisterCooperativeUpdate(StackName);
+		}
+	}
+
+	FCooperativeExecutor::FCooperativeExecutor(FCooperativeExecutor&& Rhs)
+		: FExecutorBase(Rhs.StackName, MoveTemp(Rhs.RootNode))
+		, DataStorage(std::exchange(Rhs.DataStorage, nullptr))
+		, Priority(Rhs.Priority)
+	{
+		if (RootNode && DataStorage)
+		{
+			// Remove the call registered by the moved object.
+			DataStorage->UnregisterCooperativeUpdate(StackName);
+			// Register the new callback
+			DataStorage->RegisterCooperativeUpdate(StackName, Priority,
+				[this](FTimespan AllottedTime)
+				{
+					TimedUpdate(AllottedTime);
+				});
+		}
+	}
+
+	FCooperativeExecutor& FCooperativeExecutor::operator=(FCooperativeExecutor&& Rhs)
+	{
+		if (this != &Rhs)
+		{
+			this->~FCooperativeExecutor();
+			new(this) FCooperativeExecutor(MoveTemp(Rhs));
+		}
+		return *this;
+	}
+
+
+
+	//
+	// FOnUpdateExecutor
+	//
+	
+	FOnUpdateExecutor::FOnUpdateExecutor(const FName& InStackName, ICoreProvider& InDataStorage, TSharedPtr<INode> InRootNode)
+		: FExecutorBase(InStackName, MoveTemp(InRootNode))
+		, DataStorage(&InDataStorage)
+	{
+		if (RootNode)
+		{
+			Delegate = DataStorage->OnUpdate().AddRaw(this, &FOnUpdateExecutor::OnUpdate);
+		}
+	}
+
+	FOnUpdateExecutor::~FOnUpdateExecutor()
+	{
+		if (DataStorage)
+		{
+			DataStorage->OnUpdate().Remove(Delegate);
+		}
+		Delegate.Reset();
+	}
+
+	FOnUpdateExecutor::FOnUpdateExecutor(FOnUpdateExecutor&& Rhs)
+		: FExecutorBase(Rhs.StackName, MoveTemp(Rhs.RootNode))
+		, DataStorage(std::exchange(Rhs.DataStorage, nullptr))
+	{
+		if (RootNode && DataStorage)
+		{
+			DataStorage->OnUpdate().Remove(Rhs.Delegate);
+			Rhs.Delegate.Reset();
+
+			Delegate = DataStorage->OnUpdate().AddRaw(this, &FOnUpdateExecutor::OnUpdate);
+		}
+	}
+
+	FOnUpdateExecutor& FOnUpdateExecutor::operator=(FOnUpdateExecutor&& Rhs)
+	{
+		if (this != &Rhs)
+		{
+			this->~FOnUpdateExecutor();
+			new(this) FOnUpdateExecutor(MoveTemp(Rhs));
+		}
+		return *this;
+	}
+
+	void FOnUpdateExecutor::OnUpdate()
+	{
+		FullUpdate();
+	}
+
+
+
+	//
+	// FOnUpdateCompletedExecutor
+	//
+
+	FOnUpdateCompletedExecutor::FOnUpdateCompletedExecutor(const FName& InStackName, ICoreProvider& InDataStorage, TSharedPtr<INode> InRootNode)
+		: FExecutorBase(InStackName, MoveTemp(InRootNode))
+		, DataStorage(&InDataStorage)
+	{
+		if (RootNode)
+		{
+			Delegate = DataStorage->OnUpdateCompleted().AddRaw(this, &FOnUpdateCompletedExecutor::OnUpdateCompleted);
+		}
+	}
+
+	FOnUpdateCompletedExecutor::~FOnUpdateCompletedExecutor()
+	{
+		if (DataStorage)
+		{
+			DataStorage->OnUpdateCompleted().Remove(Delegate);
+		}
+		Delegate.Reset();
+	}
+
+	FOnUpdateCompletedExecutor::FOnUpdateCompletedExecutor(FOnUpdateCompletedExecutor&& Rhs)
+		: FExecutorBase(Rhs.StackName, MoveTemp(Rhs.RootNode))
+		, DataStorage(std::exchange(Rhs.DataStorage, nullptr))
+	{
+		if (RootNode && DataStorage)
+		{
+			DataStorage->OnUpdateCompleted().Remove(Rhs.Delegate);
+			Rhs.Delegate.Reset();
+
+			Delegate = DataStorage->OnUpdateCompleted().AddRaw(this, &FOnUpdateCompletedExecutor::OnUpdateCompleted);
+		}
+	}
+
+	FOnUpdateCompletedExecutor& FOnUpdateCompletedExecutor::operator=(FOnUpdateCompletedExecutor&& Rhs)
+	{
+		if (this != &Rhs)
+		{
+			this->~FOnUpdateCompletedExecutor();
+			new(this) FOnUpdateCompletedExecutor(MoveTemp(Rhs));
+		}
+		return *this;
+	}
+
+	void FOnUpdateCompletedExecutor::OnUpdateCompleted()
+	{
+		FullUpdate();
+	}
+
+
+
+	//
+	// FExplicitCooperativeExecutor
+	//
+
+	FExplicitCooperativeExecutor::FExplicitCooperativeExecutor(const FName& StackName, TSharedPtr<INode> RootNode)
+		: FExecutorBase(StackName, MoveTemp(RootNode))
+	{
+	}
+
+	void FExplicitCooperativeExecutor::Update(FTimespan AllottedTime)
+	{
+		if (RootNode)
+		{
+			TimedUpdate(AllottedTime);
+		}
+	}
+
+
+
+	//
+	// FExplicitUpdateExecutor
+	//
+
+	FExplicitUpdateExecutor::FExplicitUpdateExecutor(const FName& StackName, TSharedPtr<INode> RootNode)
+		: FExecutorBase(StackName, MoveTemp(RootNode))
+	{
+	}
+
+	void FExplicitUpdateExecutor::Update()
+	{
+		if (RootNode)
+		{
+			FullUpdate();
+		}
+	}
+} // namespace UE::Editor::DataStorage::QueryStack
\ No newline at end of file
diff --git a/Engine/Plugins/Experimental/EditorDataStorageFeatures/Source/TedsQueryStack/Private/TedsRowArrayNode.cpp b/Engine/Plugins/Experimental/EditorDataStorageFeatures/Source/TedsQueryStack/Private/TedsRowArrayNode.cpp
--- a/Engine/Plugins/Experimental/EditorDataStorageFeatures/Source/TedsQueryStack/Private/TedsRowArrayNode.cpp
+++ b/Engine/Plugins/Experimental/EditorDataStorageFeatures/Source/TedsQueryStack/Private/TedsRowArrayNode.cpp
@@ -19,7 +19,11 @@ namespace UE::Editor::DataStorage::QueryStack
 		return Revision;
 	}
 
-	void FRowArrayNode::Update()
+	void FRowArrayNode::Update(FTimespan AllottedTime)
+	{
+	}
+
+	void FRowArrayNode::VisitParents(ParentListCallback Callback)
 	{
 	}
 
diff --git a/Engine/Plugins/Experimental/EditorDataStorageFeatures/Source/TedsQueryStack/Private/TedsRowCopyNode.cpp b/Engine/Plugins/Experimental/EditorDataStorageFeatures/Source/TedsQueryStack/Private/TedsRowCopyNode.cpp
--- a/Engine/Plugins/Experimental/EditorDataStorageFeatures/Source/TedsQueryStack/Private/TedsRowCopyNode.cpp
+++ b/Engine/Plugins/Experimental/EditorDataStorageFeatures/Source/TedsQueryStack/Private/TedsRowCopyNode.cpp
@@ -22,10 +22,8 @@ namespace UE::Editor::DataStorage::QueryStack
 		return Revision;
 	}
 
-	void FRowCopyNode::Update()
+	void FRowCopyNode::Update(FTimespan AllottedTime)
 	{
-		Parent->Update();
-
 		if (Parent->GetRevision() != ParentRevision)
 		{
 			Reset();
@@ -33,6 +31,11 @@ namespace UE::Editor::DataStorage::QueryStack
 		}
 	}
 
+	void FRowCopyNode::VisitParents(ParentListCallback Callback)
+	{
+		Callback(Parent);
+	}
+
 	FRowHandleArrayView FRowCopyNode::GetRows() const
 	{
 		return Rows.GetRows();
diff --git a/Engine/Plugins/Experimental/EditorDataStorageFeatures/Source/TedsQueryStack/Private/TedsRowFilterNode.cpp b/Engine/Plugins/Experimental/EditorDataStorageFeatures/Source/TedsQueryStack/Private/TedsRowFilterNode.cpp
--- a/Engine/Plugins/Experimental/EditorDataStorageFeatures/Source/TedsQueryStack/Private/TedsRowFilterNode.cpp
+++ b/Engine/Plugins/Experimental/EditorDataStorageFeatures/Source/TedsQueryStack/Private/TedsRowFilterNode.cpp
@@ -19,9 +19,8 @@ namespace UE::Editor::DataStorage::QueryStack
 		return CachedParentRevisionID;
 	}
 
-	void FRowFilterNode::Update()
+	void FRowFilterNode::Update(FTimespan AllottedTime)
 	{
-		ParentRowNode->Update();
 		if (CachedParentRevisionID != ParentRowNode->GetRevision())
 		{
 			Rows.Reset();
@@ -31,6 +30,11 @@ namespace UE::Editor::DataStorage::QueryStack
 		}
 	}
 
+	void FRowFilterNode::VisitParents(ParentListCallback Callback)
+	{
+		Callback(ParentRowNode);
+	}
+
 	void FRowFilterNode::ForceRefresh()
 	{
 		CachedParentRevisionID = 0;
diff --git a/Engine/Plugins/Experimental/EditorDataStorageFeatures/Source/TedsQueryStack/Private/TedsRowHandleSortNode.cpp b/Engine/Plugins/Experimental/EditorDataStorageFeatures/Source/TedsQueryStack/Private/TedsRowHandleSortNode.cpp
--- a/Engine/Plugins/Experimental/EditorDataStorageFeatures/Source/TedsQueryStack/Private/TedsRowHandleSortNode.cpp
+++ b/Engine/Plugins/Experimental/EditorDataStorageFeatures/Source/TedsQueryStack/Private/TedsRowHandleSortNode.cpp
@@ -17,10 +17,8 @@ namespace UE::Editor::DataStorage::QueryStack
 		return ParentRevision;
 	}
 
-	void FRowHandleSortNode::Update()
+	void FRowHandleSortNode::Update(FTimespan AllottedTime)
 	{
-		Parent->Update();
-
 		if (Parent->GetRevision() != ParentRevision)
 		{
 			Parent->GetMutableRows().Sort();
@@ -28,6 +26,11 @@ namespace UE::Editor::DataStorage::QueryStack
 		}
 	}
 
+	void FRowHandleSortNode::VisitParents(ParentListCallback Callback)
+	{
+		Callback(Parent);
+	}
+
 	FRowHandleArrayView FRowHandleSortNode::GetRows() const
 	{
 		return Parent->GetRows();
diff --git a/Engine/Plugins/Experimental/EditorDataStorageFeatures/Source/TedsQueryStack/Private/TedsRowMergeNode.cpp b/Engine/Plugins/Experimental/EditorDataStorageFeatures/Source/TedsQueryStack/Private/TedsRowMergeNode.cpp
--- a/Engine/Plugins/Experimental/EditorDataStorageFeatures/Source/TedsQueryStack/Private/TedsRowMergeNode.cpp
+++ b/Engine/Plugins/Experimental/EditorDataStorageFeatures/Source/TedsQueryStack/Private/TedsRowMergeNode.cpp
@@ -20,12 +20,11 @@ namespace UE::Editor::DataStorage::QueryStack
 		return Revision;
 	}
 
-	void FRowMergeNode::Update()
+	void FRowMergeNode::Update(FTimespan AllottedTime)
 	{
 		bool bRebuild = false;
 		for (FParentInfo& ParentInfo : Parents)
 		{
-			ParentInfo.Parent->Update();
 			RevisionId ParentRevision = ParentInfo.Parent->GetRevision();
 			if (ParentInfo.Revision != ParentRevision)
 			{
@@ -41,6 +40,14 @@ namespace UE::Editor::DataStorage::QueryStack
 		}
 	}
 
+	void FRowMergeNode::VisitParents(ParentListCallback Callback)
+	{
+		for (const FParentInfo& ParentInfo : Parents)
+		{
+			Callback(ParentInfo.Parent);
+		}
+	}
+
 	FRowHandleArrayView FRowMergeNode::GetRows() const
 	{
 		return Rows.GetRows();
diff --git a/Engine/Plugins/Experimental/EditorDataStorageFeatures/Source/TedsQueryStack/Private/TedsRowMonitorNode.cpp b/Engine/Plugins/Experimental/EditorDataStorageFeatures/Source/TedsQueryStack/Private/TedsRowMonitorNode.cpp
--- a/Engine/Plugins/Experimental/EditorDataStorageFeatures/Source/TedsQueryStack/Private/TedsRowMonitorNode.cpp
+++ b/Engine/Plugins/Experimental/EditorDataStorageFeatures/Source/TedsQueryStack/Private/TedsRowMonitorNode.cpp
@@ -78,29 +78,21 @@ namespace UE::Editor::DataStorage::QueryStack
 		return Revision;
 	}
 
-	void FRowMonitorNode::Update()
+	void FRowMonitorNode::Update(FTimespan AllottedTime)
 	{
-		if (QueryNode)
+		if (QueryNode && QueryNode->GetRevision() != QueryRevision)
 		{
-			QueryNode->Update();
-			if (QueryNode->GetRevision() != QueryRevision)
-			{
-				UpdateColumnsFromQuery();
-				UpdateMonitoredColumns();
-				Revision++;
-			}
+			UpdateColumnsFromQuery();
+			UpdateMonitoredColumns();
+			Revision++;
 		}
 		
 		bool MergeChanges = true;
-		if (ParentNode)
+		if (ParentNode && ParentNode->GetRevision() != ParentRevision)
 		{
-			ParentNode->Update();
-			if (ParentNode->GetRevision() != ParentRevision)
-			{
-				ParentRevision = ParentNode->GetRevision();
-				Revision++;
-				MergeChanges = false;
-			}
+			ParentRevision = ParentNode->GetRevision();
+			Revision++;
+			MergeChanges = false;
 		}
 		
 		if (MergeChanges && (!AddedRows.IsEmpty() || !RemovedRows.IsEmpty()))
@@ -113,9 +105,12 @@ namespace UE::Editor::DataStorage::QueryStack
 
 			ResolveRemovedRows();
 
-			FRowHandleArray& TargetRows = ParentNode->GetMutableRows();
-			TargetRows.SortedMerge(AddedRows);
-			TargetRows.Remove(RemovedRows.GetRows());
+			if (ParentNode)
+			{
+				FRowHandleArray& TargetRows = ParentNode->GetMutableRows();
+				TargetRows.SortedMerge(AddedRows);
+				TargetRows.Remove(RemovedRows.GetRows());
+			}
 
 			if (!AddedRows.IsEmpty())
 			{
@@ -134,6 +129,18 @@ namespace UE::Editor::DataStorage::QueryStack
 		}
 	}
 
+	void FRowMonitorNode::VisitParents(ParentListCallback Callback)
+	{
+		if (QueryNode)
+		{
+			Callback(QueryNode);
+		}
+		if (ParentNode)
+		{
+			Callback(ParentNode);
+		}
+	}
+
 	FRowHandleArrayView FRowMonitorNode::GetRows() const
 	{
 		return ParentNode->GetRows();
@@ -223,6 +230,10 @@ namespace UE::Editor::DataStorage::QueryStack
 				QueryRevision = 0;
 			}
 		}
+		else
+		{
+			QueryRevision = QueryNode->GetRevision();
+		}
 	}
 
 	void FRowMonitorNode::UpdateMonitoredColumns()
@@ -307,36 +318,39 @@ namespace UE::Editor::DataStorage::QueryStack
 
 	void FRowMonitorNode::ResolveRemovedRows()
 	{
-		// OnRemove observers are fired when any column from our conditions are removed from the row, however this doesn't necessarily mean
-		// that the row should be removed as it could still be matching the query. For instance, if the Conditions are (ColumnA || ColumnB), the row
-		// contains both columns and only ColumnB is removed the row still matches the original query and shouldn't be removed from the parent node.
-		const FQueryDescription& Query = Storage.GetQueryDescription(QueryNode->GetQuery());
-
-		// RemovedRows is guaranteed to be sorted when this function is called, but we'll add a sanity check to sort it anyways in case that changes
-		if (!RemovedRows.IsSorted())
+		if (QueryNode)
 		{
-			RemovedRows.Sort();
-		}
+			// OnRemove observers are fired when any column from our conditions are removed from the row, however this doesn't necessarily mean
+			// that the row should be removed as it could still be matching the query. For instance, if the Conditions are (ColumnA || ColumnB), the row
+			// contains both columns and only ColumnB is removed the row still matches the original query and shouldn't be removed from the parent node.
+			const FQueryDescription& Query = Storage.GetQueryDescription(QueryNode->GetQuery());
 
-		// We don't currently have a way to detect if a row matches a query using the old syntax, so this logic only applies to the FConditions
-		// based syntax. Which means that you could get an incorrect output if you use this node with a complex .Where() syntax.
-		if (const Queries::FConditions* Conditions = Query.Conditions.GetPtrOrNull())
-		{
-			FRowHandleArray RowsToNotRemove;
-			RowsToNotRemove.Reserve(RemovedRows.Num());
-			
-			for (RowHandle Row : RemovedRows.GetRows())
+			// RemovedRows is guaranteed to be sorted when this function is called, but we'll add a sanity check to sort it anyways in case that changes
+			if (!RemovedRows.IsSorted())
 			{
-				// If the row still matches the conditions, don't remove it.
-				if (Storage.MatchesColumns(Row, *Conditions))
+				RemovedRows.Sort();
+			}
+
+			// We don't currently have a way to detect if a row matches a query using the old syntax, so this logic only applies to the FConditions
+			// based syntax. Which means that you could get an incorrect output if you use this node with a complex .Where() syntax.
+			if (const Queries::FConditions* Conditions = Query.Conditions.GetPtrOrNull())
+			{
+				FRowHandleArray RowsToNotRemove;
+				RowsToNotRemove.Reserve(RemovedRows.Num());
+
+				for (RowHandle Row : RemovedRows.GetRows())
 				{
-					// Since RemovedRows is sorted, iterating through it and adding rows keep that property
-					static constexpr FRowHandleArrayView::EFlags Flags(FRowHandleArrayView::EFlags::IsSorted | FRowHandleArrayView::EFlags::IsUnique );
-					RowsToNotRemove.Add(Row, Flags);
+					// If the row still matches the conditions, don't remove it.
+					if (Storage.MatchesColumns(Row, *Conditions))
+					{
+						// Since RemovedRows is sorted, iterating through it and adding rows keep that property
+						static constexpr FRowHandleArrayView::EFlags Flags(FRowHandleArrayView::EFlags::IsSorted | FRowHandleArrayView::EFlags::IsUnique);
+						RowsToNotRemove.Add(Row, Flags);
+					}
 				}
+
+				RemovedRows.Remove(RowsToNotRemove.GetRows());
 			}
-			
-			RemovedRows.Remove(RowsToNotRemove.GetRows());
 		}
 	}
 } // namespace UE::Editor::DataStorage::QueryStack
diff --git a/Engine/Plugins/Experimental/EditorDataStorageFeatures/Source/TedsQueryStack/Private/TedsRowOrderInversionNode.cpp b/Engine/Plugins/Experimental/EditorDataStorageFeatures/Source/TedsQueryStack/Private/TedsRowOrderInversionNode.cpp
--- a/Engine/Plugins/Experimental/EditorDataStorageFeatures/Source/TedsQueryStack/Private/TedsRowOrderInversionNode.cpp
+++ b/Engine/Plugins/Experimental/EditorDataStorageFeatures/Source/TedsQueryStack/Private/TedsRowOrderInversionNode.cpp
@@ -33,10 +33,8 @@ namespace UE::Editor::DataStorage::QueryStack
 		return Revision;
 	}
 
-	void FRowOrderInversionNode::Update()
+	void FRowOrderInversionNode::Update(FTimespan AllottedTime)
 	{
-		Parent->Update();
-
 		RevisionId LocalParentRevision = Parent->GetRevision();
 		if (bIsEnabled)
 		{
@@ -59,6 +57,11 @@ namespace UE::Editor::DataStorage::QueryStack
 		}
 	}
 
+	void FRowOrderInversionNode::VisitParents(ParentListCallback Callback)
+	{
+		Callback(Parent);
+	}
+
 	FRowHandleArrayView FRowOrderInversionNode::GetRows() const
 	{
 		return Parent->GetRows();
diff --git a/Engine/Plugins/Experimental/EditorDataStorageFeatures/Source/TedsQueryStack/Private/TedsRowQueryResultsNode.cpp b/Engine/Plugins/Experimental/EditorDataStorageFeatures/Source/TedsQueryStack/Private/TedsRowQueryResultsNode.cpp
--- a/Engine/Plugins/Experimental/EditorDataStorageFeatures/Source/TedsQueryStack/Private/TedsRowQueryResultsNode.cpp
+++ b/Engine/Plugins/Experimental/EditorDataStorageFeatures/Source/TedsQueryStack/Private/TedsRowQueryResultsNode.cpp
@@ -52,10 +52,8 @@ namespace UE::Editor::DataStorage::QueryStack
 		}
 	}
 
-	void FRowQueryResultsNode::Update()
+	void FRowQueryResultsNode::Update(FTimespan AllottedTime)
 	{
-		QueryNode->Update();
-
 		if (EnumHasAllFlags(SyncFlags, ESyncFlags::RefreshOnUpdate) ||
 			(EnumHasAllFlags(SyncFlags, ESyncFlags::RefreshOnQueryChange) && QueryRevision != QueryNode->GetRevision()))
 		{
@@ -64,6 +62,11 @@ namespace UE::Editor::DataStorage::QueryStack
 		}
 	}
 
+	void FRowQueryResultsNode::VisitParents(ParentListCallback Callback)
+	{
+		Callback(QueryNode);
+	}
+
 	FRowHandleArrayView FRowQueryResultsNode::GetRows() const
 	{
 		return Rows.GetRows();
diff --git a/Engine/Plugins/Experimental/EditorDataStorageFeatures/Source/TedsQueryStack/Private/TedsRowSortNode.cpp b/Engine/Plugins/Experimental/EditorDataStorageFeatures/Source/TedsQueryStack/Private/TedsRowSortNode.cpp
--- a/Engine/Plugins/Experimental/EditorDataStorageFeatures/Source/TedsQueryStack/Private/TedsRowSortNode.cpp
+++ b/Engine/Plugins/Experimental/EditorDataStorageFeatures/Source/TedsQueryStack/Private/TedsRowSortNode.cpp
@@ -67,13 +67,12 @@ namespace UE::Editor::DataStorage::QueryStack
 		return Revision;
 	}
 
-	void FRowSortNode::Update()
+	void FRowSortNode::Update(FTimespan AllottedTime)
 	{
 		TRACE_CPUPROFILER_EVENT_SCOPE("[TEDS Sorting] Sort node update");
 
 		uint64 StartTime = FPlatformTime::Cycles64();
 
-		Parent->Update();
 		if (Parent->GetRevision() != LastUpdatedRevision)
 		{
 			LastUpdatedRevision = Parent->GetRevision();
@@ -112,7 +111,11 @@ namespace UE::Editor::DataStorage::QueryStack
 		{
 			if (SortContext->Tasks.HasRemainingTasks())
 			{
-				SortContext->Tasks.Update(MaxFrameDuration);
+				// If the update is uncapped then limit to the requested upper limit as otherwise the editor can lock up
+				// for a long time. If there's explicit time assigned, use as much of it as needed, even if it exceeds
+				// the maximum duration that's specified because TEDS is telling this function that's the amount of
+				// time that's available.
+				SortContext->Tasks.Update(AllottedTime == INode::UnlimitedTime ? MaxFrameDuration : AllottedTime);
 			}
 
 			// Check remaining tasks again as the update above could have added new tasks.
@@ -125,6 +128,11 @@ namespace UE::Editor::DataStorage::QueryStack
 		}
 	}
 
+	void FRowSortNode::VisitParents(ParentListCallback Callback)
+	{
+		Callback(Parent);
+	}
+
 	void FRowSortNode::SetupFixedSize64()
 	{
 		FRowHandleArray& MutableRows = Parent->GetMutableRows();
diff --git a/Engine/Plugins/Experimental/EditorDataStorageFeatures/Source/TedsQueryStack/Private/TedsRowViewNode.cpp b/Engine/Plugins/Experimental/EditorDataStorageFeatures/Source/TedsQueryStack/Private/TedsRowViewNode.cpp
--- a/Engine/Plugins/Experimental/EditorDataStorageFeatures/Source/TedsQueryStack/Private/TedsRowViewNode.cpp
+++ b/Engine/Plugins/Experimental/EditorDataStorageFeatures/Source/TedsQueryStack/Private/TedsRowViewNode.cpp
@@ -24,10 +24,15 @@ namespace UE::Editor::DataStorage::QueryStack
 		return Revision;
 	}
 
-	void FRowViewNode::Update() 
+	void FRowViewNode::Update(FTimespan AllottedTime)
 	{
 	}
 
+	void FRowViewNode::VisitParents(ParentListCallback Callback)
+	{
+		// No parents to visit.
+	}
+
 	FRowHandleArrayView FRowViewNode::GetRows() const
 	{
 		return Rows->GetRows();
diff --git a/Engine/Plugins/Experimental/EditorDataStorageFeatures/Source/TedsQueryStack/Public/TedsHierarchyNode.h b/Engine/Plugins/Experimental/EditorDataStorageFeatures/Source/TedsQueryStack/Public/TedsHierarchyNode.h
--- a/Engine/Plugins/Experimental/EditorDataStorageFeatures/Source/TedsQueryStack/Public/TedsHierarchyNode.h
+++ b/Engine/Plugins/Experimental/EditorDataStorageFeatures/Source/TedsQueryStack/Public/TedsHierarchyNode.h
@@ -35,10 +35,13 @@ namespace UE::Editor::DataStorage
 			ESyncFlags SyncFlags);
 		
 		TEDSQUERYSTACK_API virtual RevisionId GetRevision() const override;
-		TEDSQUERYSTACK_API virtual void Update() override;
+		TEDSQUERYSTACK_API virtual void VisitParents(ParentListCallback Callback) override;
 		TEDSQUERYSTACK_API virtual FRowHandleArrayView GetRows() const override;
 		TEDSQUERYSTACK_API virtual FRowHandleArray& GetMutableRows() override;
 
+	protected:
+		TEDSQUERYSTACK_API virtual void Update(FTimespan AllottedTime) override;
+
 	private:
 		void RefreshInternal(ICoreProvider& InCoreProvider, FRowHandleArray& NewRows) const;
 		FRowHandleArray& GetTempRowHandleArray();
diff --git a/Engine/Plugins/Experimental/EditorDataStorageFeatures/Source/TedsQueryStack/Public/TedsQueryHandleNode.h b/Engine/Plugins/Experimental/EditorDataStorageFeatures/Source/TedsQueryStack/Public/TedsQueryHandleNode.h
--- a/Engine/Plugins/Experimental/EditorDataStorageFeatures/Source/TedsQueryStack/Public/TedsQueryHandleNode.h
+++ b/Engine/Plugins/Experimental/EditorDataStorageFeatures/Source/TedsQueryStack/Public/TedsQueryHandleNode.h
@@ -19,9 +19,12 @@ namespace UE::Editor::DataStorage::QueryStack
 		TEDSQUERYSTACK_API void SetQuery(QueryHandle InQuery);
 
 		TEDSQUERYSTACK_API virtual RevisionId GetRevision() const override;
-		TEDSQUERYSTACK_API virtual void Update() override;
+		TEDSQUERYSTACK_API virtual void VisitParents(ParentListCallback Callback) override;
 		TEDSQUERYSTACK_API virtual QueryHandle GetQuery() const override;
 
+	protected:
+		TEDSQUERYSTACK_API virtual void Update(FTimespan AllottedTime) override;
+
 	private:
 		QueryHandle Query = InvalidQueryHandle;
 		RevisionId Revision = 0;
diff --git a/Engine/Plugins/Experimental/EditorDataStorageFeatures/Source/TedsQueryStack/Public/TedsQueryMergeNode.h b/Engine/Plugins/Experimental/EditorDataStorageFeatures/Source/TedsQueryStack/Public/TedsQueryMergeNode.h
--- a/Engine/Plugins/Experimental/EditorDataStorageFeatures/Source/TedsQueryStack/Public/TedsQueryMergeNode.h
+++ b/Engine/Plugins/Experimental/EditorDataStorageFeatures/Source/TedsQueryStack/Public/TedsQueryMergeNode.h
@@ -23,10 +23,11 @@ namespace UE::Editor::DataStorage::QueryStack
 		TEDSQUERYSTACK_API virtual ~FQueryMergeNode() override;
 
 		TEDSQUERYSTACK_API virtual RevisionId GetRevision() const override;
-		TEDSQUERYSTACK_API virtual void Update() override;
+		TEDSQUERYSTACK_API virtual void VisitParents(ParentListCallback Callback) override;
 		TEDSQUERYSTACK_API virtual QueryHandle GetQuery() const override;
 
 	protected:
+		TEDSQUERYSTACK_API virtual void Update(FTimespan AllottedTime) override;
 		void Rebuild();
 		
 	private:
diff --git a/Engine/Plugins/Experimental/EditorDataStorageFeatures/Source/TedsQueryStack/Public/TedsQueryNode.h b/Engine/Plugins/Experimental/EditorDataStorageFeatures/Source/TedsQueryStack/Public/TedsQueryNode.h
--- a/Engine/Plugins/Experimental/EditorDataStorageFeatures/Source/TedsQueryStack/Public/TedsQueryNode.h
+++ b/Engine/Plugins/Experimental/EditorDataStorageFeatures/Source/TedsQueryStack/Public/TedsQueryNode.h
@@ -26,9 +26,12 @@ namespace UE::Editor::DataStorage::QueryStack
 		TEDSQUERYSTACK_API void ClearQuery();
 
 		TEDSQUERYSTACK_API virtual RevisionId GetRevision() const override;
-		TEDSQUERYSTACK_API virtual void Update() override;
+		TEDSQUERYSTACK_API virtual void VisitParents(ParentListCallback Callback) override;
 		TEDSQUERYSTACK_API virtual QueryHandle GetQuery() const override;
 
+	protected:
+		TEDSQUERYSTACK_API virtual void Update(FTimespan AllottedTime) override;
+
 	private:
 		QueryHandle QueryHandle = InvalidQueryHandle;
 		ICoreProvider& Storage;
diff --git a/Engine/Plugins/Experimental/EditorDataStorageFeatures/Source/TedsQueryStack/Public/TedsQueryStackExecutor.h b/Engine/Plugins/Experimental/EditorDataStorageFeatures/Source/TedsQueryStack/Public/TedsQueryStackExecutor.h
new file mode 100644
--- /dev/null
+++ b/Engine/Plugins/Experimental/EditorDataStorageFeatures/Source/TedsQueryStack/Public/TedsQueryStackExecutor.h
@@ -0,0 +1,153 @@
+// Copyright Epic Games, Inc. All Rights Reserved.
+
+#pragma once
+
+#include "Elements/Interfaces/TypedElementDataStorageInterface.h"
+#include "Misc/Timespan.h"
+#include "Misc/TVariant.h"
+#include "Templates/SharedPointer.h"
+
+namespace UE::Editor::DataStorage::QueryStack
+{
+	class INode;
+
+	class FExecutorBase
+	{
+	public:
+		using OnChangedFunction = TFunction<void(const FName&)>;
+
+		TEDSQUERYSTACK_API virtual ~FExecutorBase() = 0;
+
+		TEDSQUERYSTACK_API const FName& GetStackName() const;
+		TEDSQUERYSTACK_API TSharedPtr<const INode> GetRootNode() const;
+		TEDSQUERYSTACK_API TSharedPtr<INode> GetRootNode();
+
+		/** If set, the provided function will be called when the a change in the query stack was reported. */
+		TEDSQUERYSTACK_API void SetOnChangedCallback(OnChangedFunction Callback);
+		TEDSQUERYSTACK_API void ClearOnChangedCallback();
+
+	// Not exporting these protected functions because this is not intended to be globally extended.
+	protected:
+		FExecutorBase() = default;
+		FExecutorBase(const FName& StackName, TSharedPtr<INode> RootNode);
+
+		// Delete copying to keep the guarantee that the name is unique and remove the risk
+		// of the same query stack being updated multiple times.
+		FExecutorBase(const FExecutorBase&) = delete;
+		FExecutorBase& operator=(const FExecutorBase&) = delete;
+
+		FExecutorBase(FExecutorBase&&) = default;
+		FExecutorBase& operator=(FExecutorBase&&) = default;
+
+		void FullUpdate();
+		void TimedUpdate(FTimespan AllottedTime);
+		
+		OnChangedFunction OnChangedCallback;
+		TSharedPtr<INode> RootNode;
+		FName StackName;
+
+	private:
+		static void FullNodeUpdate(const TSharedPtr<INode>& Node);
+		static void TimedNodeUpdate(const TSharedPtr<INode>& Node, uint64 StartTime, FTimespan AllottedTime);
+	};
+
+	/**
+	 * Updates the query stack cooperatively with other background tasks.
+	 * Preferred default as it reduces overall cost to the editor by allowing TEDS to fairly distribute available CPU cycles.
+	 * The cooperative executor may stop partially through updating the stack and continue at a later time.
+	 */
+	class FCooperativeExecutor final : public FExecutorBase
+	{
+	public:
+		FCooperativeExecutor() = default;
+		TEDSQUERYSTACK_API FCooperativeExecutor(const FName& InStackName, ICoreProvider& InDataStorage, TSharedPtr<INode> InRootNode,
+			ICoreProvider::ECooperativeTaskPriority InPriority);
+
+		TEDSQUERYSTACK_API FCooperativeExecutor(FCooperativeExecutor&& Rhs);
+		TEDSQUERYSTACK_API FCooperativeExecutor& operator=(FCooperativeExecutor&& Rhs);
+
+		TEDSQUERYSTACK_API virtual ~FCooperativeExecutor() override;
+
+	private:
+		ICoreProvider* DataStorage = nullptr;
+		ICoreProvider::ECooperativeTaskPriority Priority;
+	};
+
+	/**
+	 * Automatically updates the query stack just before TEDS starts running queries.
+	 */
+	class FOnUpdateExecutor final : public FExecutorBase
+	{
+	public:
+		FOnUpdateExecutor() = default;
+		TEDSQUERYSTACK_API FOnUpdateExecutor(const FName& InStackName, ICoreProvider& InDataStorage, TSharedPtr<INode> InRootNode);
+		
+		TEDSQUERYSTACK_API FOnUpdateExecutor(FOnUpdateExecutor&& Rhs);
+		TEDSQUERYSTACK_API FOnUpdateExecutor& operator=(FOnUpdateExecutor&& Rhs);
+		
+		TEDSQUERYSTACK_API virtual ~FOnUpdateExecutor() override;
+
+	private:
+		void OnUpdate();
+
+		FDelegateHandle Delegate;
+		ICoreProvider* DataStorage = nullptr;
+	};
+
+	/**
+	 * Automatically updates the query stack at the very end of TEDS' update cycle.
+	 */
+	class FOnUpdateCompletedExecutor final : public FExecutorBase
+	{
+	public:
+		FOnUpdateCompletedExecutor() = default;
+		TEDSQUERYSTACK_API FOnUpdateCompletedExecutor(const FName& InStackName, ICoreProvider& InDataStorage, TSharedPtr<INode> InRootNode);
+
+		TEDSQUERYSTACK_API FOnUpdateCompletedExecutor(FOnUpdateCompletedExecutor&& Rhs);
+		TEDSQUERYSTACK_API FOnUpdateCompletedExecutor& operator=(FOnUpdateCompletedExecutor&& Rhs);
+
+		TEDSQUERYSTACK_API virtual ~FOnUpdateCompletedExecutor() override;
+
+	private:
+		void OnUpdateCompleted();
+
+		FDelegateHandle Delegate;
+		ICoreProvider* DataStorage = nullptr;
+	};
+
+	/**
+	 * Allows explicit calls to update the query stack as part of a cooperative task.
+	 * Does not update the provided query stack automatically but depends on calling Update.
+	 */
+	class FExplicitCooperativeExecutor final : public FExecutorBase
+	{
+	public:
+		FExplicitCooperativeExecutor() = default;
+		TEDSQUERYSTACK_API FExplicitCooperativeExecutor(const FName& StackName, TSharedPtr<INode> RootNode);
+
+		FExplicitCooperativeExecutor(FExplicitCooperativeExecutor&&) = default;
+		FExplicitCooperativeExecutor& operator=(FExplicitCooperativeExecutor&&) = default;
+
+		virtual ~FExplicitCooperativeExecutor() override = default;
+
+		TEDSQUERYSTACK_API void Update(FTimespan AllottedTime);
+	};
+
+	/**
+	 * Allows explicit calls to update the query stack.
+	 * Does not update the provided query stack automatically but depends on calling Update.
+	 */
+	class FExplicitUpdateExecutor final : public FExecutorBase
+	{
+	public:
+		FExplicitUpdateExecutor() = default;
+		TEDSQUERYSTACK_API FExplicitUpdateExecutor(const FName& StackName, TSharedPtr<INode> RootNode);
+		
+		FExplicitUpdateExecutor(FExplicitUpdateExecutor&&) = default;
+		FExplicitUpdateExecutor& operator=(FExplicitUpdateExecutor&&) = default;
+
+		virtual ~FExplicitUpdateExecutor() override = default;
+
+		TEDSQUERYSTACK_API void Update();
+	};
+} // namespace UE::Editor::DataStorage::QueryStack
\ No newline at end of file
diff --git a/Engine/Plugins/Experimental/EditorDataStorageFeatures/Source/TedsQueryStack/Public/TedsQueryStackInterfaces.h b/Engine/Plugins/Experimental/EditorDataStorageFeatures/Source/TedsQueryStack/Public/TedsQueryStackInterfaces.h
--- a/Engine/Plugins/Experimental/EditorDataStorageFeatures/Source/TedsQueryStack/Public/TedsQueryStackInterfaces.h
+++ b/Engine/Plugins/Experimental/EditorDataStorageFeatures/Source/TedsQueryStack/Public/TedsQueryStackInterfaces.h
@@ -6,6 +6,8 @@
 #include "Containers/ContainersFwd.h"
 #include "Elements/Framework/TypedElementRowHandleArray.h"
 #include "Elements/Framework/TypedElementRowHandleArrayView.h"
+#include "Misc/Timespan.h"
+#include "Templates/FunctionFwd.h"
 
 namespace UE::Editor::DataStorage::QueryStack
 {
@@ -13,12 +15,31 @@ namespace UE::Editor::DataStorage::QueryStack
 	class INode
 	{
 	public:
+		friend class FExecutorBase;
+
 		using RevisionId = uint32;
+		using ParentListCallback = TFunctionRef<void(const TSharedPtr<INode>& Parent)>;
+		static inline FTimespan UnlimitedTime = FTimespan::MaxValue();
 
 		virtual ~INode() = default;
 		
+		/**
+		 * Get the current revision of the node. Whenever its internal state changes the revision should be incremented to signal a need 
+		 * to update to downstream nodes.
+		 */
 		virtual RevisionId GetRevision() const = 0;
-		virtual void Update() = 0;
+		
+		/** List all direct parent nodes. */
+		virtual void VisitParents(ParentListCallback Callback) = 0;
+
+	protected:
+		/**
+		 * Periodically called to allow a node to perform some work or to verify its internal state.
+		 * @AllottedTime The maximum amount of time for the update. The update should try to stay below this time. In case the
+		 *		allotted time is equal to "UnlimitedTime" it means there is no upper bound to the amount of time that can be spend
+		 *		on the update.
+		 */
+		virtual void Update(FTimespan AllottedTime) = 0;
 	};
 
 	/** Query stack node that works on queries handles. These nodes are typically run in some fashion to be turned into a row node. */
diff --git a/Engine/Plugins/Experimental/EditorDataStorageFeatures/Source/TedsQueryStack/Public/TedsRowArrayNode.h b/Engine/Plugins/Experimental/EditorDataStorageFeatures/Source/TedsQueryStack/Public/TedsRowArrayNode.h
--- a/Engine/Plugins/Experimental/EditorDataStorageFeatures/Source/TedsQueryStack/Public/TedsRowArrayNode.h
+++ b/Engine/Plugins/Experimental/EditorDataStorageFeatures/Source/TedsQueryStack/Public/TedsRowArrayNode.h
@@ -23,10 +23,13 @@ namespace UE::Editor::DataStorage::QueryStack
 		TEDSQUERYSTACK_API void MarkDirty();
 		
 		TEDSQUERYSTACK_API virtual RevisionId GetRevision() const override;
-		TEDSQUERYSTACK_API virtual void Update() override;
+		TEDSQUERYSTACK_API virtual void VisitParents(ParentListCallback Callback) override;
 		TEDSQUERYSTACK_API virtual FRowHandleArrayView GetRows() const override;
 		TEDSQUERYSTACK_API virtual FRowHandleArray& GetMutableRows() override;
 
+	protected:
+		TEDSQUERYSTACK_API virtual void Update(FTimespan AllottedTime) override;
+	
 	private:
 		FRowHandleArray Rows;
 		RevisionId Revision = 0;
diff --git a/Engine/Plugins/Experimental/EditorDataStorageFeatures/Source/TedsQueryStack/Public/TedsRowCopyNode.h b/Engine/Plugins/Experimental/EditorDataStorageFeatures/Source/TedsQueryStack/Public/TedsRowCopyNode.h
--- a/Engine/Plugins/Experimental/EditorDataStorageFeatures/Source/TedsQueryStack/Public/TedsRowCopyNode.h
+++ b/Engine/Plugins/Experimental/EditorDataStorageFeatures/Source/TedsQueryStack/Public/TedsRowCopyNode.h
@@ -24,10 +24,13 @@ namespace UE::Editor::DataStorage::QueryStack
 		TEDSQUERYSTACK_API void Reset();
 
 		TEDSQUERYSTACK_API virtual RevisionId GetRevision() const override;
-		TEDSQUERYSTACK_API virtual void Update() override;
+		TEDSQUERYSTACK_API virtual void VisitParents(ParentListCallback Callback) override;
 		TEDSQUERYSTACK_API virtual FRowHandleArrayView GetRows() const override;
 		TEDSQUERYSTACK_API virtual FRowHandleArray& GetMutableRows() override;
 
+	protected:
+		TEDSQUERYSTACK_API virtual void Update(FTimespan AllottedTime) override;
+	
 	private:
 		FRowHandleArray Rows;
 		TSharedPtr<IRowNode> Parent;
diff --git a/Engine/Plugins/Experimental/EditorDataStorageFeatures/Source/TedsQueryStack/Public/TedsRowFilterNode.h b/Engine/Plugins/Experimental/EditorDataStorageFeatures/Source/TedsQueryStack/Public/TedsRowFilterNode.h
--- a/Engine/Plugins/Experimental/EditorDataStorageFeatures/Source/TedsQueryStack/Public/TedsRowFilterNode.h
+++ b/Engine/Plugins/Experimental/EditorDataStorageFeatures/Source/TedsQueryStack/Public/TedsRowFilterNode.h
@@ -72,11 +72,14 @@ namespace UE::Editor::DataStorage::QueryStack
 		TEDSQUERYSTACK_API virtual FRowHandleArrayView GetRows() const override;
 		TEDSQUERYSTACK_API virtual FRowHandleArray& GetMutableRows() override;
 		TEDSQUERYSTACK_API virtual RevisionId GetRevision() const override;
-		TEDSQUERYSTACK_API virtual void Update() override;
-		
+		TEDSQUERYSTACK_API virtual void VisitParents(ParentListCallback Callback) override;
+
 		/** Call this when external inputs into the filter, such as search strings or filtering groups, have been updated. */
 		TEDSQUERYSTACK_API void ForceRefresh();
 
+	protected:
+		TEDSQUERYSTACK_API virtual void Update(FTimespan AllottedTime) override;
+	
 	private:
 		ICoreProvider* Storage;
 		Queries::TQueryFunction<bool> Filter;
diff --git a/Engine/Plugins/Experimental/EditorDataStorageFeatures/Source/TedsQueryStack/Public/TedsRowHandleSortNode.h b/Engine/Plugins/Experimental/EditorDataStorageFeatures/Source/TedsQueryStack/Public/TedsRowHandleSortNode.h
--- a/Engine/Plugins/Experimental/EditorDataStorageFeatures/Source/TedsQueryStack/Public/TedsRowHandleSortNode.h
+++ b/Engine/Plugins/Experimental/EditorDataStorageFeatures/Source/TedsQueryStack/Public/TedsRowHandleSortNode.h
@@ -19,10 +19,13 @@ namespace UE::Editor::DataStorage::QueryStack
 		virtual ~FRowHandleSortNode() override = default;
 
 		TEDSQUERYSTACK_API virtual RevisionId GetRevision() const override;
-		TEDSQUERYSTACK_API virtual void Update() override;
+		TEDSQUERYSTACK_API virtual void VisitParents(ParentListCallback Callback) override;
 		TEDSQUERYSTACK_API virtual FRowHandleArrayView GetRows() const override;
 		TEDSQUERYSTACK_API virtual FRowHandleArray& GetMutableRows() override;
 
+	protected:
+		TEDSQUERYSTACK_API virtual void Update(FTimespan AllottedTime) override;
+	
 	private:
 		TSharedPtr<IRowNode> Parent;
 		RevisionId ParentRevision = 0;
diff --git a/Engine/Plugins/Experimental/EditorDataStorageFeatures/Source/TedsQueryStack/Public/TedsRowMergeNode.h b/Engine/Plugins/Experimental/EditorDataStorageFeatures/Source/TedsQueryStack/Public/TedsRowMergeNode.h
--- a/Engine/Plugins/Experimental/EditorDataStorageFeatures/Source/TedsQueryStack/Public/TedsRowMergeNode.h
+++ b/Engine/Plugins/Experimental/EditorDataStorageFeatures/Source/TedsQueryStack/Public/TedsRowMergeNode.h
@@ -24,10 +24,13 @@ namespace UE::Editor::DataStorage::QueryStack
 		virtual ~FRowMergeNode() override = default;
 
 		TEDSQUERYSTACK_API virtual RevisionId GetRevision() const override;
-		TEDSQUERYSTACK_API virtual void Update() override;
+		TEDSQUERYSTACK_API virtual void VisitParents(ParentListCallback Callback) override;
 		TEDSQUERYSTACK_API virtual FRowHandleArrayView GetRows() const override;
 		TEDSQUERYSTACK_API virtual FRowHandleArray& GetMutableRows() override;
 
+	protected:
+		TEDSQUERYSTACK_API virtual void Update(FTimespan AllottedTime) override;
+	
 	private:
 		void Merge();
 
diff --git a/Engine/Plugins/Experimental/EditorDataStorageFeatures/Source/TedsQueryStack/Public/TedsRowMonitorNode.h b/Engine/Plugins/Experimental/EditorDataStorageFeatures/Source/TedsQueryStack/Public/TedsRowMonitorNode.h
--- a/Engine/Plugins/Experimental/EditorDataStorageFeatures/Source/TedsQueryStack/Public/TedsRowMonitorNode.h
+++ b/Engine/Plugins/Experimental/EditorDataStorageFeatures/Source/TedsQueryStack/Public/TedsRowMonitorNode.h
@@ -40,7 +40,7 @@ namespace UE::Editor::DataStorage::QueryStack
 		TEDSQUERYSTACK_API virtual ~FRowMonitorNode() override;
 
 		TEDSQUERYSTACK_API virtual RevisionId GetRevision() const override;
-		TEDSQUERYSTACK_API virtual void Update() override;
+		TEDSQUERYSTACK_API virtual void VisitParents(ParentListCallback Callback) override;
 		TEDSQUERYSTACK_API virtual FRowHandleArrayView GetRows() const override;
 		TEDSQUERYSTACK_API virtual FRowHandleArray& GetMutableRows() override;
 
@@ -53,6 +53,10 @@ namespace UE::Editor::DataStorage::QueryStack
 
 		// Event that is called during Update() to allow systems to respond to monitored rows being removed
 		TEDSQUERYSTACK_API FOnRowsChanged& OnMonitoredRowsRemoved();
+	
+	protected:
+		TEDSQUERYSTACK_API virtual void Update(FTimespan AllottedTime) override;
+	
 	private:
 		void UpdateColumnsFromQuery();
 		void UpdateMonitoredColumns();
diff --git a/Engine/Plugins/Experimental/EditorDataStorageFeatures/Source/TedsQueryStack/Public/TedsRowOrderInversionNode.h b/Engine/Plugins/Experimental/EditorDataStorageFeatures/Source/TedsQueryStack/Public/TedsRowOrderInversionNode.h
--- a/Engine/Plugins/Experimental/EditorDataStorageFeatures/Source/TedsQueryStack/Public/TedsRowOrderInversionNode.h
+++ b/Engine/Plugins/Experimental/EditorDataStorageFeatures/Source/TedsQueryStack/Public/TedsRowOrderInversionNode.h
@@ -23,10 +23,13 @@ namespace UE::Editor::DataStorage::QueryStack
 		TEDSQUERYSTACK_API bool IsEnabled() const;
 
 		TEDSQUERYSTACK_API virtual RevisionId GetRevision() const override;
-		TEDSQUERYSTACK_API virtual void Update() override;
+		TEDSQUERYSTACK_API virtual void VisitParents(ParentListCallback Callback) override;
 		TEDSQUERYSTACK_API virtual FRowHandleArrayView GetRows() const override;
 		TEDSQUERYSTACK_API virtual FRowHandleArray& GetMutableRows() override;
 
+	protected:
+		TEDSQUERYSTACK_API virtual void Update(FTimespan AllottedTime) override;
+	
 	private:
 		void ApplyOrder();
 
diff --git a/Engine/Plugins/Experimental/EditorDataStorageFeatures/Source/TedsQueryStack/Public/TedsRowQueryResultsNode.h b/Engine/Plugins/Experimental/EditorDataStorageFeatures/Source/TedsQueryStack/Public/TedsRowQueryResultsNode.h
--- a/Engine/Plugins/Experimental/EditorDataStorageFeatures/Source/TedsQueryStack/Public/TedsRowQueryResultsNode.h
+++ b/Engine/Plugins/Experimental/EditorDataStorageFeatures/Source/TedsQueryStack/Public/TedsRowQueryResultsNode.h
@@ -46,11 +46,12 @@ namespace UE::Editor::DataStorage::QueryStack
 		TEDSQUERYSTACK_API void Refresh();
 
 		TEDSQUERYSTACK_API virtual RevisionId GetRevision() const override;
-		TEDSQUERYSTACK_API virtual void Update() override;
+		TEDSQUERYSTACK_API virtual void VisitParents(ParentListCallback Callback) override;
 		TEDSQUERYSTACK_API virtual FRowHandleArrayView GetRows() const override;
 		TEDSQUERYSTACK_API virtual FRowHandleArray& GetMutableRows() override;
 
 	protected:
+		TEDSQUERYSTACK_API virtual void Update(FTimespan AllottedTime) override;
 		TEDSQUERYSTACK_API virtual void RefreshInternal(FRowHandleArray& TargetRows);
 
 		FRowHandleArray Rows;
diff --git a/Engine/Plugins/Experimental/EditorDataStorageFeatures/Source/TedsQueryStack/Public/TedsRowSortNode.h b/Engine/Plugins/Experimental/EditorDataStorageFeatures/Source/TedsQueryStack/Public/TedsRowSortNode.h
--- a/Engine/Plugins/Experimental/EditorDataStorageFeatures/Source/TedsQueryStack/Public/TedsRowSortNode.h
+++ b/Engine/Plugins/Experimental/EditorDataStorageFeatures/Source/TedsQueryStack/Public/TedsRowSortNode.h
@@ -42,10 +42,13 @@ namespace UE::Editor::DataStorage::QueryStack
 		TEDSQUERYSTACK_API bool IsSorting() const;
 
 		TEDSQUERYSTACK_API virtual RevisionId GetRevision() const override;
-		TEDSQUERYSTACK_API virtual void Update() override;
+		TEDSQUERYSTACK_API virtual void VisitParents(ParentListCallback Callback) override;
 		TEDSQUERYSTACK_API virtual FRowHandleArrayView GetRows() const override;
 		TEDSQUERYSTACK_API virtual FRowHandleArray& GetMutableRows() override;
 
+	protected:
+		TEDSQUERYSTACK_API virtual void Update(FTimespan AllottedTime) override;
+	
 	private:
 		void SetupFixedSize64();
 		void SetupFixedSizeOnly();
diff --git a/Engine/Plugins/Experimental/EditorDataStorageFeatures/Source/TedsQueryStack/Public/TedsRowViewNode.h b/Engine/Plugins/Experimental/EditorDataStorageFeatures/Source/TedsQueryStack/Public/TedsRowViewNode.h
--- a/Engine/Plugins/Experimental/EditorDataStorageFeatures/Source/TedsQueryStack/Public/TedsRowViewNode.h
+++ b/Engine/Plugins/Experimental/EditorDataStorageFeatures/Source/TedsQueryStack/Public/TedsRowViewNode.h
@@ -21,10 +21,13 @@ namespace UE::Editor::DataStorage::QueryStack
 		TEDSQUERYSTACK_API void ResetView(FRowHandleArray& InRows);
 
 		TEDSQUERYSTACK_API virtual RevisionId GetRevision() const override;
-		TEDSQUERYSTACK_API virtual void Update() override;
+		TEDSQUERYSTACK_API virtual void VisitParents(ParentListCallback Callback) override;
 		TEDSQUERYSTACK_API virtual FRowHandleArrayView GetRows() const override;
 		TEDSQUERYSTACK_API virtual FRowHandleArray& GetMutableRows() override;
 
+	protected:
+		TEDSQUERYSTACK_API virtual void Update(FTimespan AllottedTime) override;
+	
 	private:
 		FRowHandleArray* Rows = nullptr;
 		RevisionId Revision = 0;
diff --git a/Engine/Plugins/Experimental/EditorDataStorageFeatures/Source/TedsQueryStack/Tests/TedsRowMergeNodeTests.cpp b/Engine/Plugins/Experimental/EditorDataStorageFeatures/Source/TedsQueryStack/Tests/TedsRowMergeNodeTests.cpp
--- a/Engine/Plugins/Experimental/EditorDataStorageFeatures/Source/TedsQueryStack/Tests/TedsRowMergeNodeTests.cpp
+++ b/Engine/Plugins/Experimental/EditorDataStorageFeatures/Source/TedsQueryStack/Tests/TedsRowMergeNodeTests.cpp
@@ -2,8 +2,10 @@
 
 #if WITH_TESTS
 
+#include "DataStorage/Features.h"
 #include "TedsRowMergeNode.h"
 #include "TedsRowViewNode.h"
+#include "TedsQueryStackExecutor.h"
 #include "Tests/TestHarnessAdapter.h"
 
 TEST_CASE_NAMED(TEDS_QueryStack_RowMergeNode_Tests, "Editor::DataStorage::QueryStack::FRowMergeNode", "[ApplicationContextMask][EngineFilter]")
@@ -13,11 +15,13 @@ TEST_CASE_NAMED(TEDS_QueryStack_RowMergeNode_Tests, "Editor::DataStorage::QueryS
 
 	SECTION("No parents")
 	{
-		FRowMergeNode Node({}, FRowMergeNode::EMergeApproach::Append);
-		CHECK_EQUALS(TEXT("Revision"), Node.GetRevision(), 0);
-		CHECK_MESSAGE(TEXT("Default for rows is not empty."), Node.GetRows().IsEmpty());
+		TSharedPtr<FRowMergeNode> Node = MakeShared<FRowMergeNode>(TConstArrayView<TSharedPtr<IRowNode>>(), FRowMergeNode::EMergeApproach::Append);
+		CHECK_EQUALS(TEXT("Revision"), Node->GetRevision(), 0);
+		CHECK_MESSAGE(TEXT("Default for rows is not empty."), Node->GetRows().IsEmpty());
 		
-		Node.Update();
+		ICoreProvider& DataStorage = *GetMutableDataStorageFeature<ICoreProvider>(StorageFeatureName);
+		FExplicitUpdateExecutor Executor("NoParentsTestExecutor", Node);
+		Executor.Update();
 	}
 
 	SECTION("Append two nodes")
diff --git a/Engine/Plugins/Experimental/EditorDataStorageFeatures/Source/TedsQueryStack/Tests/TedsRowSortNodeTests.cpp b/Engine/Plugins/Experimental/EditorDataStorageFeatures/Source/TedsQueryStack/Tests/TedsRowSortNodeTests.cpp
--- a/Engine/Plugins/Experimental/EditorDataStorageFeatures/Source/TedsQueryStack/Tests/TedsRowSortNodeTests.cpp
+++ b/Engine/Plugins/Experimental/EditorDataStorageFeatures/Source/TedsQueryStack/Tests/TedsRowSortNodeTests.cpp
@@ -10,6 +10,7 @@
 #include "Sorting/SortTasks.h"
 #include "TedsRowArrayNode.h"
 #include "TedsRowSortNode.h"
+#include "TedsQueryStackExecutor.h"
 #include "Tests/TestHarnessAdapter.h"
 #include "Templates/SharedPointer.h"
 
@@ -131,17 +132,20 @@ namespace UE::Editor::DataStorage::Testing
 		CHECK_MESSAGE(TEXT("Not all tasks completed in time"), !Tasks.HasRemainingTasks());
 	}
 
-	void UpdateNode(const TSharedPtr<UE::Editor::DataStorage::QueryStack::FRowSortNode>& Node)
+	void UpdateNode(
+		UE::Editor::DataStorage::ICoreProvider& DataStorage, const TSharedPtr<UE::Editor::DataStorage::QueryStack::FRowSortNode>& Node)
 	{
+		using namespace UE::Editor::DataStorage::QueryStack;
+		FExplicitCooperativeExecutor Executor("RowSortNodeTestExecutor", Node);
+		
 		uint64 StartTime = FPlatformTime::Cycles64();
 		do
 		{
-			Node->Update();
+			Executor.Update(FTimespan::FromMilliseconds(10));
 		} while (FTimespan(FPlatformTime::Cycles64() - StartTime) < FTimespan::FromSeconds(5) && Node->IsSorting());
 		CHECK_MESSAGE(TEXT("Sorting didn't complete in time"), !Node->IsSorting());
 	}
 
-
 	bool CompareRows(TConstArrayView<RowHandle> Source, TConstArrayView<RowHandle> Compare)
 	{
 		int32 Count = Source.Num();
@@ -538,7 +542,7 @@ TEST_CASE_NAMED(TEDS_QueryStack_RowSortNode_SortNode_Integration_Tests, "Editor:
 			TSharedPtr<FRowArrayNode> ArrayNode = MakeShared<FRowArrayNode>(Array);
 			TSharedPtr<FRowSortNode> SortNode = MakeShared<FRowSortNode>(*Storage, ArrayNode, StringSorter);
 
-			UpdateNode(SortNode);
+			UpdateNode(*Storage, SortNode);
 			
 			CHECK_MESSAGE(TEXT("Sorted array didn't match expectation."),
 				CompareRows(SortNode->GetRows(), {0, 1, 3, 4, 6, 8, 9, 11, 14, 16, 17, 13, 7, 12, 15, 10, 5, 2}));
@@ -574,7 +578,7 @@ TEST_CASE_NAMED(TEDS_QueryStack_RowSortNode_SortNode_Integration_Tests, "Editor:
 			TSharedPtr<FRowArrayNode> ArrayNode = MakeShared<FRowArrayNode>(Array);
 			TSharedPtr<FRowSortNode> SortNode = MakeShared<FRowSortNode>(*Storage, ArrayNode, NameSorter);
 
-			UpdateNode(SortNode);
+			UpdateNode(*Storage, SortNode);
 			
 			CHECK_MESSAGE(TEXT("Sorted array didn't match expectation."),
 				CompareRows(SortNode->GetRows(), { 0, 1, 3, 4, 6, 8, 9, 11, 14, 16, 17, 13, 7, 12, 15, 10, 5, 2 }));
@@ -611,7 +615,7 @@ TEST_CASE_NAMED(TEDS_QueryStack_RowSortNode_SortNode_Integration_Tests, "Editor:
 			TSharedPtr<FRowArrayNode> ArrayNode = MakeShared<FRowArrayNode>(Array);
 			TSharedPtr<FRowSortNode> SortNode = MakeShared<FRowSortNode>(*Storage, ArrayNode, NameSorter);
 
-			UpdateNode(SortNode);
+			UpdateNode(*Storage, SortNode);
 
 			CHECK_MESSAGE(TEXT("Sorted array didn't match expectation."),
 				CompareRows(SortNode->GetRows(), { 13, 7, 12, 15, 0, 1, 3, 4, 6, 8, 9, 11, 14, 16, 17, 10, 5, 2 }));
diff --git a/Engine/Plugins/Experimental/EditorDataStorageFeatures/Source/TedsTableViewer/Private/QueryStack/TopLevelRowsNode.cpp b/Engine/Plugins/Experimental/EditorDataStorageFeatures/Source/TedsTableViewer/Private/QueryStack/TopLevelRowsNode.cpp
--- a/Engine/Plugins/Experimental/EditorDataStorageFeatures/Source/TedsTableViewer/Private/QueryStack/TopLevelRowsNode.cpp
+++ b/Engine/Plugins/Experimental/EditorDataStorageFeatures/Source/TedsTableViewer/Private/QueryStack/TopLevelRowsNode.cpp
@@ -30,9 +30,8 @@ namespace UE::Editor::DataStorage::QueryStack
 		return Revision;
 	}
 
-	void FTopLevelRowsNode::Update()
+	void FTopLevelRowsNode::Update(FTimespan AllottedTime)
 	{
-		Parent->Update();
 		if (Parent->GetRevision() != CachedParentRevision)
 		{
 			UpdateRows();
@@ -41,6 +40,11 @@ namespace UE::Editor::DataStorage::QueryStack
 		}
 	}
 
+	void FTopLevelRowsNode::VisitParents(ParentListCallback Callback)
+	{
+		Callback(Parent);
+	}
+
 	void FTopLevelRowsNode::UpdateRows()
 	{
 		Rows.Empty();
diff --git a/Engine/Plugins/Experimental/EditorDataStorageFeatures/Source/TedsTableViewer/Private/QueryStack/TopLevelRowsNode.h b/Engine/Plugins/Experimental/EditorDataStorageFeatures/Source/TedsTableViewer/Private/QueryStack/TopLevelRowsNode.h
--- a/Engine/Plugins/Experimental/EditorDataStorageFeatures/Source/TedsTableViewer/Private/QueryStack/TopLevelRowsNode.h
+++ b/Engine/Plugins/Experimental/EditorDataStorageFeatures/Source/TedsTableViewer/Private/QueryStack/TopLevelRowsNode.h
@@ -26,9 +26,10 @@ namespace UE::Editor::DataStorage::QueryStack
 		virtual FRowHandleArrayView GetRows() const override;
 		virtual FRowHandleArray& GetMutableRows() override;
 		virtual RevisionId GetRevision() const override;
-		virtual void Update() override;
-
+		virtual void VisitParents(ParentListCallback Callback) override;
+		
 	protected:
+		virtual void Update(FTimespan AllottedTime) override;
 		void UpdateRows();
 	protected:
 		RevisionId Revision = 0;
diff --git a/Engine/Plugins/Experimental/EditorDataStorageFeatures/Source/TedsTableViewer/Private/TedsTableViewerModel.cpp b/Engine/Plugins/Experimental/EditorDataStorageFeatures/Source/TedsTableViewer/Private/TedsTableViewerModel.cpp
--- a/Engine/Plugins/Experimental/EditorDataStorageFeatures/Source/TedsTableViewer/Private/TedsTableViewerModel.cpp
+++ b/Engine/Plugins/Experimental/EditorDataStorageFeatures/Source/TedsTableViewer/Private/TedsTableViewerModel.cpp
@@ -28,6 +28,9 @@ namespace UE::Editor::DataStorage
 		
 		SetQueryStack(InRowQueryStack);
 
+		static int InstanceCounter = 0;
+		QueryStackExecutor = QueryStack::FExplicitUpdateExecutor(FName("TableViewerModel", InstanceCounter++), RowQueryStack);
+
 		TickerHandle = FTSTicker::GetCoreTicker().AddTicker(FTickerDelegate::CreateRaw(this, &FTedsTableViewerModel::Tick), 0);
 
 		GenerateColumns();
@@ -85,7 +88,7 @@ namespace UE::Editor::DataStorage
 
 	bool FTedsTableViewerModel::Tick(float DeltaTime)
 	{
-		RowQueryStack->Update();
+		QueryStackExecutor.Update();
 
 		// If the revision ID has changed, refresh to update our rows
 		if(RowQueryStack->GetRevision() != CachedRowQueryStackRevision)
diff --git a/Engine/Plugins/Experimental/EditorDataStorageFeatures/Source/TedsTableViewer/Private/Widgets/STedsHierarchyViewer.cpp b/Engine/Plugins/Experimental/EditorDataStorageFeatures/Source/TedsTableViewer/Private/Widgets/STedsHierarchyViewer.cpp
--- a/Engine/Plugins/Experimental/EditorDataStorageFeatures/Source/TedsTableViewer/Private/Widgets/STedsHierarchyViewer.cpp
+++ b/Engine/Plugins/Experimental/EditorDataStorageFeatures/Source/TedsTableViewer/Private/Widgets/STedsHierarchyViewer.cpp
@@ -54,6 +54,9 @@ namespace UE::Editor::DataStorage
 		
 		TopLevelRowsNode = MakeShared<QueryStack::FTopLevelRowsNode>(Model->GetDataStorageInterface(), HierarchyInterface, Model->GetRowNode());
 		
+		static int InstanceCounter = 0;
+		QueryStackExecutor = QueryStack::FExplicitUpdateExecutor(FName("HierarchyViewer", InstanceCounter++), TopLevelRowsNode);
+		
 		HeaderRowWidget = SNew( SHeaderRow )
 							.CanSelectGeneratedColumn(true);
 		
@@ -90,10 +93,9 @@ namespace UE::Editor::DataStorage
 	{
 		TopLevelRows.Empty();
 
-		// The current logic relies on the query stack to invalidate and cause a model refresh when the hierarchy changes since we don't have a
-		// generic way of detecting hierarchy changes. If the model isn't refreshed on hierarchy change the top level rows will still use the old
-		// hierarchy
-		TopLevelRowsNode->Update();
+		// Includes a generic way of detecting hierarchy changes. If the model isn't refreshed on hierarchy change the top level rows will still use the old
+		// hierarchy, so the explicit update is needed.
+		QueryStackExecutor.Update();
 		
 		FRowHandleArrayView Rows(TopLevelRowsNode->GetRows());
 
diff --git a/Engine/Plugins/Experimental/EditorDataStorageFeatures/Source/TedsTableViewer/Public/TedsTableViewerModel.h b/Engine/Plugins/Experimental/EditorDataStorageFeatures/Source/TedsTableViewer/Public/TedsTableViewerModel.h
--- a/Engine/Plugins/Experimental/EditorDataStorageFeatures/Source/TedsTableViewer/Public/TedsTableViewerModel.h
+++ b/Engine/Plugins/Experimental/EditorDataStorageFeatures/Source/TedsTableViewer/Public/TedsTableViewerModel.h
@@ -6,6 +6,7 @@
 #include "Delegates/Delegate.h"
 #include "DataStorage/Handles.h"
 #include "Elements/Interfaces/TypedElementDataStorageUiInterface.h"
+#include "TedsQueryStackExecutor.h"
 #include "TedsQueryStackInterfaces.h"
 #include "Templates/SharedPointer.h"
 #include "TypedElementUITypes.h"
@@ -109,6 +110,9 @@ namespace UE::Editor::DataStorage
 		// The row query stack used to supply the rows to display
 		TSharedPtr<QueryStack::IRowNode> RowQueryStack;
 
+		// The executor used to update the RowQueryStack.
+		QueryStack::FExplicitUpdateExecutor QueryStackExecutor;
+
 		// Row node used to sort the rows in the query stack.
 		TSharedPtr<QueryStack::FRowSortNode> RowPrimarySortingNode;
 		// Row node used as a backup to sort the rows in the query stack. Users can hold down shift to setup a secondary sort. As a result
diff --git a/Engine/Plugins/Experimental/EditorDataStorageFeatures/Source/TedsTableViewer/Public/Widgets/STedsHierarchyViewer.h b/Engine/Plugins/Experimental/EditorDataStorageFeatures/Source/TedsTableViewer/Public/Widgets/STedsHierarchyViewer.h
--- a/Engine/Plugins/Experimental/EditorDataStorageFeatures/Source/TedsTableViewer/Public/Widgets/STedsHierarchyViewer.h
+++ b/Engine/Plugins/Experimental/EditorDataStorageFeatures/Source/TedsTableViewer/Public/Widgets/STedsHierarchyViewer.h
@@ -7,6 +7,7 @@
 #include "TedsRowFilterNode.h"
 #include "TedsTableViewerModel.h"
 #include "TedsQueryStackInterfaces.h"
+#include "TedsQueryStackExecutor.h"
 #include "Templates/Function.h"
 #include "Elements/Columns/TypedElementHiearchyColumns.h"
 #include "Widgets/SCompoundWidget.h"
@@ -163,6 +164,9 @@ namespace UE::Editor::DataStorage
 
 		// A specialized query stack node to only contain the top level rows present in the actual query stack we are using.
 		TSharedPtr<QueryStack::IRowNode> TopLevelRowsNode;
+
+		// The executor used to update the hierarchy stack.
+		QueryStack::FExplicitUpdateExecutor QueryStackExecutor;
 		
 		// The cached list of top level rows
 		TArray<TableViewerItemPtr> TopLevelRows;
