[EditorDataStorageFeatures] Introduced "TedsOperations" and "TedsDropOperations" EditorDataStorage factories for new drag&drop system.

diff --git a/Engine/Plugins/Experimental/EditorDataStorageFeatures/EditorDataStorageFeatures.uplugin b/Engine/Plugins/Experimental/EditorDataStorageFeatures/EditorDataStorageFeatures.uplugin
--- a/Engine/Plugins/Experimental/EditorDataStorageFeatures/EditorDataStorageFeatures.uplugin
+++ b/Engine/Plugins/Experimental/EditorDataStorageFeatures/EditorDataStorageFeatures.uplugin
@@ -75,6 +75,14 @@
 		{
 			"Name": "TedsEditorCompatibility",
 			"Type": "Editor"
+		},
+		{
+			"Name": "TedsOperations",
+			"Type": "Editor"
+		},
+		{
+			"Name": "TedsDropOperations",
+			"Type": "Editor"
 		}
 	],
 	"Plugins": [
diff --git a/Engine/Plugins/Experimental/EditorDataStorageFeatures/Source/TedsDropOperations/Private/TedsDropOperationsFactory.cpp b/Engine/Plugins/Experimental/EditorDataStorageFeatures/Source/TedsDropOperations/Private/TedsDropOperationsFactory.cpp
new file mode 100644
--- /dev/null
+++ b/Engine/Plugins/Experimental/EditorDataStorageFeatures/Source/TedsDropOperations/Private/TedsDropOperationsFactory.cpp
@@ -0,0 +1,63 @@
+// Copyright Epic Games, Inc. All Rights Reserved.
+
+#include "TedsDropOperationsFactory.h"
+
+#include "Elements/Framework/TypedElementQueryBuilder.h"
+#include "TedsDropOperationsColumns.h"
+#include "TedsOperationsColumns.h"
+#include "TedsQueryNode.h"
+#include "TedsRowQueryResultsNode.h"
+
+UDropOperationsFactory::UDropOperationsFactory()
+	: UOperationsFactory()
+	, InputTable()
+{
+}
+
+UDropOperationsFactory::~UDropOperationsFactory()
+{
+}
+
+void UDropOperationsFactory::RegisterTables(UE::Editor::DataStorage::ICoreProvider& Storage)
+{
+	static const FName TableName      = FName("DragAndDropOperationsTable");
+	static const FName InputTableName = FName("DragAndDropOperationsInputTable");
+	
+	using namespace UE::Editor::DataStorage;
+	
+	Table      = Storage.RegisterTable<Operations::FApplyColumn, DragAndDrop::FOperationTag>(TableName);
+	InputTable = Storage.RegisterTable<DragAndDrop::FTargetColumn, DragAndDrop::FSourceColumn>(InputTableName);
+}
+
+UE::Editor::DataStorage::FScopedRow UDropOperationsFactory::CreateInputRow(UE::Editor::DataStorage::RowHandle TargetRow,
+	UE::Editor::DataStorage::RowHandle SourceRow, bool bAddDescription)
+{	
+	using namespace UE::Editor::DataStorage;
+
+	if (!StoragePtr->IsRowAvailable(TargetRow) || !StoragePtr->IsRowAvailable(SourceRow))
+	{
+		return { StoragePtr, InvalidRowHandle };
+	}
+
+	if (InputTable == InvalidTableHandle)
+	{
+		return { StoragePtr, InvalidRowHandle };
+	}
+
+	RowHandle InputRow = StoragePtr->AddRow(InputTable);
+	StoragePtr->AddColumn(InputRow, DragAndDrop::FTargetColumn{ .Value = TargetRow });
+	StoragePtr->AddColumn(InputRow, DragAndDrop::FSourceColumn{ .Value = SourceRow });
+	if (bAddDescription)
+	{
+		StoragePtr->AddColumn<DragAndDrop::FDescriptionColumn>(InputRow);
+	}
+	return { StoragePtr, InputRow };
+}
+
+void UDropOperationsFactory::GetQueryDescription(UE::Editor::DataStorage::FQueryDescription& OutDescription) const
+{
+	using namespace UE::Editor::DataStorage;
+	using namespace UE::Editor::DataStorage::Queries;
+
+	OutDescription = Select().Where().All<Operations::FApplyColumn, DragAndDrop::FOperationTag>().Compile();
+}
diff --git a/Engine/Plugins/Experimental/EditorDataStorageFeatures/Source/TedsDropOperations/Private/TedsDropOperationsModule.cpp b/Engine/Plugins/Experimental/EditorDataStorageFeatures/Source/TedsDropOperations/Private/TedsDropOperationsModule.cpp
new file mode 100644
--- /dev/null
+++ b/Engine/Plugins/Experimental/EditorDataStorageFeatures/Source/TedsDropOperations/Private/TedsDropOperationsModule.cpp
@@ -0,0 +1,22 @@
+// Copyright Epic Games, Inc. All Rights Reserved.
+
+#include "TedsDropOperationsModule.h"
+
+#include "Modules/ModuleManager.h"
+
+namespace UE::Editor::DragAndDrop
+{
+
+void FTedsDropOperationsModule::StartupModule()
+{
+	IModuleInterface::StartupModule();
+}
+
+void FTedsDropOperationsModule::ShutdownModule()
+{
+	IModuleInterface::ShutdownModule();
+}
+
+}
+
+IMPLEMENT_MODULE(UE::Editor::DragAndDrop::FTedsDropOperationsModule, TedsDropOperations);
diff --git a/Engine/Plugins/Experimental/EditorDataStorageFeatures/Source/TedsDropOperations/Private/TedsDropOperationsModule.h b/Engine/Plugins/Experimental/EditorDataStorageFeatures/Source/TedsDropOperations/Private/TedsDropOperationsModule.h
new file mode 100644
--- /dev/null
+++ b/Engine/Plugins/Experimental/EditorDataStorageFeatures/Source/TedsDropOperations/Private/TedsDropOperationsModule.h
@@ -0,0 +1,15 @@
+// Copyright Epic Games, Inc. All Rights Reserved.
+
+#pragma once
+
+#include "Modules/ModuleInterface.h"
+
+namespace UE::Editor::DragAndDrop
+{
+	class FTedsDropOperationsModule : public IModuleInterface
+	{
+	public:
+		virtual void StartupModule() override;
+		virtual void ShutdownModule() override;
+	};
+}
\ No newline at end of file
diff --git a/Engine/Plugins/Experimental/EditorDataStorageFeatures/Source/TedsDropOperations/Public/TedsDropOperations.h b/Engine/Plugins/Experimental/EditorDataStorageFeatures/Source/TedsDropOperations/Public/TedsDropOperations.h
new file mode 100644
--- /dev/null
+++ b/Engine/Plugins/Experimental/EditorDataStorageFeatures/Source/TedsDropOperations/Public/TedsDropOperations.h
@@ -0,0 +1,7 @@
+// Copyright Epic Games, Inc. All Rights Reserved.
+
+#pragma once
+
+#include "TedsDropOperationsColumns.h"
+#include "TedsDropOperationsColumnsHelpers.h"
+#include "TedsDropOperationsFactory.h"
diff --git a/Engine/Plugins/Experimental/EditorDataStorageFeatures/Source/TedsDropOperations/Public/TedsDropOperationsColumns.h b/Engine/Plugins/Experimental/EditorDataStorageFeatures/Source/TedsDropOperations/Public/TedsDropOperationsColumns.h
new file mode 100644
--- /dev/null
+++ b/Engine/Plugins/Experimental/EditorDataStorageFeatures/Source/TedsDropOperations/Public/TedsDropOperationsColumns.h
@@ -0,0 +1,74 @@
+// Copyright Epic Games, Inc. All Rights Reserved.
+
+#pragma once
+
+#include "CoreMinimal.h"
+#include "DataStorage/CommonTypes.h"
+
+#include "TedsDropOperationsColumns.generated.h"
+
+USTRUCT(meta = (DisplayName = "Drop Operation"))
+struct FDropOperationTag final : public FEditorDataStorageTag
+{
+	GENERATED_BODY()
+};
+
+USTRUCT(meta = (DisplayName = "Drop Target"))
+struct FDropOperationTargetColumn final : public FEditorDataStorageColumn
+{
+	GENERATED_BODY()
+
+	UE::Editor::DataStorage::RowHandle Value;
+};
+
+USTRUCT(meta = (DisplayName = "Drop Source"))
+struct FDropOperationSourceColumn final : public FEditorDataStorageColumn
+{
+	GENERATED_BODY()
+
+	UE::Editor::DataStorage::RowHandle Value;
+};
+
+USTRUCT(meta = (DisplayName = "Description"))
+struct FDropOperationDescriptionColumn final : public FEditorDataStorageColumn
+{
+	GENERATED_BODY()
+
+	FText Value;
+};
+
+USTRUCT(meta = (DisplayName = "Placement Name"))
+struct FDropOperationPlacementNameColumn final : public FEditorDataStorageColumn
+{
+	GENERATED_BODY()
+
+	FName Value;
+};
+
+USTRUCT(meta = (DisplayName = "Placement Transform"))
+struct FDropOperationPlacementTransformColumn final : public FEditorDataStorageColumn
+{
+	GENERATED_BODY()
+
+	FTransform Value;
+};
+
+USTRUCT(meta = (DisplayName = "Placement Preview"))
+struct FDropOperationPlacementPreviewTag final : public FEditorDataStorageTag
+{
+	GENERATED_BODY()
+};
+
+namespace UE::Editor::DataStorage::DragAndDrop
+{
+	/** Tag to identify this operation as a drop operation. UDropOperationsFactory will only consider rows with this tag. */
+	using FOperationTag = FDropOperationTag;
+
+	using FTargetColumn      = FDropOperationTargetColumn;
+	using FSourceColumn      = FDropOperationSourceColumn;
+	using FDescriptionColumn = FDropOperationDescriptionColumn;
+
+	using FPlacementNameColumn      = FDropOperationPlacementNameColumn;
+	using FPlacementTransformColumn = FDropOperationPlacementTransformColumn;
+	using FPlacementPreviewTag      = FDropOperationPlacementPreviewTag;
+}
diff --git a/Engine/Plugins/Experimental/EditorDataStorageFeatures/Source/TedsDropOperations/Public/TedsDropOperationsColumnsHelpers.h b/Engine/Plugins/Experimental/EditorDataStorageFeatures/Source/TedsDropOperations/Public/TedsDropOperationsColumnsHelpers.h
new file mode 100644
--- /dev/null
+++ b/Engine/Plugins/Experimental/EditorDataStorageFeatures/Source/TedsDropOperations/Public/TedsDropOperationsColumnsHelpers.h
@@ -0,0 +1,33 @@
+// Copyright Epic Games, Inc. All Rights Reserved.
+
+#pragma once
+
+#include "CoreMinimal.h"
+#include "Elements/Interfaces/TypedElementDataStorageInterface.h"
+#include "TedsDropOperationsColumns.h"
+
+namespace UE::Editor::DataStorage::DragAndDrop::Helpers
+{
+
+/** Returns the value of the target column, if set. */
+inline RowHandle GetTargetRow(const ICoreProvider& Storage, RowHandle InputRow)
+{
+	const FTargetColumn* Column = Storage.GetColumn<FTargetColumn>(InputRow);
+	return Column ? Column->Value : InvalidRowHandle;
+}
+
+/** Returns the value of the source column, if set. */
+inline RowHandle GetSourceRow(const ICoreProvider& Storage, RowHandle InputRow)
+{
+	const FSourceColumn* Column = Storage.GetColumn<FSourceColumn>(InputRow);
+	return Column ? Column->Value : InvalidRowHandle;
+}
+
+/** Returns a pointer to the description text, if set.*/
+inline FText* GetDescription(ICoreProvider& Storage, RowHandle InputRow)
+{
+	FDescriptionColumn* Column = Storage.GetColumn<FDescriptionColumn>(InputRow);
+	return Column ? &Column->Value : nullptr;
+}
+
+}
diff --git a/Engine/Plugins/Experimental/EditorDataStorageFeatures/Source/TedsDropOperations/Public/TedsDropOperationsFactory.h b/Engine/Plugins/Experimental/EditorDataStorageFeatures/Source/TedsDropOperations/Public/TedsDropOperationsFactory.h
new file mode 100644
--- /dev/null
+++ b/Engine/Plugins/Experimental/EditorDataStorageFeatures/Source/TedsDropOperations/Public/TedsDropOperationsFactory.h
@@ -0,0 +1,38 @@
+// Copyright Epic Games, Inc. All Rights Reserved.
+
+#pragma once
+
+#include "CoreMinimal.h"
+#include "DataStorage/ScopedRowHandle.h"
+#include "TedsOperationsFactory.h"
+
+#include "TedsDropOperationsFactory.generated.h"
+
+#define UE_API TEDSDROPOPERATIONS_API
+
+/**
+ * Operations system for drop operations.
+ * Only operation rows with the "Drop Tag" will be considered by this system.
+ */
+UCLASS(MinimalAPI)
+class UDropOperationsFactory : public UOperationsFactory
+{
+	GENERATED_BODY()
+public:
+	UE_API UDropOperationsFactory();
+	UE_API virtual ~UDropOperationsFactory() override;
+	
+	UE_API virtual void RegisterTables(UE::Editor::DataStorage::ICoreProvider& Storage) override;
+
+	/** Utility method to create a row with columns for the given input arguments. */
+	UE_API UE::Editor::DataStorage::FScopedRow CreateInputRow(UE::Editor::DataStorage::RowHandle TargetRow,
+		UE::Editor::DataStorage::RowHandle SourceRow, bool bAddDescription = false);
+
+protected:
+	UE_API virtual void GetQueryDescription(UE::Editor::DataStorage::FQueryDescription& OutDescription) const;
+	
+protected:
+	UE::Editor::DataStorage::TableHandle InputTable;	
+};
+
+#undef UE_API
diff --git a/Engine/Plugins/Experimental/EditorDataStorageFeatures/Source/TedsDropOperations/TedsDropOperations.build.cs b/Engine/Plugins/Experimental/EditorDataStorageFeatures/Source/TedsDropOperations/TedsDropOperations.build.cs
new file mode 100644
--- /dev/null
+++ b/Engine/Plugins/Experimental/EditorDataStorageFeatures/Source/TedsDropOperations/TedsDropOperations.build.cs
@@ -0,0 +1,40 @@
+// Copyright Epic Games, Inc. All Rights Reserved.
+
+using UnrealBuildTool;
+
+// Experimental test module. please refrain from depending on it until this warning is removed
+public class TedsDropOperations : ModuleRules
+{
+	public TedsDropOperations(ReadOnlyTargetRules Target) : base(Target)
+	{
+		// Enable truncation warnings in this module
+		CppCompileWarningSettings.UnsafeTypeCastWarningLevel = WarningLevel.Warning;
+
+		ShortName = "TEDSDropOperations";
+
+		PCHUsage = ModuleRules.PCHUsageMode.UseExplicitOrSharedPCHs;
+
+		if (Target.bBuildEditor)
+		{
+			PublicIncludePaths.AddRange(new string[] { });
+			PrivateIncludePaths.AddRange(new string[] { });
+
+			PrivateDependencyModuleNames.AddRange(
+				new string[]
+				{
+					"Core",
+					"CoreUObject",
+					"Engine",
+					"TypedElementFramework",
+					"TedsOperations",
+					"TedsQueryStack"
+				});
+
+			PublicDependencyModuleNames.AddRange(
+				new string[]
+				{
+				});
+
+		}
+	}
+}
diff --git a/Engine/Plugins/Experimental/EditorDataStorageFeatures/Source/TedsOperations/Private/TedsOperationsFactory.cpp b/Engine/Plugins/Experimental/EditorDataStorageFeatures/Source/TedsOperations/Private/TedsOperationsFactory.cpp
new file mode 100644
--- /dev/null
+++ b/Engine/Plugins/Experimental/EditorDataStorageFeatures/Source/TedsOperations/Private/TedsOperationsFactory.cpp
@@ -0,0 +1,203 @@
+// Copyright Epic Games, Inc. All Rights Reserved.
+
+#include "TedsOperationsFactory.h"
+
+#include "Algo/AnyOf.h"
+#include "Elements/Framework/TypedElementQueryBuilder.h"
+#include "Elements/Framework/TypedElementQueryContext.h"
+#include "Elements/Interfaces/TypedElementDataStorageInterface.h"
+#include "TedsOperationsColumns.h"
+#include "TedsQueryNode.h"
+#include "TedsQueryStackExecutor.h"
+#include "TedsRowQueryResultsNode.h"
+
+bool UOperationsFactory::FDefaultLess::operator()(const UE::Editor::DataStorage::ICoreProvider& Storage, UE::Editor::DataStorage::RowHandle A,
+	UE::Editor::DataStorage::RowHandle B) const
+{
+	using namespace UE::Editor::DataStorage;
+	// First sort by priority column.
+	{
+		const Operations::FPriorityColumn* ColumnA = Storage.GetColumn<Operations::FPriorityColumn>(A);
+		const Operations::FPriorityColumn* ColumnB = Storage.GetColumn<Operations::FPriorityColumn>(B);
+		int64 PriorityA = ColumnA ? ColumnA->Value : 0;
+		int64 PriorityB = ColumnB ? ColumnB->Value : 0;
+		if (PriorityA != PriorityB)
+		{
+			return PriorityA > PriorityB; // Highest priority should be the first element.
+		}
+	}
+	// Second sort by name column.
+	{
+		const Operations::FNameColumn* ColumnA = Storage.GetColumn<Operations::FNameColumn>(A);
+		const Operations::FNameColumn* ColumnB = Storage.GetColumn<Operations::FNameColumn>(B);
+		FName NameA = ColumnA ? ColumnA->Value : FName();
+		FName NameB = ColumnB ? ColumnB->Value : FName();
+		if (int32 Result = NameA.Compare(NameB); Result != 0)
+		{
+			return Result < 0;
+		}
+	}
+	// Third sort by row handle.
+	return A < B;
+}
+
+UOperationsFactory::UOperationsFactory()
+	: StoragePtr(nullptr)
+	, Table(UE::Editor::DataStorage::InvalidTableHandle)
+	, QueryNode()
+	, ResultsNode()
+	, Executor()
+{
+}
+
+UOperationsFactory::~UOperationsFactory()
+{
+}
+
+void UOperationsFactory::PreRegister(UE::Editor::DataStorage::ICoreProvider& Storage)
+{
+	StoragePtr = &Storage;
+}
+
+void UOperationsFactory::RegisterTables(UE::Editor::DataStorage::ICoreProvider& Storage)
+{
+	static const FName TableName("OperationsTable");
+	Table = Storage.RegisterTable<FOperationApplyColumn>(TableName);
+}
+
+void UOperationsFactory::RegisterQueries(UE::Editor::DataStorage::ICoreProvider& Storage)
+{
+	using namespace UE::Editor::DataStorage;
+	using namespace UE::Editor::DataStorage::Queries;
+	using namespace UE::Editor::DataStorage::QueryStack;
+
+	FQueryDescription QueryDescription;
+	GetQueryDescription(QueryDescription);
+
+	QueryNode   = MakeShared<FQueryNode>(Storage, MoveTemp(QueryDescription));
+	ResultsNode = MakeShared<FRowQueryResultsNode>(Storage, QueryNode, FRowQueryResultsNode::ESyncFlags::RefreshOnUpdate);
+	Executor    = MakeShared<FExplicitUpdateExecutor>("Operations", ResultsNode);
+}
+
+UE::Editor::DataStorage::RowHandle UOperationsFactory::AddOperation(FName Name,
+	TNotNull<UE::Editor::DataStorage::Operations::FApplyCallback> Operation,
+	UE::Editor::DataStorage::Operations::FTestCallback Test,
+	UE::Editor::DataStorage::Operations::FProbeCallback Probe, int64 Priority)
+{
+	using namespace UE::Editor::DataStorage;
+
+	if (Table == InvalidTableHandle)
+	{
+		return InvalidRowHandle;
+	}
+
+	RowHandle Row = StoragePtr->AddRow(Table);
+	if (Row != InvalidRowHandle)
+	{
+		StoragePtr->AddColumn(Row, Operations::FApplyColumn{ .Callback = MoveTemp(Operation) });
+		if (Name != NAME_None)
+		{
+			StoragePtr->AddColumn(Row, Operations::FNameColumn{ .Value = MoveTemp(Name) });
+		}
+		if (Test)
+		{
+			StoragePtr->AddColumn(Row, Operations::FTestColumn{ .Callback = MoveTemp(Test) });
+		}
+		if (Probe)
+		{
+			StoragePtr->AddColumn(Row, Operations::FProbeColumn{ .Callback = MoveTemp(Probe) });
+		}
+		if (Priority != DefaultPriority)
+		{
+			StoragePtr->AddColumn(Row, Operations::FPriorityColumn{ .Value = Priority });
+		}
+	}
+	return Row;	
+}
+
+void UOperationsFactory::GetOperations(TArray<UE::Editor::DataStorage::RowHandle>& OutRows, UE::Editor::DataStorage::RowHandle InputRow,
+	const FFilter& Filter, const FLess& Less) const
+{
+	using namespace UE::Editor::DataStorage;
+	
+	OutRows.Empty();
+	if (!Executor)
+	{
+		return;
+	}
+	
+	Executor->Update();
+	FRowHandleArrayView AllRows = ResultsNode->GetRows();
+	if (AllRows.IsEmpty())
+	{
+		return;
+	}
+
+	for (RowHandle Row : AllRows)
+	{
+		if (Row == InvalidRowHandle)
+		{
+			continue;
+		}
+
+		if (Filter && !Filter(*StoragePtr, Row))
+		{
+			continue;
+		}
+
+		const Operations::FProbeColumn* ProbeColumn = StoragePtr->GetColumn<Operations::FProbeColumn>(Row);
+		if (ProbeColumn && ProbeColumn->Callback && !ProbeColumn->Callback(*StoragePtr, InputRow))
+		{
+			continue;
+		}
+
+		OutRows.Add(Row);
+	}
+
+	if (Less)
+	{
+		Algo::Sort(OutRows, [this, &Less](RowHandle A, RowHandle B) { return Less(*StoragePtr, A, B); });
+	}
+}
+
+bool UOperationsFactory::Test(UE::Editor::DataStorage::RowHandle InputRow, const FFilter& Filter, const FLess& Less) const
+{
+	using namespace UE::Editor::DataStorage;
+	
+	TArray<RowHandle> OperationRows;
+	GetOperations(OperationRows, InputRow, Filter, Less);
+	return Algo::AnyOf(OperationRows, [this, InputRow](RowHandle OperationRow)
+	{
+		const Operations::FTestColumn* TestColumn = StoragePtr->GetColumn<Operations::FTestColumn>(OperationRow);
+		return !TestColumn || !TestColumn->Callback || TestColumn->Callback(*StoragePtr, InputRow);	// No column equals True.
+	});
+}
+
+TOptional<UE::Editor::DataStorage::Operations::FResult> UOperationsFactory::Apply(UE::Editor::DataStorage::RowHandle InputRow, const FFilter& Filter,
+	const FLess& Less) const
+{
+	using namespace UE::Editor::DataStorage;
+	
+	TArray<RowHandle> OperationRows;
+	GetOperations(OperationRows, InputRow, Filter, Less);
+	for (RowHandle OperationRow : OperationRows)
+	{
+		const Operations::FApplyColumn* ApplyColumn = StoragePtr->GetColumn<Operations::FApplyColumn>(OperationRow);
+		if (ApplyColumn && ApplyColumn->Callback)
+		{
+			if (TOptional<Operations::FResult> Result = ApplyColumn->Callback(*StoragePtr, InputRow))
+			{
+				return Result; // Return the first successful result.
+			}
+		}
+	}
+	return {};
+}
+
+void UOperationsFactory::GetQueryDescription(UE::Editor::DataStorage::FQueryDescription& OutDescription) const
+{
+	using namespace UE::Editor::DataStorage;
+	using namespace UE::Editor::DataStorage::Queries;
+
+	OutDescription = Select().Where().All<Operations::FApplyColumn>().Compile();
+}
diff --git a/Engine/Plugins/Experimental/EditorDataStorageFeatures/Source/TedsOperations/Private/TedsOperationsModule.cpp b/Engine/Plugins/Experimental/EditorDataStorageFeatures/Source/TedsOperations/Private/TedsOperationsModule.cpp
new file mode 100644
--- /dev/null
+++ b/Engine/Plugins/Experimental/EditorDataStorageFeatures/Source/TedsOperations/Private/TedsOperationsModule.cpp
@@ -0,0 +1,22 @@
+// Copyright Epic Games, Inc. All Rights Reserved.
+
+#include "TedsOperationsModule.h"
+
+#include "Modules/ModuleManager.h"
+
+namespace UE::Editor::DataStorage::Operations
+{
+
+void FTedsOperationsModule::StartupModule()
+{
+	IModuleInterface::StartupModule();
+}
+
+void FTedsOperationsModule::ShutdownModule()
+{
+	IModuleInterface::ShutdownModule();
+}
+
+}
+
+IMPLEMENT_MODULE(UE::Editor::DataStorage::Operations::FTedsOperationsModule, TedsOperations);
diff --git a/Engine/Plugins/Experimental/EditorDataStorageFeatures/Source/TedsOperations/Private/TedsOperationsModule.h b/Engine/Plugins/Experimental/EditorDataStorageFeatures/Source/TedsOperations/Private/TedsOperationsModule.h
new file mode 100644
--- /dev/null
+++ b/Engine/Plugins/Experimental/EditorDataStorageFeatures/Source/TedsOperations/Private/TedsOperationsModule.h
@@ -0,0 +1,16 @@
+// Copyright Epic Games, Inc. All Rights Reserved.
+
+#pragma once
+
+#include "CoreMinimal.h"
+#include "Modules/ModuleInterface.h"
+
+namespace UE::Editor::DataStorage::Operations
+{
+	class FTedsOperationsModule : public IModuleInterface
+	{
+	public:
+		virtual void StartupModule() override;
+		virtual void ShutdownModule() override;
+	};
+}
\ No newline at end of file
diff --git a/Engine/Plugins/Experimental/EditorDataStorageFeatures/Source/TedsOperations/Public/TedsOperations.h b/Engine/Plugins/Experimental/EditorDataStorageFeatures/Source/TedsOperations/Public/TedsOperations.h
new file mode 100644
--- /dev/null
+++ b/Engine/Plugins/Experimental/EditorDataStorageFeatures/Source/TedsOperations/Public/TedsOperations.h
@@ -0,0 +1,5 @@
+// Copyright Epic Games, Inc. All Rights Reserved.
+
+#pragma once
+
+#include "TedsOperationsFactory.h"
diff --git a/Engine/Plugins/Experimental/EditorDataStorageFeatures/Source/TedsOperations/Public/TedsOperationsCallbacks.h b/Engine/Plugins/Experimental/EditorDataStorageFeatures/Source/TedsOperations/Public/TedsOperationsCallbacks.h
new file mode 100644
--- /dev/null
+++ b/Engine/Plugins/Experimental/EditorDataStorageFeatures/Source/TedsOperations/Public/TedsOperationsCallbacks.h
@@ -0,0 +1,22 @@
+// Copyright Epic Games, Inc. All Rights Reserved.
+
+#pragma once
+
+#include "CoreMinimal.h"
+#include "DataStorage/Handles.h"
+#include "TedsOperationsResult.h"
+
+namespace UE::Editor::DataStorage
+{
+class ICoreProvider;
+}
+
+namespace UE::Editor::DataStorage::Operations
+{
+	/** Probe call to test superficially if the given input is potentially acceptable for the operation. */
+	using FProbeCallback = TFunction<bool(const ICoreProvider&, RowHandle)>;
+	/** Test call to query if the operation for the given input can be executed. */
+	using FTestCallback  = TFunction<bool(ICoreProvider&, RowHandle)>;
+	/** Apply call to execute the operation for the given input. */
+	using FApplyCallback = TFunction<TOptional<FResult>(ICoreProvider&, RowHandle)>;
+}
diff --git a/Engine/Plugins/Experimental/EditorDataStorageFeatures/Source/TedsOperations/Public/TedsOperationsColumns.h b/Engine/Plugins/Experimental/EditorDataStorageFeatures/Source/TedsOperations/Public/TedsOperationsColumns.h
new file mode 100644
--- /dev/null
+++ b/Engine/Plugins/Experimental/EditorDataStorageFeatures/Source/TedsOperations/Public/TedsOperationsColumns.h
@@ -0,0 +1,64 @@
+// Copyright Epic Games, Inc. All Rights Reserved.
+
+#pragma once
+
+#include "CoreMinimal.h"
+#include "DataStorage/CommonTypes.h"
+#include "DataStorage/Handles.h"
+#include "TedsOperationsCallbacks.h"
+
+#include "TedsOperationsColumns.generated.h"
+
+USTRUCT(meta = (DisplayName = "Operation Name"))
+struct FOperationNameColumn final : public FEditorDataStorageColumn
+{
+	GENERATED_BODY()
+
+	FName Value;
+};
+
+USTRUCT(meta = (DisplayName = "Operation Probe"))
+struct FOperationProbeColumn final : public FEditorDataStorageColumn
+{
+	GENERATED_BODY()
+
+	UE::Editor::DataStorage::Operations::FProbeCallback Callback;
+};
+
+USTRUCT(meta = (DisplayName = "Operation Test"))
+struct FOperationTestColumn final : public FEditorDataStorageColumn
+{
+	GENERATED_BODY()
+
+	UE::Editor::DataStorage::Operations::FTestCallback Callback;
+};
+
+USTRUCT(meta = (DisplayName = "Operation Apply"))
+struct FOperationApplyColumn final : public FEditorDataStorageColumn
+{
+	GENERATED_BODY()
+
+	UE::Editor::DataStorage::Operations::FApplyCallback Callback;
+};
+
+USTRUCT(meta = (DisplayName = "Operation Priority"))
+struct FOperationPriorityColumn final : public FEditorDataStorageColumn
+{
+	GENERATED_BODY()
+
+	int64 Value;
+};
+
+namespace UE::Editor::DataStorage::Operations
+{
+	/** Column for the name of the operation. */
+	using FNameColumn = FOperationNameColumn;
+	/** Probe call column, to test superficially if a given input is potentially acceptable for the operation. */
+	using FProbeColumn = FOperationProbeColumn;
+	/** Test call column, to query if the operation for a given input can be executed. */
+	using FTestColumn = FOperationTestColumn;
+	/** Apply call column, to execute the operation for a given input. */
+	using FApplyColumn = FOperationApplyColumn;
+	/** Priority value of this operation. Used for the default sorting when calling the operation system. */
+	using FPriorityColumn = FOperationPriorityColumn;
+}
diff --git a/Engine/Plugins/Experimental/EditorDataStorageFeatures/Source/TedsOperations/Public/TedsOperationsFactory.h b/Engine/Plugins/Experimental/EditorDataStorageFeatures/Source/TedsOperations/Public/TedsOperationsFactory.h
new file mode 100644
--- /dev/null
+++ b/Engine/Plugins/Experimental/EditorDataStorageFeatures/Source/TedsOperations/Public/TedsOperationsFactory.h
@@ -0,0 +1,93 @@
+// Copyright Epic Games, Inc. All Rights Reserved.
+
+#pragma once
+
+#include "CoreMinimal.h"
+#include "DataStorage/Handles.h"
+#include "Elements/Interfaces/TypedElementDataStorageFactory.h"
+#include "TedsOperationsCallbacks.h"
+
+#include "TedsOperationsFactory.generated.h"
+
+#define UE_API TEDSOPERATIONS_API
+
+namespace UE::Editor::DataStorage
+{
+struct FQueryDescription;
+}
+namespace UE::Editor::DataStorage::QueryStack
+{
+class IQueryNode;
+class IRowNode;
+class FRowFilterNode;
+class FExplicitUpdateExecutor;
+}
+
+/** 
+ * Generic system to provide functionality for the operation columns. Operations
+ * that were added to the data storage can be selected and "tested" or "applied"
+ * in a consistent way by this system depending on an input row.
+ */
+UCLASS(MinimalAPI)
+class UOperationsFactory : public UEditorDataStorageFactory
+{
+	GENERATED_BODY()
+public:
+	static constexpr int64 DefaultPriority = 0;
+	
+	using FFilter = TFunction<bool(const UE::Editor::DataStorage::ICoreProvider&, UE::Editor::DataStorage::RowHandle)>;
+	using FLess   = TFunction<bool(const UE::Editor::DataStorage::ICoreProvider&, UE::Editor::DataStorage::RowHandle, UE::Editor::DataStorage::RowHandle)>;
+
+	struct FDefaultLess
+	{
+		UE_API bool operator()(const UE::Editor::DataStorage::ICoreProvider& Storage, UE::Editor::DataStorage::RowHandle A, UE::Editor::DataStorage::RowHandle B) const;
+	};
+
+public:
+	UE_API UOperationsFactory();
+	UE_API virtual ~UOperationsFactory() override;
+
+	UE_API virtual void PreRegister(UE::Editor::DataStorage::ICoreProvider& Storage) override;
+	UE_API virtual void RegisterTables(UE::Editor::DataStorage::ICoreProvider& Storage) override;
+	UE_API virtual void RegisterQueries(UE::Editor::DataStorage::ICoreProvider& Storage) override;
+
+	/** Utility method to create a row with columns for the given input operation. */
+	UE_API UE::Editor::DataStorage::RowHandle AddOperation(FName Name, TNotNull<UE::Editor::DataStorage::Operations::FApplyCallback> Operation,
+		UE::Editor::DataStorage::Operations::FTestCallback Test, UE::Editor::DataStorage::Operations::FProbeCallback Probe,
+		int64 Priority = DefaultPriority);
+
+	/** Retrieves all operation rows that accept (probe call) the given input row.
+	 * @param OutRows The array to fill with the operation rows.
+	 * @param InputRow The row that contains input data that will be passed to the operations.
+	 * @param Filter [optional] Predicate that can be used to customarily filter the operation rows. False = Skip, True = Keep.
+	 * @param Less [optional] Sorter that can be used to customarily sort the operation rows.
+	 */
+	UE_API void GetOperations(TArray<UE::Editor::DataStorage::RowHandle>& OutRows, UE::Editor::DataStorage::RowHandle InputRow,
+		const FFilter& Filter = {}, const FLess& Less = FDefaultLess()) const;
+
+	/** Invokes the Test callback for all operation rows that accept (probe call) the given input row and returns on the first successful one.
+	 * @param InputRow The row that contains input data that will be passed to the operations.
+	 * @param Filter [optional] Predicate that can be used to customarily filter the operation rows. False = Skip, True = Keep.
+	 * @param Less [optional] Sorter that can be used to customarily sort the operation rows.
+	 */
+	UE_API bool Test(UE::Editor::DataStorage::RowHandle InputRow, const FFilter& Filter = {}, const FLess& Less = FDefaultLess()) const;
+
+	/** Invokes the Apply callback for all operation rows that accept (probe call) the given input row and returns on the first successful one.
+	 * @param InputRow The row that contains input data that will be passed to the operations.
+	 * @param Filter [optional] Predicate that can be used to customarily filter the operation rows. False = Skip, True = Keep.
+	 * @param Less [optional] Sorter that can be used to customarily sort the operation rows.
+	 */
+	UE_API TOptional<UE::Editor::DataStorage::Operations::FResult> Apply(UE::Editor::DataStorage::RowHandle InputRow,
+		const FFilter& Filter = {}, const FLess& Less = FDefaultLess()) const;
+protected:
+	UE_API virtual void GetQueryDescription(UE::Editor::DataStorage::FQueryDescription& OutDescription) const;
+	
+protected:
+	UE::Editor::DataStorage::ICoreProvider*                     StoragePtr;
+	UE::Editor::DataStorage::TableHandle                        Table;
+	TSharedPtr<UE::Editor::DataStorage::QueryStack::IQueryNode> QueryNode;
+	TSharedPtr<UE::Editor::DataStorage::QueryStack::IRowNode>   ResultsNode;
+	TSharedPtr<UE::Editor::DataStorage::QueryStack::FExplicitUpdateExecutor> Executor;
+};
+
+#undef UE_API
diff --git a/Engine/Plugins/Experimental/EditorDataStorageFeatures/Source/TedsOperations/Public/TedsOperationsResult.h b/Engine/Plugins/Experimental/EditorDataStorageFeatures/Source/TedsOperations/Public/TedsOperationsResult.h
new file mode 100644
--- /dev/null
+++ b/Engine/Plugins/Experimental/EditorDataStorageFeatures/Source/TedsOperations/Public/TedsOperationsResult.h
@@ -0,0 +1,26 @@
+// Copyright Epic Games, Inc. All Rights Reserved.
+
+#pragma once
+
+#include "CoreMinimal.h"
+#include "DataStorage/Handles.h"
+
+namespace UE::Editor::DataStorage::Operations
+{
+
+enum class EResultType
+{
+	Created,
+	Modified,
+	Removed,
+};
+
+struct FResultEntry
+{
+	RowHandle   Row;
+	EResultType Type;
+};
+
+using FResult = TArray<FResultEntry>;
+
+}
diff --git a/Engine/Plugins/Experimental/EditorDataStorageFeatures/Source/TedsOperations/TedsOperations.build.cs b/Engine/Plugins/Experimental/EditorDataStorageFeatures/Source/TedsOperations/TedsOperations.build.cs
new file mode 100644
--- /dev/null
+++ b/Engine/Plugins/Experimental/EditorDataStorageFeatures/Source/TedsOperations/TedsOperations.build.cs
@@ -0,0 +1,39 @@
+// Copyright Epic Games, Inc. All Rights Reserved.
+
+using UnrealBuildTool;
+
+// Experimental test module. please refrain from depending on it until this warning is removed
+public class TedsOperations : ModuleRules
+{
+	public TedsOperations(ReadOnlyTargetRules Target) : base(Target)
+	{
+		// Enable truncation warnings in this module
+		CppCompileWarningSettings.UnsafeTypeCastWarningLevel = WarningLevel.Warning;
+
+		ShortName = "TEDSOperations";
+
+		PCHUsage = ModuleRules.PCHUsageMode.UseExplicitOrSharedPCHs;
+
+		if (Target.bBuildEditor)
+		{
+			PublicIncludePaths.AddRange(new string[] { });
+			PrivateIncludePaths.AddRange(new string[] { });
+
+			PublicDependencyModuleNames.AddRange(
+				new string[]
+				{
+					"Core",
+					"CoreUObject",
+					"Engine",
+					"TypedElementFramework",
+					"TedsQueryStack",
+				});
+
+			PrivateDependencyModuleNames.AddRange(
+				new string[]
+				{
+				});
+
+		}
+	}
+}
diff --git a/Engine/Source/Runtime/TypedElementFramework/Public/DataStorage/ScopedRowHandle.h b/Engine/Source/Runtime/TypedElementFramework/Public/DataStorage/ScopedRowHandle.h
new file mode 100644
--- /dev/null
+++ b/Engine/Source/Runtime/TypedElementFramework/Public/DataStorage/ScopedRowHandle.h
@@ -0,0 +1,56 @@
+// Copyright Epic Games, Inc. All Rights Reserved.
+
+#pragma once
+
+#include "CoreMinimal.h"
+#include "DataStorage/Handles.h"
+#include "Elements/Interfaces/TypedElementDataStorageInterface.h"
+
+#include <utility>
+
+namespace UE::Editor::DataStorage
+{
+	/** RAII class that removes its row from storage when going out of scope. */
+	class FScopedRow
+	{
+	public:
+		FScopedRow(ICoreProvider* InStorage, RowHandle InRow)
+			: Storage(InStorage)
+			, Row(InRow)
+		{
+		}
+
+		FScopedRow(FScopedRow&& Other)
+			: Storage(std::exchange(Other.Storage, nullptr))
+			, Row(std::exchange(Other.Row, InvalidRowHandle))
+		{
+		}
+
+		FScopedRow& operator=(FScopedRow&& Other)
+		{
+			std::swap(Storage, Other.Storage);
+			std::swap(Row, Other.Row);
+			return *this;
+		}
+
+		FScopedRow(const FScopedRow&) = delete;
+		FScopedRow& operator=(const FScopedRow&) = delete;
+
+		~FScopedRow()
+		{
+			if (Storage && Row != InvalidRowHandle)
+			{
+				Storage->RemoveRow(Row);
+			}
+		}
+
+		ICoreProvider* GetStorage() const { return Storage; }
+		RowHandle      GetRow() const { return Row; }
+
+		explicit operator RowHandle() const { return GetRow(); }
+
+	private:
+		ICoreProvider* Storage;
+		RowHandle      Row;
+	};
+}
