Teds Table Viewer: Add support for TEDS drag and drop
TEDS Drag and Drop: Add a new FTedsDragDropOp to allow interop with slate drag and drop

diff --git a/Engine/Plugins/Experimental/EditorDataStorageFeatures/Source/TedsDropOperations/Private/TedsDragDropOp.cpp b/Engine/Plugins/Experimental/EditorDataStorageFeatures/Source/TedsDropOperations/Private/TedsDragDropOp.cpp
new file mode 100644
--- /dev/null
+++ b/Engine/Plugins/Experimental/EditorDataStorageFeatures/Source/TedsDropOperations/Private/TedsDragDropOp.cpp
@@ -0,0 +1,34 @@
+// Copyright Epic Games, Inc. All Rights Reserved
+// .
+#include "TedsDragDropOp.h"
+
+namespace UE::Editor::DataStorage::Widgets
+{
+	void FTedsDragDropOp::Init(FRowHandleArray InDraggedRows)
+	{
+		DraggedRows = MoveTemp(InDraggedRows);
+	}
+	
+	TSharedRef<FTedsDragDropOp> FTedsDragDropOp::New(FRowHandleArray InDraggedRows)
+	{
+		TSharedRef<FTedsDragDropOp> Operation = MakeShareable(new FTedsDragDropOp);
+		Operation->Init(MoveTemp(InDraggedRows));
+		Operation->Construct();
+		return Operation;
+	}
+	
+	FRowHandleArrayView FTedsDragDropOp::GetRows() const
+	{
+		return DraggedRows.GetRows();
+	}
+
+	void FTedsDragDropOp::SetDropStatus(EDropTestResult InDropTestResult)
+	{
+		DropTestResult = InDropTestResult;
+	}
+
+	EDropTestResult FTedsDragDropOp::GetDropStatus() const
+	{
+		return DropTestResult;
+	}
+}
diff --git a/Engine/Plugins/Experimental/EditorDataStorageFeatures/Source/TedsDropOperations/Public/TedsDragDropOp.h b/Engine/Plugins/Experimental/EditorDataStorageFeatures/Source/TedsDropOperations/Public/TedsDragDropOp.h
new file mode 100644
--- /dev/null
+++ b/Engine/Plugins/Experimental/EditorDataStorageFeatures/Source/TedsDropOperations/Public/TedsDragDropOp.h
@@ -0,0 +1,42 @@
+// Copyright Epic Games, Inc. All Rights Reserved.
+
+#pragma once
+
+#include "DragAndDrop/DecoratedDragDropOp.h"
+#include "Elements/Framework/TypedElementDataStorageWidget.h"
+#include "Elements/Framework/TypedElementRowHandleArray.h"
+#include "Input/DragAndDrop.h"
+
+#define UE_API TEDSDROPOPERATIONS_API
+
+namespace UE::Editor::DataStorage::Widgets
+{
+	enum EDropTestResult
+	{
+		NoTestsRun, // Default = no drop tests have been run yet
+		SomeRowsPass, // Some but not all  the rows passed the drop test
+		AllRowsPass // All the rows passed the drop test
+	};
+	
+	// TEDS UI TODO: Instead of using a decorated operation use a default one and allow decorator customization using TEDS UI
+	class FTedsDragDropOp : public FDecoratedDragDropOp
+	{
+	public:
+		DRAG_DROP_OPERATOR_TYPE(FTedsDragDropOperation, FDecoratedDragDropOp)
+
+		UE_API virtual ~FTedsDragDropOp() override = default;
+		UE_API static TSharedRef<FTedsDragDropOp> New(FRowHandleArray InDraggedRows);
+		UE_API FRowHandleArrayView GetRows() const;
+		UE_API void SetDropStatus(EDropTestResult InDropStatus);
+		UE_API EDropTestResult GetDropStatus() const;
+		
+	protected:
+		void Init(FRowHandleArray InDraggedRows);
+		
+	protected:
+		FRowHandleArray DraggedRows;
+		EDropTestResult DropTestResult = EDropTestResult::NoTestsRun;
+	};
+}
+
+#undef UE_API
\ No newline at end of file
diff --git a/Engine/Plugins/Experimental/EditorDataStorageFeatures/Source/TedsDropOperations/TedsDropOperations.build.cs b/Engine/Plugins/Experimental/EditorDataStorageFeatures/Source/TedsDropOperations/TedsDropOperations.build.cs
--- a/Engine/Plugins/Experimental/EditorDataStorageFeatures/Source/TedsDropOperations/TedsDropOperations.build.cs
+++ b/Engine/Plugins/Experimental/EditorDataStorageFeatures/Source/TedsDropOperations/TedsDropOperations.build.cs
@@ -25,14 +25,17 @@ public class TedsDropOperations : ModuleRules
 					"Core",
 					"CoreUObject",
 					"Engine",
+					"Slate",
+					"SlateCore",
 					"TypedElementFramework",
-					"TedsOperations",
-					"TedsQueryStack"
+					"TedsQueryStack",
+					"UnrealEd"
 				});
 
 			PublicDependencyModuleNames.AddRange(
 				new string[]
 				{
+					"TedsOperations"
 				});
 
 		}
diff --git a/Engine/Plugins/Experimental/EditorDataStorageFeatures/Source/TedsTableViewer/Private/Widgets/STedsHierarchyViewer.cpp b/Engine/Plugins/Experimental/EditorDataStorageFeatures/Source/TedsTableViewer/Private/Widgets/STedsHierarchyViewer.cpp
--- a/Engine/Plugins/Experimental/EditorDataStorageFeatures/Source/TedsTableViewer/Private/Widgets/STedsHierarchyViewer.cpp
+++ b/Engine/Plugins/Experimental/EditorDataStorageFeatures/Source/TedsTableViewer/Private/Widgets/STedsHierarchyViewer.cpp
@@ -268,7 +268,7 @@ namespace UE::Editor::DataStorage
 	{
 		TreeView->SetItemExpansion(InItem, DefaultExpansionState.Get() == EExpansionState::Expanded);
 
-		return SNew(SHierarchyViewerRow, OwnerTable, Model.ToSharedRef())
+		return SNew(SHierarchyViewerRow, OwnerTable, Model.ToSharedRef(), SharedThis(this))
 				.Padding(FMargin(0))
 				.Item(InItem);
 	}
diff --git a/Engine/Plugins/Experimental/EditorDataStorageFeatures/Source/TedsTableViewer/Private/Widgets/STedsTableViewer.cpp b/Engine/Plugins/Experimental/EditorDataStorageFeatures/Source/TedsTableViewer/Private/Widgets/STedsTableViewer.cpp
--- a/Engine/Plugins/Experimental/EditorDataStorageFeatures/Source/TedsTableViewer/Private/Widgets/STedsTableViewer.cpp
+++ b/Engine/Plugins/Experimental/EditorDataStorageFeatures/Source/TedsTableViewer/Private/Widgets/STedsTableViewer.cpp
@@ -228,7 +228,7 @@ namespace UE::Editor::DataStorage
 
 	TSharedRef<ITableRow> STedsTableViewer::MakeTableRowWidget(TableViewerItemPtr InItem, const TSharedRef<STableViewBase>& OwnerTable) const
 	{
-		return SNew(STedsTableViewerRow, OwnerTable, Model.ToSharedRef())
+		return SNew(STedsTableViewerRow, OwnerTable, Model.ToSharedRef(), SharedThis(this))
 				.ItemHeight(ItemHeight)
 				.Padding(ItemPadding)
 				.ParentWidgetRowHandle(GetWidgetRowHandle())
diff --git a/Engine/Plugins/Experimental/EditorDataStorageFeatures/Source/TedsTableViewer/Private/Widgets/STedsTableViewerRow.cpp b/Engine/Plugins/Experimental/EditorDataStorageFeatures/Source/TedsTableViewer/Private/Widgets/STedsTableViewerRow.cpp
--- a/Engine/Plugins/Experimental/EditorDataStorageFeatures/Source/TedsTableViewer/Private/Widgets/STedsTableViewerRow.cpp
+++ b/Engine/Plugins/Experimental/EditorDataStorageFeatures/Source/TedsTableViewer/Private/Widgets/STedsTableViewerRow.cpp
@@ -3,24 +3,32 @@
 
 #include "Widgets/STedsTableViewerRow.h"
 
+#include "Editor.h"
+#include "TedsDragDropOp.h"
+#include "TedsDropOperations.h"
 #include "TedsTableViewerColumn.h"
 #include "Elements/Columns/TypedElementHiearchyColumns.h"
 #include "Widgets/STedsTableViewer.h"
 
 namespace UE::Editor::DataStorage
 {
-	void STedsTableViewerRow::Construct(const FArguments& InArgs, const TSharedRef<STableViewBase>& OwnerTableView, const TSharedRef<FTedsTableViewerModel>& InTableViewerModel)
+	void STedsTableViewerRow::Construct(const FArguments& InArgs, TSharedRef<STableViewBase> OwnerTableView, TSharedRef<FTedsTableViewerModel> InTableViewerModel, TWeakPtr<const ITableViewer> InTableViewer)
 	{
 		Item = InArgs._Item;
-		TableViewerModel = InTableViewerModel;
+		TableViewerModel = MoveTemp(InTableViewerModel);
+		WeakTableViewer = MoveTemp(InTableViewer);
 		ParentWidgetRowHandle = InArgs._ParentWidgetRowHandle;
 		ItemHeight = InArgs._ItemHeight;
 
 		const auto Args = FSuperRowType::FArguments()
 		.Padding(InArgs._Padding)
+		.OnDragDetected(this, &STedsTableViewerRow::HandleOnDragDetected)
+		.OnDragEnter(this, &STedsTableViewerRow::HandleOnDragEnter)
+		.OnAcceptDrop(this, &STedsTableViewerRow::HandleOnAcceptDrop)
+		.OnDragLeave(this, &STedsTableViewerRow::HandleOnDragLeave)
 		.Style(&FAppStyle::Get().GetWidgetStyle<FTableRowStyle>("SceneOutliner.TableViewRow"));
 
-		SMultiColumnTableRow<TableViewerItemPtr>::Construct(Args, OwnerTableView);
+		SMultiColumnTableRow<TableViewerItemPtr>::Construct(Args, MoveTemp(OwnerTableView));
 	}
 
 	TSharedRef<SWidget> STedsTableViewerRow::GenerateWidgetForColumn(const FName& ColumnName)
@@ -52,6 +60,109 @@ namespace UE::Editor::DataStorage
 		return ItemHeight.IsSet() ? ItemHeight.Get() : FOptionalSize();
 	}
 
+	void STedsTableViewerRow::HandleOnDragEnter(const FDragDropEvent& DragDropEvent)
+	{
+		if (ICoreProvider* Storage = TableViewerModel->GetDataStorageInterface())
+		{
+			if (UDropOperationsFactory* DropFactory = Storage->FindFactory<UDropOperationsFactory>())
+			{
+				if (TSharedPtr<Widgets::FTedsDragDropOp> Operation = DragDropEvent.GetOperationAs<Widgets::FTedsDragDropOp>())
+				{
+					RowHandle TargetRow = Item;
+		
+					for (RowHandle SourceRow : Operation->GetRows())
+					{
+						FScopedRow InputRow = DropFactory->CreateInputRow(TargetRow, SourceRow, true);
+
+						bool bResult = DropFactory->Test(InputRow.GetRow());
+
+						// Currently the description of the last operation will take precedence, we might want support for batched tests in the future
+						if (const FText* Description = DragAndDrop::Helpers::GetDescription(*Storage, InputRow.GetRow()))
+						{
+							Operation->CurrentHoverText = *Description;
+						}
+
+						// No need to test further if any operation failed
+						if (!bResult)
+						{
+							return;
+						}
+
+						Operation->SetDropStatus(Widgets::SomeRowsPass);
+					}
+
+					Operation->SetDropStatus(Widgets::AllRowsPass);
+				}
+			}
+		}
+
+	}
+
+	void STedsTableViewerRow::HandleOnDragLeave(const FDragDropEvent& DragDropEvent)
+	{
+		if (ICoreProvider* Storage = TableViewerModel->GetDataStorageInterface())
+		{
+			if (UDropOperationsFactory* DropFactory = Storage->FindFactory<UDropOperationsFactory>())
+			{
+				if (TSharedPtr<Widgets::FTedsDragDropOp> Operation = DragDropEvent.GetOperationAs<Widgets::FTedsDragDropOp>())
+				{
+					// Reset the drop test result since we are leaving this row
+					Operation->SetDropStatus(Widgets::NoTestsRun);
+				}
+			}
+		}
+	}
+
+	FReply STedsTableViewerRow::HandleOnAcceptDrop(const FDragDropEvent& DragDropEvent, EItemDropZone DropZone, TableViewerItemPtr InItem)
+	{
+		if (ICoreProvider* Storage = TableViewerModel->GetDataStorageInterface())
+		{
+			if (UDropOperationsFactory* DropFactory = Storage->FindFactory<UDropOperationsFactory>())
+			{
+				if (TSharedPtr<Widgets::FTedsDragDropOp> Operation = DragDropEvent.GetOperationAs<Widgets::FTedsDragDropOp>())
+				{
+					// We currently only allow drop if all rows pass the test
+					if (Operation->GetDropStatus() == Widgets::AllRowsPass)
+					{
+						RowHandle TargetRow = Item;
+						FText OutDescription;
+		
+						for (RowHandle SourceRow : Operation->GetRows())
+						{
+							FScopedRow InputRow = DropFactory->CreateInputRow(TargetRow, SourceRow, false);
+							DropFactory->Apply(InputRow.GetRow());
+						}
+
+						return FReply::Handled();
+					}
+				}
+			}
+		}
+
+		return FReply::Unhandled();
+	}
+
+	FReply STedsTableViewerRow::HandleOnDragDetected(const FGeometry& MyGeometry, const FPointerEvent& MouseEvent)
+	{
+		// TEDS UI TODO: some way to check whether this table viewer allows drag/drop? Column on widget row?
+		if (MouseEvent.IsMouseButtonDown(EKeys::LeftMouseButton))
+		{
+			if (TSharedPtr<const ITableViewer> TableViewer = WeakTableViewer.Pin())
+			{
+				FRowHandleArray Rows;
+				TableViewer->ForEachSelectedRow([&Rows](RowHandle Row)
+				{
+					Rows.Add(Row);
+				});
+				
+				TSharedRef<Widgets::FTedsDragDropOp> Operation = Widgets::FTedsDragDropOp::New(MoveTemp(Rows));
+				return FReply::Handled().BeginDragDrop(Operation);
+			}
+		}
+
+		return FReply::Unhandled();
+	}
+
 	TSharedRef<SWidget> SHierarchyViewerRow::GenerateWidgetForColumn(const FName& ColumnName)
 	{
 		TSharedRef<SWidget> ActualWidget = STedsTableViewerRow::GenerateWidgetForColumn(ColumnName);
diff --git a/Engine/Plugins/Experimental/EditorDataStorageFeatures/Source/TedsTableViewer/Public/Widgets/STedsTableViewerRow.h b/Engine/Plugins/Experimental/EditorDataStorageFeatures/Source/TedsTableViewer/Public/Widgets/STedsTableViewerRow.h
--- a/Engine/Plugins/Experimental/EditorDataStorageFeatures/Source/TedsTableViewer/Public/Widgets/STedsTableViewerRow.h
+++ b/Engine/Plugins/Experimental/EditorDataStorageFeatures/Source/TedsTableViewer/Public/Widgets/STedsTableViewerRow.h
@@ -6,6 +6,11 @@
 #include "Widgets/Views/SListView.h"
 #include "Widgets/Views/STableRow.h"
 
+namespace UE::Editor::DataStorage
+{
+	class ITableViewer;
+}
+
 namespace UE::Editor::DataStorage
 {
 	class STedsTableViewer;
@@ -36,7 +41,7 @@ namespace UE::Editor::DataStorage
 		SLATE_END_ARGS()
 		
 		/** Construct function for this widget */
-		void Construct( const FArguments& InArgs, const TSharedRef<STableViewBase>& OwnerTableView, const TSharedRef<FTedsTableViewerModel>& InTableViewerModel );
+		void Construct( const FArguments& InArgs, TSharedRef<STableViewBase> OwnerTableView, TSharedRef<FTedsTableViewerModel> InTableViewerModel, TWeakPtr<const ITableViewer> InTableViewer);
 
 		/** Overridden from SMultiColumnTableRow.  Generates a widget for this column of the tree row. */
 		virtual TSharedRef<SWidget> GenerateWidgetForColumn( const FName& ColumnName ) override;
@@ -45,10 +50,18 @@ namespace UE::Editor::DataStorage
 		/** Get the current ItemHeight as an OptionalSize for the Box HeightOverride */
 		FOptionalSize GetCurrentItemHeight() const;
 
+		/** Drag & Drop handlers */
+		void HandleOnDragEnter(const FDragDropEvent& DragDropEvent);
+		void HandleOnDragLeave(const FDragDropEvent& DragDropEvent);
+		FReply HandleOnDragDetected(const FGeometry& MyGeometry, const FPointerEvent& MouseEvent);
+		FReply HandleOnAcceptDrop(const FDragDropEvent& DragDropEvent, EItemDropZone DropZone, TableViewerItemPtr Item);
+
+
 	protected:
 		TAttribute<float> ItemHeight;
 		RowHandle ParentWidgetRowHandle = InvalidRowHandle;
 		TSharedPtr<FTedsTableViewerModel> TableViewerModel;
+		TWeakPtr<const ITableViewer> WeakTableViewer;
 		TableViewerItemPtr Item;
 	};
 
diff --git a/Engine/Plugins/Experimental/EditorDataStorageFeatures/Source/TedsTableViewer/TedsTableViewer.Build.cs b/Engine/Plugins/Experimental/EditorDataStorageFeatures/Source/TedsTableViewer/TedsTableViewer.Build.cs
--- a/Engine/Plugins/Experimental/EditorDataStorageFeatures/Source/TedsTableViewer/TedsTableViewer.Build.cs
+++ b/Engine/Plugins/Experimental/EditorDataStorageFeatures/Source/TedsTableViewer/TedsTableViewer.Build.cs
@@ -33,7 +33,9 @@ public class TedsTableViewer : ModuleRules
 					"SlateCore",
 					"WorkspaceMenuStructure",
 					"ToolMenus",
-					"ToolWidgets"
+					"ToolWidgets",
+					"TedsOperations",
+					"TedsDropOperations"
 				});
 		}
 	}
