Teds Drag Drop Op: Fix incorrect DRAG_DROP_OPERATOR_TYPE macro not matching the class name

diff --git a/Engine/Plugins/Experimental/EditorDataStorageFeatures/Source/TedsDropOperations/Public/TedsDragDropOp.h b/Engine/Plugins/Experimental/EditorDataStorageFeatures/Source/TedsDropOperations/Public/TedsDragDropOp.h
--- a/Engine/Plugins/Experimental/EditorDataStorageFeatures/Source/TedsDropOperations/Public/TedsDragDropOp.h
+++ b/Engine/Plugins/Experimental/EditorDataStorageFeatures/Source/TedsDropOperations/Public/TedsDragDropOp.h
@@ -22,7 +22,7 @@ namespace UE::Editor::DataStorage::Widgets
 	class FTedsDragDropOp : public FDecoratedDragDropOp
 	{
 	public:
-		DRAG_DROP_OPERATOR_TYPE(FTedsDragDropOperation, FDecoratedDragDropOp)
+		DRAG_DROP_OPERATOR_TYPE(FTedsDragDropOp, FDecoratedDragDropOp)
 
 		UE_API virtual ~FTedsDragDropOp() override = default;
 		UE_API static TSharedRef<FTedsDragDropOp> New(FRowHandleArray InDraggedRows);
