STedsTableViewer: Add support for renaming rows via column metadata + event column to execute the rename event

Teds Metadata: Modify columns metadata view to allow construction using FColumnMetaData array

diff --git a/Engine/Plugins/Experimental/EditorDataStorageFeatures/Source/TedsOutliner/Private/Widgets/OutlinerTextWidget.cpp b/Engine/Plugins/Experimental/EditorDataStorageFeatures/Source/TedsOutliner/Private/Widgets/OutlinerTextWidget.cpp
--- a/Engine/Plugins/Experimental/EditorDataStorageFeatures/Source/TedsOutliner/Private/Widgets/OutlinerTextWidget.cpp
+++ b/Engine/Plugins/Experimental/EditorDataStorageFeatures/Source/TedsOutliner/Private/Widgets/OutlinerTextWidget.cpp
@@ -3,7 +3,11 @@
 #include "Widgets/OutlinerTextWidget.h"
 
 #include "ActorEditorUtils.h"
+#include "ITedsTableViewer.h"
+#include "TedsTableViewerWidgetColumns.h"
+#include "TedsTableViewerWidgetColumns.h"
 #include "Columns/SlateDelegateColumns.h"
+#include "Elements/Columns/TypedElementHiearchyColumns.h"
 #include "Elements/Columns/TypedElementLabelColumns.h"
 #include "Elements/Columns/TypedElementMiscColumns.h"
 #include "Elements/Framework/TypedElementAttributeBinding.h"
@@ -16,6 +20,114 @@
 
 #define LOCTEXT_NAMESPACE "OutlinerTextWidget"
 
+namespace UE::Editor::Outliner::TextWidget::Private
+{
+	// We are using an inheritance based model with a custom SWidget class just so we can manage lifetimes of delegates easier using BindSP
+	class SOutlinerEditableTextWidget : public SInlineEditableTextBlock
+	{
+	public:
+		SLATE_BEGIN_ARGS(SOutlinerEditableTextWidget)
+			: _DataStorage(nullptr)
+			, _DataRow(DataStorage::InvalidRowHandle)
+			, _WidgetRow(DataStorage::InvalidRowHandle)
+		{}
+
+		SLATE_ARGUMENT(DataStorage::ICoreProvider*, DataStorage)
+		SLATE_ARGUMENT(DataStorage::RowHandle, DataRow)
+		SLATE_ARGUMENT(DataStorage::RowHandle, WidgetRow)
+	SLATE_END_ARGS()
+		
+		void Construct(const FArguments& InArgs)
+		{
+			DataStorage = InArgs._DataStorage;
+			DataRow = InArgs._DataRow;
+			WidgetRow = InArgs._WidgetRow;
+			
+			using namespace UE::Editor::DataStorage;
+			FAttributeBinder DataRowBinder(DataRow, DataStorage);
+			FAttributeBinder WidgetRowBinder(WidgetRow, DataStorage);
+
+			const bool bShouldUseMiddleEllipsis = GetDefault<UEditorStyleSettings>()->bEnableMiddleEllipsis;
+
+			SInlineEditableTextBlock::Construct(
+				SInlineEditableTextBlock::FArguments()
+				.OnTextCommitted_Static(&FOutlinerTextWidgetConstructor::OnCommitText, DataStorage, DataRow)
+				.OnVerifyTextChanged_Static(&FOutlinerTextWidgetConstructor::OnVerifyText)
+				.Text(DataRowBinder.BindText(&FTypedElementLabelColumn::Label))
+				.ToolTipText(DataRowBinder.BindText(&FTypedElementLabelColumn::Label))
+				.OverflowPolicy(bShouldUseMiddleEllipsis ? ETextOverflowPolicy::MiddleEllipsis : TOptional<ETextOverflowPolicy>())
+				.IsSelected(WidgetRowBinder.BindEvent(&FExternalWidgetSelectionColumn::IsSelected)));
+
+			BindTableViewerRenameDelegate();
+
+			AddMetadata(MakeShared<TTypedElementUiStyleOverrideCapability<SInlineEditableTextBlock>>(*this));
+		}
+		
+		virtual ~SOutlinerEditableTextWidget() override = default;
+
+	protected:
+
+		void BindTableViewerRenameDelegate()
+		{
+			// The Teds Outliner directly uses this column to enter edit mode
+			DataStorage->AddColumn<FWidgetEnterEditModeColumn>(WidgetRow, FWidgetEnterEditModeColumn{
+				.OnEnterEditMode = FSimpleDelegate::CreateLambda([this]()
+				{
+					EnterEditingMode();
+				})
+			});
+
+			// The table viewer, if present, is the parent widget of our widget row
+			if (FTableRowParentColumn* ParentColumn = DataStorage->GetColumn<FTableRowParentColumn>(WidgetRow))
+			{
+				if (DataStorage->IsRowAvailable(ParentColumn->Parent))
+				{
+					// Bind to selection change on the table viewer
+					if (DataStorage::FTableViewerSelectionChangedColumn* SelectionChangedColumn =
+						DataStorage->GetColumn<DataStorage::FTableViewerSelectionChangedColumn>(ParentColumn->Parent))
+					{
+						SelectionChangedColumn->OnSelectionChanged.AddSP(this, &SOutlinerEditableTextWidget::OnTableViewerSelectionChanged, ParentColumn->Parent);
+					}
+				}
+				
+			}
+		}
+
+		void OnTableViewerSelectionChanged(DataStorage::RowHandle TableViewerRow)
+		{
+			DataStorage::FTableViewerRenameSelectionColumn* RenameColumn = DataStorage->GetColumn<DataStorage::FTableViewerRenameSelectionColumn>(TableViewerRow);
+			DataStorage::FTableViewerReferenceColumn* ReferenceColumn = DataStorage->GetColumn<FTedsTableViewerReferenceColumn>(TableViewerRow);
+
+			// On selection change, if the data row for this widget is selected add it as 'active' for the rename. So when FTableViewerRenameSelectionColumn
+			// is called on the table viewer the event is propagated to this widget
+			if (RenameColumn && ReferenceColumn)
+			{
+				RenameColumn->OnRenameSelection.RemoveAll(this);
+				
+				if (TSharedPtr<DataStorage::ITableViewer> TableViewer = ReferenceColumn->TableViewerWeak.Pin())
+				{
+					TableViewer->ForEachSelectedRow([this, RenameColumn](DataStorage::RowHandle SelectedRow)
+					{
+						if (SelectedRow == DataRow)
+						{
+							RenameColumn->OnRenameSelection.AddSPLambda(this, [this]()
+								{
+									EnterEditingMode();
+								});
+						}
+					});
+				}
+			}
+		}
+
+	protected:
+
+		DataStorage::ICoreProvider* DataStorage;
+		DataStorage::RowHandle WidgetRow;
+		DataStorage::RowHandle DataRow;
+	};
+}
+
 void UOutlinerTextWidgetFactory::RegisterWidgetConstructors(UE::Editor::DataStorage::ICoreProvider& DataStorage,
 	UE::Editor::DataStorage::IUiProvider& DataStorageUi) const
 {
@@ -115,29 +227,11 @@ TSharedPtr<SInlineEditableTextBlock> FOutlinerTextWidgetConstructor::CreateEdita
 	UE::Editor::DataStorage::RowHandle TargetRow,
 	UE::Editor::DataStorage::RowHandle WidgetRow) const
 {
-	using namespace UE::Editor::DataStorage;
-	FAttributeBinder TargetRowBinder(TargetRow, DataStorage);
-	FAttributeBinder WidgetRowBinder(WidgetRow, DataStorage);
-
-	const bool bShouldUseMiddleEllipsis = GetDefault<UEditorStyleSettings>()->bEnableMiddleEllipsis;
-
-	TSharedPtr<SInlineEditableTextBlock> TextBlock = SNew(SInlineEditableTextBlock)
-		.OnTextCommitted_Static(&FOutlinerTextWidgetConstructor::OnCommitText, DataStorage, TargetRow)
-		.OnVerifyTextChanged_Static(&FOutlinerTextWidgetConstructor::OnVerifyText)
-		.Text(TargetRowBinder.BindText(&FTypedElementLabelColumn::Label))
-		.ToolTipText(TargetRowBinder.BindText(&FTypedElementLabelColumn::Label))
-		.OverflowPolicy(bShouldUseMiddleEllipsis ? ETextOverflowPolicy::MiddleEllipsis : TOptional<ETextOverflowPolicy>())
-		.IsSelected(WidgetRowBinder.BindEvent(&FExternalWidgetSelectionColumn::IsSelected));
-
-	DataStorage->AddColumn<FWidgetEnterEditModeColumn>(WidgetRow, FWidgetEnterEditModeColumn{
-		.OnEnterEditMode = FSimpleDelegate::CreateLambda([TextBlock]()
-		{
-			TextBlock->EnterEditingMode();
-		})
-	});
-
-	TextBlock->AddMetadata(MakeShared<TTypedElementUiStyleOverrideCapability<SInlineEditableTextBlock>>(*TextBlock));
-	return TextBlock;
+	
+	return SNew(UE::Editor::Outliner::TextWidget::Private::SOutlinerEditableTextWidget)
+			.DataStorage(DataStorage)
+			.DataRow(TargetRow)
+			.WidgetRow(WidgetRow);
 }
 
 TSharedPtr<STextBlock> FOutlinerTextWidgetConstructor::CreateNonEditableTextBlock(
diff --git a/Engine/Plugins/Experimental/EditorDataStorageFeatures/Source/TedsTableViewer/Private/TedsTableViewerColumn.cpp b/Engine/Plugins/Experimental/EditorDataStorageFeatures/Source/TedsTableViewer/Private/TedsTableViewerColumn.cpp
--- a/Engine/Plugins/Experimental/EditorDataStorageFeatures/Source/TedsTableViewer/Private/TedsTableViewerColumn.cpp
+++ b/Engine/Plugins/Experimental/EditorDataStorageFeatures/Source/TedsTableViewer/Private/TedsTableViewerColumn.cpp
@@ -22,12 +22,13 @@ namespace UE::Editor::DataStorage
 {
 	FTedsTableViewerColumn::FTedsTableViewerColumn(const FName& InColumnName, const TSharedPtr<FTypedElementWidgetConstructor>& InCellWidgetConstructor,
 		const TArray<TWeakObjectPtr<const UScriptStruct>>& InMatchedColumns, const TSharedPtr<FTypedElementWidgetConstructor>& InHeaderWidgetConstructor,
-		const FMetaData& InWidgetMetaData)
+		const FMetaData& InWidgetMetaData, const TArray<TSharedPtr<FColumnMetaData>>& InColumnMetaData)
 		: ColumnName(InColumnName)
 		, CellWidgetConstructor(InCellWidgetConstructor)
 		, HeaderWidgetConstructor(InHeaderWidgetConstructor)
 		, MatchedColumns(InMatchedColumns)
 		, WidgetMetaData(InWidgetMetaData)
+		, ColumnMetaData(InColumnMetaData)
 	{
 		Storage = GetMutableDataStorageFeature<ICoreProvider>(StorageFeatureName);
 		StorageUi = GetMutableDataStorageFeature<IUiProvider>(UiFeatureName);
@@ -43,12 +44,12 @@ namespace UE::Editor::DataStorage
 		// Retrieve sort functors to allow the content of the column to be used for sorting
 		if (HeaderWidgetConstructor)
 		{
-			ColumnSorters = HeaderWidgetConstructor->ConstructColumnSorters(Storage, StorageUi, FGenericMetaDataView(WidgetMetaData));
+			ColumnSorters = HeaderWidgetConstructor->ConstructColumnSorters(Storage, StorageUi, GetMetaData());
 		}
 		// Prefer to use the sort function in the header, but if it doesn't exist, use the one from the cell.
 		if (CellWidgetConstructor && ColumnSorters.IsEmpty())
 		{
-			ColumnSorters = CellWidgetConstructor->ConstructColumnSorters(Storage, StorageUi, FGenericMetaDataView(WidgetMetaData));
+			ColumnSorters = CellWidgetConstructor->ConstructColumnSorters(Storage, StorageUi, GetMetaData());
 		}
 
 		RegisterQueries();
@@ -88,7 +89,7 @@ namespace UE::Editor::DataStorage
 				WidgetRowSetupDelegate(*Storage, UiRowHandle);
 			}
 
-			RowWidget = StorageUi->ConstructWidget(UiRowHandle, *CellWidgetConstructor, FGenericMetaDataView(WidgetMetaData));
+			RowWidget = StorageUi->ConstructWidget(UiRowHandle, *CellWidgetConstructor, GetMetaData());
 		}
 
 		return RowWidget;
@@ -123,7 +124,7 @@ namespace UE::Editor::DataStorage
 			}
 
 			Widget = StorageUi->ConstructWidget(UiRowHandle, *HeaderWidgetConstructor,
-				FGenericMetaDataView(WidgetMetaData));
+				GetMetaData());
 		}
 		if (!Widget.IsValid())
 		{
@@ -142,11 +143,11 @@ namespace UE::Editor::DataStorage
 		const bool* bColumnInitialSortIsAscending;
 
 		{
-			UE::Editor::DataStorage::FMetaDataEntryView MetaDataEntryView = FGenericMetaDataView(WidgetMetaData).FindGeneric("bAllowSorting");
+			UE::Editor::DataStorage::FMetaDataEntryView MetaDataEntryView = GetMetaData().FindGeneric("bAllowSorting");
 		    bAllowSorting = MetaDataEntryView.TryGetExact<bool>();
 		}
 		{
-			UE::Editor::DataStorage::FMetaDataEntryView MetaDataEntryView = FGenericMetaDataView(WidgetMetaData).FindGeneric("bColumnInitialSortIsAscending");
+			UE::Editor::DataStorage::FMetaDataEntryView MetaDataEntryView = GetMetaData().FindGeneric("bColumnInitialSortIsAscending");
 			bColumnInitialSortIsAscending = MetaDataEntryView.TryGetExact<bool>();
 		}
 
@@ -398,7 +399,7 @@ namespace UE::Editor::DataStorage
 			// If a column was added and we are here, we need to re-create the widget
 			if(*bColumnAddedPtr)
 			{
-				if(const TSharedPtr<SWidget> RowWidget = StorageUi->ConstructInternalWidget(Row, *CellWidgetConstructor, FGenericMetaDataView(WidgetMetaData)))
+				if(const TSharedPtr<SWidget> RowWidget = StorageUi->ConstructInternalWidget(Row, *CellWidgetConstructor, GetMetaData()))
 				{
 					TedsWidget->SetContent(RowWidget.ToSharedRef());
 				}
@@ -410,7 +411,12 @@ namespace UE::Editor::DataStorage
 			}
 		}
 	}
-	
+
+	FComboMetaDataView<FGenericMetaDataView, FColumnsMetaDataView> FTedsTableViewerColumn::GetMetaData() const
+	{
+		return FComboMetaDataView(FGenericMetaDataView(WidgetMetaData)).Next(FColumnsMetaDataView(ColumnMetaData));
+	}
+
 	FName FTedsTableViewerColumn::GetColumnName() const
 	{
 		return ColumnName;
diff --git a/Engine/Plugins/Experimental/EditorDataStorageFeatures/Source/TedsTableViewer/Private/TedsTableViewerModel.cpp b/Engine/Plugins/Experimental/EditorDataStorageFeatures/Source/TedsTableViewer/Private/TedsTableViewerModel.cpp
--- a/Engine/Plugins/Experimental/EditorDataStorageFeatures/Source/TedsTableViewer/Private/TedsTableViewerModel.cpp
+++ b/Engine/Plugins/Experimental/EditorDataStorageFeatures/Source/TedsTableViewer/Private/TedsTableViewerModel.cpp
@@ -16,9 +16,11 @@ DEFINE_LOG_CATEGORY_STATIC(LogTableViewer, Log, All)
 namespace UE::Editor::DataStorage
 {
 	FTedsTableViewerModel::FTedsTableViewerModel(const TSharedPtr<QueryStack::IRowNode>& InRowQueryStack,
-		const TArray<TWeakObjectPtr<const UScriptStruct>>& InRequestedColumns, const IUiProvider::FPurposeID& InCellWidgetPurpose,
-		const IUiProvider::FPurposeID& InHeaderWidgetPurpose, const FIsItemVisible& InIsItemVisibleDelegate)
+		const TArray<TWeakObjectPtr<const UScriptStruct>>& InRequestedColumns, const TArray<TSharedPtr<FColumnMetaData>>& InColumnMetaData,
+		const IUiProvider::FPurposeID& InCellWidgetPurpose, const IUiProvider::FPurposeID& InHeaderWidgetPurpose,
+		const FIsItemVisible& InIsItemVisibleDelegate)
 		: RequestedTedsColumns(InRequestedColumns)
+		, ColumnMetaData(InColumnMetaData)
 		, CellWidgetPurpose(InCellWidgetPurpose)
 		, HeaderWidgetPurpose(InHeaderWidgetPurpose)
 		, IsItemVisible(InIsItemVisibleDelegate)
@@ -224,12 +226,16 @@ namespace UE::Editor::DataStorage
 				TSharedPtr<FTypedElementWidgetConstructor> CellConstructor(Constructor.Release());
 
 				TSharedPtr<FTypedElementWidgetConstructor> HeaderConstructor =
-					TableViewerUtils::CreateHeaderWidgetConstructor(*StorageUi, FMetaDataView(), MatchedColumns, HeaderWidgetPurposeRow);
+					TableViewerUtils::CreateHeaderWidgetConstructor(*StorageUi, FColumnsMetaDataView(ColumnMetaData), MatchedColumns, HeaderWidgetPurposeRow);
 			
 				FName NameId = TableViewerUtils::FindLongestMatchingName(MatchedColumns, IndexOffset);
 
 				TArray<TWeakObjectPtr<const UScriptStruct>> MatchedColumnsCopy(MatchedColumns);
-				TSharedRef<FTedsTableViewerColumn> Column = MakeShared<FTedsTableViewerColumn>(NameId, CellConstructor, MoveTemp(MatchedColumnsCopy), HeaderConstructor);
+
+				// Currently we are passing in metadata for all columns to individual columns, we might want to only pass in the metadata for the matched columns
+				// for each one
+				TSharedRef<FTedsTableViewerColumn> Column = MakeShared<FTedsTableViewerColumn>(NameId, CellConstructor, MoveTemp(MatchedColumnsCopy), HeaderConstructor,
+					FMetaData(), ColumnMetaData);
 			
 				Column->SetIsRowVisibleDelegate(FTedsTableViewerColumn::FIsRowVisible::CreateRaw(this, &FTedsTableViewerModel::IsRowVisible));
 				Column->SetSortDelegates(FTedsTableViewerColumn::FSortHandler
@@ -249,7 +255,7 @@ namespace UE::Editor::DataStorage
 
 		// Create the widget constructors for the columns
 		StorageUi->CreateWidgetConstructors(CellWidgetPurposeRow, MatchApproach::LongestMatch, ColumnsCopy, 
-			FMetaDataView(), ColumnConstructor);
+			FColumnsMetaDataView(ColumnMetaData), ColumnConstructor);
 		
 		// Add the actual UI columns in the order the Teds Columns were specified
 		for (const TWeakObjectPtr<const UScriptStruct>& ColumnType : RequestedTedsColumns)
diff --git a/Engine/Plugins/Experimental/EditorDataStorageFeatures/Source/TedsTableViewer/Private/Widgets/STedsHierarchyViewer.cpp b/Engine/Plugins/Experimental/EditorDataStorageFeatures/Source/TedsTableViewer/Private/Widgets/STedsHierarchyViewer.cpp
--- a/Engine/Plugins/Experimental/EditorDataStorageFeatures/Source/TedsTableViewer/Private/Widgets/STedsHierarchyViewer.cpp
+++ b/Engine/Plugins/Experimental/EditorDataStorageFeatures/Source/TedsTableViewer/Private/Widgets/STedsHierarchyViewer.cpp
@@ -3,6 +3,7 @@
 #include "Widgets/STedsHierarchyViewer.h"
 
 #include "TedsTableViewerColumn.h"
+#include "TedsTableViewerWidgetColumns.h"
 #include "Columns/SlateDelegateColumns.h"
 #include "DataStorage/Features.h"
 #include "Elements/Columns/TypedElementHiearchyColumns.h"
@@ -32,7 +33,7 @@ namespace UE::Editor::DataStorage
 
 		IUiProvider* StorageUi = GetMutableDataStorageFeature<IUiProvider>(UiFeatureName);
 
-		if (!ensureMsgf(StorageUi, TEXT("Cannot use SHierarchyVIewer without TEDS UI being initialized")))
+		if (!ensureMsgf(StorageUi, TEXT("Cannot use SHierarchyViewer without TEDS UI being initialized")))
 		{
 			return;
 		}
@@ -49,7 +50,7 @@ namespace UE::Editor::DataStorage
 			HeaderWidgetPurpose = IUiProvider::FPurposeInfo("General", "Header", NAME_None).GeneratePurposeID();
 		}
 
-		Model = MakeShared<FTedsTableViewerModel>(InArgs._AllNodeProvider, InArgs._Columns, CellWidgetPurpose, HeaderWidgetPurpose,
+		Model = MakeShared<FTedsTableViewerModel>(InArgs._AllNodeProvider, InArgs._Columns, InArgs._ColumnMetaData, CellWidgetPurpose, HeaderWidgetPurpose,
 			FTedsTableViewerModel::FIsItemVisible::CreateSP(this, &SHierarchyViewer::IsItemVisible));
 		
 		TopLevelRowsNode = MakeShared<QueryStack::FTopLevelRowsNode>(Model->GetDataStorageInterface(), HierarchyInterface, Model->GetRowNode());
@@ -127,7 +128,9 @@ namespace UE::Editor::DataStorage
 			{
 				// FHideRowFromUITag - The table viewer should not show up as a row in a table viewer because that will cause all sorts of recursion issues
 				// The others are Columns we are going to bind to attributes on STreeView
-				DataStorage->AddColumns<FHideRowFromUITag, FWidgetContextMenuColumn, FWidgetRowScrolledIntoView, FWidgetDoubleClickedColumn>(WidgetRowHandle);
+				DataStorage->AddColumns<FHideRowFromUITag, FWidgetContextMenuColumn, FWidgetRowScrolledIntoView,
+					FWidgetDoubleClickedColumn, FTableViewerSelectionChangedColumn, FTableViewerRenameSelectionColumn>(WidgetRowHandle);
+				DataStorage->AddColumn(WidgetRowHandle, FTableViewerReferenceColumn{.TableViewerWeak = SharedThis(this)});
 			}
 		}
 	}
@@ -196,6 +199,15 @@ namespace UE::Editor::DataStorage
 		{
 			OnSelectionChanged.Execute(Item);
 		}
+
+		// Some systems can only access the table viewer by the row handle so also fire the selection changed event on the column
+		if (ICoreProvider* Storage = Model->GetDataStorageInterface())
+		{
+			if (FTableViewerSelectionChangedColumn* Column = Storage->GetColumn<FTableViewerSelectionChangedColumn>(GetWidgetRowHandle()))
+			{
+				Column->OnSelectionChanged.Broadcast();
+			}
+		}
 	}
 
 	void SHierarchyViewer::SetColumns(const TArray<TWeakObjectPtr<const UScriptStruct>>& Columns)
@@ -270,6 +282,7 @@ namespace UE::Editor::DataStorage
 
 		return SNew(SHierarchyViewerRow, OwnerTable, Model.ToSharedRef(), SharedThis(this))
 				.Padding(FMargin(0))
+				.ParentWidgetRowHandle(GetWidgetRowHandle())
 				.Item(InItem);
 	}
 } // namespace UE::Editor::DataStorage
diff --git a/Engine/Plugins/Experimental/EditorDataStorageFeatures/Source/TedsTableViewer/Private/Widgets/STedsTableViewer.cpp b/Engine/Plugins/Experimental/EditorDataStorageFeatures/Source/TedsTableViewer/Private/Widgets/STedsTableViewer.cpp
--- a/Engine/Plugins/Experimental/EditorDataStorageFeatures/Source/TedsTableViewer/Private/Widgets/STedsTableViewer.cpp
+++ b/Engine/Plugins/Experimental/EditorDataStorageFeatures/Source/TedsTableViewer/Private/Widgets/STedsTableViewer.cpp
@@ -3,6 +3,7 @@
 #include "Widgets/STedsTableViewer.h"
 
 #include "TedsTableViewerColumn.h"
+#include "TedsTableViewerWidgetColumns.h"
 #include "Columns/SlateDelegateColumns.h"
 #include "DataStorage/Features.h"
 #include "Elements/Columns/TypedElementUIColumns.h"
@@ -23,12 +24,12 @@ namespace UE::Editor::DataStorage
 		ItemHeight = InArgs._ItemHeight;
 		ItemPadding = InArgs._ItemPadding;
 
-		IUiProvider::FPurposeID CellWigetPurpose = InArgs._CellWidgetPurpose;
+		IUiProvider::FPurposeID CellWidgetPurpose = InArgs._CellWidgetPurpose;
 
-		if (!CellWigetPurpose.IsSet())
+		if (!CellWidgetPurpose.IsSet())
 		{
 			IUiProvider* StorageUi = GetMutableDataStorageFeature<IUiProvider>(UiFeatureName);
-			CellWigetPurpose = StorageUi->GetGeneralWidgetPurposeID();
+			CellWidgetPurpose = StorageUi->GetGeneralWidgetPurposeID();
 		}
 
 		IUiProvider::FPurposeID HeaderWidgetPurpose = InArgs._HeaderWidgetPurpose;
@@ -38,8 +39,8 @@ namespace UE::Editor::DataStorage
 			HeaderWidgetPurpose = IUiProvider::FPurposeInfo("General", "Header", NAME_None).GeneratePurposeID();
 		}
 		
-		Model = MakeShared<FTedsTableViewerModel>(InArgs._QueryStack, InArgs._Columns, CellWigetPurpose, HeaderWidgetPurpose,
-			FTedsTableViewerModel::FIsItemVisible::CreateSP(this, &STedsTableViewer::IsItemVisible));
+		Model = MakeShared<FTedsTableViewerModel>(InArgs._QueryStack, InArgs._Columns, InArgs._ColumnMetaData,
+			CellWidgetPurpose, HeaderWidgetPurpose, FTedsTableViewerModel::FIsItemVisible::CreateSP(this, &STedsTableViewer::IsItemVisible));
 		
 		HeaderRowWidget = SNew( SHeaderRow )
 							.CanSelectGeneratedColumn(InArgs._CanSelectGeneratedColumn);
@@ -99,9 +100,9 @@ namespace UE::Editor::DataStorage
 				DataStorage->AddColumn(WidgetRowHandle, FHideRowFromUITag::StaticStruct());
 
 				// Columns we are going to bind to attributes on SListView
-				DataStorage->AddColumn(WidgetRowHandle, FWidgetContextMenuColumn::StaticStruct());
-				DataStorage->AddColumn(WidgetRowHandle, FWidgetRowScrolledIntoView::StaticStruct());
-				DataStorage->AddColumn(WidgetRowHandle, FWidgetDoubleClickedColumn::StaticStruct());
+				DataStorage->AddColumns<FWidgetContextMenuColumn, FWidgetRowScrolledIntoView,
+					FWidgetDoubleClickedColumn, FTableViewerSelectionChangedColumn, FTableViewerRenameSelectionColumn>(WidgetRowHandle);
+				DataStorage->AddColumn(WidgetRowHandle, FTableViewerReferenceColumn{.TableViewerWeak = SharedThis(this)});
 			}
 		}
 	}
@@ -151,13 +152,21 @@ namespace UE::Editor::DataStorage
 
 		CreateInternalWidget();
 	}
-
 	void STedsTableViewer::OnListSelectionChanged(TableViewerItemPtr Item, ESelectInfo::Type SelectInfo)
 	{
 		if(OnSelectionChanged.IsBound())
 		{
 			OnSelectionChanged.Execute(Item);
 		}
+
+		// Some systems can only access the table viewer by the row handle so also fire the selection changed event on the column
+		if (ICoreProvider* Storage = Model->GetDataStorageInterface())
+		{
+			if (FTableViewerSelectionChangedColumn* Column = Storage->GetColumn<FTableViewerSelectionChangedColumn>(GetWidgetRowHandle()))
+			{
+				Column->OnSelectionChanged.Broadcast();
+			}
+		}
 	}
 
 	void STedsTableViewer::SetColumns(const TArray<TWeakObjectPtr<const UScriptStruct>>& Columns)
diff --git a/Engine/Plugins/Experimental/EditorDataStorageFeatures/Source/TedsTableViewer/Private/Widgets/STedsTileViewer.cpp b/Engine/Plugins/Experimental/EditorDataStorageFeatures/Source/TedsTableViewer/Private/Widgets/STedsTileViewer.cpp
--- a/Engine/Plugins/Experimental/EditorDataStorageFeatures/Source/TedsTableViewer/Private/Widgets/STedsTileViewer.cpp
+++ b/Engine/Plugins/Experimental/EditorDataStorageFeatures/Source/TedsTableViewer/Private/Widgets/STedsTileViewer.cpp
@@ -37,8 +37,9 @@ namespace UE::Editor::DataStorage
 		// so we give FTedsTableViewerModel empty dummy columns and purposes
 		IUiProvider::FPurposeID DummyPurpose;
 		TArray<TWeakObjectPtr<const UScriptStruct>> DummyColumns;
+		TArray<TSharedPtr<FColumnMetaData>> DummyColumnMetaData;
 		
-		Model = MakeShared<FTedsTableViewerModel>(InArgs._QueryStack, DummyColumns, DummyPurpose, DummyPurpose,
+		Model = MakeShared<FTedsTableViewerModel>(InArgs._QueryStack, DummyColumns, DummyColumnMetaData, DummyPurpose, DummyPurpose,
 			FTedsTableViewerModel::FIsItemVisible::CreateSP(this, &STedsTileViewer::IsItemVisible));
 
 		IUiProvider* StorageUi = Model->GetDataStorageUiProvider();
@@ -217,6 +218,15 @@ namespace UE::Editor::DataStorage
 		{
 			OnSelectionChanged.Execute(Item);
 		}
+
+		// Some systems can only access the table viewer by the row handle so also fire the selection changed event on the column
+		if (ICoreProvider* Storage = Model->GetDataStorageInterface())
+		{
+			if (FTableViewerSelectionChangedColumn* Column = Storage->GetColumn<FTableViewerSelectionChangedColumn>(GetWidgetRowHandle()))
+			{
+				Column->OnSelectionChanged.Broadcast();
+			}
+		}
 	}
 
 	void STedsTileViewer::AddWidgetColumns()
@@ -231,9 +241,10 @@ namespace UE::Editor::DataStorage
 				DataStorage->AddColumn(WidgetRowHandle, FHideRowFromUITag::StaticStruct());
 
 				// Columns we are going to bind to attributes on SListView
-				DataStorage->AddColumn(WidgetRowHandle, FWidgetContextMenuColumn::StaticStruct());
-				DataStorage->AddColumn(WidgetRowHandle, FWidgetRowScrolledIntoView::StaticStruct());
-				DataStorage->AddColumn(WidgetRowHandle, FWidgetDoubleClickedColumn::StaticStruct());
+				DataStorage->AddColumns<FWidgetContextMenuColumn, FWidgetRowScrolledIntoView,
+					FWidgetDoubleClickedColumn, FTableViewerSelectionChangedColumn, FTableViewerRenameSelectionColumn>(WidgetRowHandle);
+				DataStorage->AddColumn(WidgetRowHandle, FTableViewerReferenceColumn{.TableViewerWeak = SharedThis(this)});
+
 			}
 		}
 	}
diff --git a/Engine/Plugins/Experimental/EditorDataStorageFeatures/Source/TedsTableViewer/Public/TedsTableViewerColumn.h b/Engine/Plugins/Experimental/EditorDataStorageFeatures/Source/TedsTableViewer/Public/TedsTableViewerColumn.h
--- a/Engine/Plugins/Experimental/EditorDataStorageFeatures/Source/TedsTableViewer/Public/TedsTableViewerColumn.h
+++ b/Engine/Plugins/Experimental/EditorDataStorageFeatures/Source/TedsTableViewer/Public/TedsTableViewerColumn.h
@@ -34,7 +34,8 @@ namespace UE::Editor::DataStorage
 			const TSharedPtr<FTypedElementWidgetConstructor>& InCellWidgetConstructor, // The widget constructor to use for this column
 			const TArray<TWeakObjectPtr<const UScriptStruct>>& InMatchedColumns = {}, // An optional list of matched Teds columns 
 			const TSharedPtr<FTypedElementWidgetConstructor>& InHeaderWidgetConstructor = nullptr, // Optional constructor to use for the header widget
-			const FMetaData& MetaData = FMetaData()); // Optional metadata to use when constructing widgets
+			const FMetaData& MetaData = FMetaData(), // Optional generic metadata to use when constructing widgets
+			const TArray<TSharedPtr<FColumnMetaData>>& ColumnMetaData = {});  // Optional metadata specific to the matched columns
 
 		TEDSTABLEVIEWER_API ~FTedsTableViewerColumn();
 		
@@ -77,6 +78,7 @@ namespace UE::Editor::DataStorage
 		void UnRegisterQueries();
 		bool IsRowVisible(const RowHandle InRowHandle) const;
 		void UpdateWidgets();
+		FComboMetaDataView<FGenericMetaDataView, FColumnsMetaDataView> GetMetaData() const;
 		
 	private:
 
@@ -101,6 +103,7 @@ namespace UE::Editor::DataStorage
 		
 		// The Metadata used to create widgets
 		const FMetaData WidgetMetaData;
+		const TArray<TSharedPtr<FColumnMetaData>> ColumnMetaData;
 
 		// TEDS Constructs
 		ICoreProvider* Storage;
diff --git a/Engine/Plugins/Experimental/EditorDataStorageFeatures/Source/TedsTableViewer/Public/TedsTableViewerModel.h b/Engine/Plugins/Experimental/EditorDataStorageFeatures/Source/TedsTableViewer/Public/TedsTableViewerModel.h
--- a/Engine/Plugins/Experimental/EditorDataStorageFeatures/Source/TedsTableViewer/Public/TedsTableViewerModel.h
+++ b/Engine/Plugins/Experimental/EditorDataStorageFeatures/Source/TedsTableViewer/Public/TedsTableViewerModel.h
@@ -41,6 +41,7 @@ namespace UE::Editor::DataStorage
 		TEDSTABLEVIEWER_API FTedsTableViewerModel(
 			const TSharedPtr<QueryStack::IRowNode>& RowQueryStack,
 			const TArray<TWeakObjectPtr<const UScriptStruct>>& RequestedColumns,
+			const TArray<TSharedPtr<FColumnMetaData>>& ColumnMetaData,
 			const IUiProvider::FPurposeID& CellWidgetPurpose, 
 			const IUiProvider::FPurposeID& HeaderWidgetPurpose, 
 			const FIsItemVisible& InIsItemVisibleDelegate);
@@ -138,6 +139,9 @@ namespace UE::Editor::DataStorage
 		// The initial TEDS columns the widget was requested to display
 		TArray<TWeakObjectPtr<const UScriptStruct>> RequestedTedsColumns;
 
+		// Metadata for RequestedTedsColumns
+		TArray<TSharedPtr<FColumnMetaData>> ColumnMetaData;
+
 		// The widget purposes used to create widgets in this table viewer
 		IUiProvider::FPurposeID CellWidgetPurpose;
 		IUiProvider::FPurposeID HeaderWidgetPurpose;
diff --git a/Engine/Plugins/Experimental/EditorDataStorageFeatures/Source/TedsTableViewer/Public/TedsTableViewerWidgetColumns.h b/Engine/Plugins/Experimental/EditorDataStorageFeatures/Source/TedsTableViewer/Public/TedsTableViewerWidgetColumns.h
--- a/Engine/Plugins/Experimental/EditorDataStorageFeatures/Source/TedsTableViewer/Public/TedsTableViewerWidgetColumns.h
+++ b/Engine/Plugins/Experimental/EditorDataStorageFeatures/Source/TedsTableViewer/Public/TedsTableViewerWidgetColumns.h
@@ -4,6 +4,7 @@
 
 #include "DataStorage/CommonTypes.h"
 #include "Framework/SlateDelegates.h"
+#include "ITedsTableViewer.h"
 
 #include "TedsTableViewerWidgetColumns.generated.h"
 
@@ -24,3 +25,44 @@ struct FExternalWidgetExclusiveSelectionColumn : public FEditorDataStorageColumn
 	 */
 	FIsSelected IsSelectedExclusively;
 };
+
+/**
+ * Column on table viewer row that contains a reference to the ITableViewer
+ */
+USTRUCT(meta = (DisplayName = "Table Viewer reference"))
+struct FTedsTableViewerReferenceColumn : public FEditorDataStorageColumn
+{
+	GENERATED_BODY()
+	
+	TWeakPtr<UE::Editor::DataStorage::ITableViewer> TableViewerWeak;
+};
+
+/**
+ * Column on the table viewer row with the delegate that executes when selection is changed
+ */
+USTRUCT(meta = (DisplayName = "On Table Viewer selection changed"))
+struct FTedsTableViewerSelectionChangedColumn : public FEditorDataStorageColumn
+{
+	GENERATED_BODY()
+	
+	FSimpleMulticastDelegate OnSelectionChanged;
+};
+
+/**
+ * Column on the table viewer row with the delegate that executes to rename the selected rows
+ * TEDS UI TODO: Can we use a generic column for actions/events instead of one hardcoded for rename?
+ */
+USTRUCT(meta = (DisplayName = "On Table Viewer selection changed"))
+struct FTedsTableViewerRenameSelectionColumn : public FEditorDataStorageColumn
+{
+	GENERATED_BODY()
+	
+	FSimpleMulticastDelegate OnRenameSelection;
+};
+
+namespace UE::Editor::DataStorage
+{
+	using FTableViewerReferenceColumn = FTedsTableViewerReferenceColumn;
+	using FTableViewerSelectionChangedColumn = FTedsTableViewerSelectionChangedColumn;
+	using FTableViewerRenameSelectionColumn = FTedsTableViewerRenameSelectionColumn;
+}
\ No newline at end of file
diff --git a/Engine/Plugins/Experimental/EditorDataStorageFeatures/Source/TedsTableViewer/Public/Widgets/STedsHierarchyViewer.h b/Engine/Plugins/Experimental/EditorDataStorageFeatures/Source/TedsTableViewer/Public/Widgets/STedsHierarchyViewer.h
--- a/Engine/Plugins/Experimental/EditorDataStorageFeatures/Source/TedsTableViewer/Public/Widgets/STedsHierarchyViewer.h
+++ b/Engine/Plugins/Experimental/EditorDataStorageFeatures/Source/TedsTableViewer/Public/Widgets/STedsHierarchyViewer.h
@@ -61,9 +61,11 @@ namespace UE::Editor::DataStorage
 		SLATE_ARGUMENT(TSharedPtr<QueryStack::IRowNode>, AllNodeProvider)
 
 		// The Columns that this table viewer will display
-		// Table Viewer TODO: How do we specify column metadata (ReadOnly or ReadWrite)?
 		SLATE_ARGUMENT(TArray<TWeakObjectPtr<const UScriptStruct>>, Columns)
 
+		// Metadata about the columns the table viewer is displaying
+		SLATE_ARGUMENT(TArray<TSharedPtr<FColumnMetaData>>, ColumnMetaData)
+
 		// The widget purpose to use to create the cell widgets
 		SLATE_ARGUMENT(IUiProvider::FPurposeID, CellWidgetPurpose)
 			
@@ -82,7 +84,7 @@ namespace UE::Editor::DataStorage
 
 		// Whether all rows start as collapsed or expanded by default
 		SLATE_ARGUMENT(EExpansionState, DefaultExpansionState)
-
+		
 		SLATE_END_ARGS()
 
 	public:
diff --git a/Engine/Plugins/Experimental/EditorDataStorageFeatures/Source/TedsTableViewer/Public/Widgets/STedsTableViewer.h b/Engine/Plugins/Experimental/EditorDataStorageFeatures/Source/TedsTableViewer/Public/Widgets/STedsTableViewer.h
--- a/Engine/Plugins/Experimental/EditorDataStorageFeatures/Source/TedsTableViewer/Public/Widgets/STedsTableViewer.h
+++ b/Engine/Plugins/Experimental/EditorDataStorageFeatures/Source/TedsTableViewer/Public/Widgets/STedsTableViewer.h
@@ -52,9 +52,11 @@ namespace UE::Editor::DataStorage
 		SLATE_ARGUMENT(TSharedPtr<QueryStack::IRowNode>, QueryStack)
 
 		// The Columns that this table viewer will display
-		// Table Viewer TODO: How do we specify column metadata (ReadOnly or ReadWrite)?
 		SLATE_ARGUMENT(TArray<TWeakObjectPtr<const UScriptStruct>>, Columns)
 
+		// Metadata about the columns the table viewer is displaying
+		SLATE_ARGUMENT(TArray<TSharedPtr<FColumnMetaData>>, ColumnMetaData)
+
 		// The widget purpose to use to create the cell widgets
 		SLATE_ARGUMENT(IUiProvider::FPurposeID, CellWidgetPurpose)
 			
diff --git a/Engine/Source/Runtime/TypedElementFramework/Private/Elements/Framework/TypedElementMetaData.cpp b/Engine/Source/Runtime/TypedElementFramework/Private/Elements/Framework/TypedElementMetaData.cpp
--- a/Engine/Source/Runtime/TypedElementFramework/Private/Elements/Framework/TypedElementMetaData.cpp
+++ b/Engine/Source/Runtime/TypedElementFramework/Private/Elements/Framework/TypedElementMetaData.cpp
@@ -83,6 +83,11 @@ namespace UE::Editor::DataStorage
 		return Result;
 	}
 
+	const UScriptStruct* FColumnMetaData::GetColumnType() const
+	{
+		return ColumnType;
+	}
+
 	bool FColumnMetaData::IsEmpty() const
 	{
 		return FMetaDataBase::IsEmpty() && ColumnType == nullptr;
@@ -181,6 +186,11 @@ namespace UE::Editor::DataStorage
 	{
 	}
 
+	FColumnsMetaDataView::FColumnsMetaDataView(TConstArrayView<TSharedPtr<FColumnMetaData>> InColumnsMetaData)
+		: ColumnsMetaData(InColumnsMetaData)
+	{
+	}
+
 	FMetaDataEntryView FColumnsMetaDataView::FindForColumn(
 		TWeakObjectPtr<const UScriptStruct> Column, FName AttributeName) const
 	{
@@ -196,6 +206,17 @@ namespace UE::Editor::DataStorage
 			}
 		}
 #endif
+
+		const TSharedPtr<FColumnMetaData>* MetaData = ColumnsMetaData.FindByPredicate([Column](const TSharedPtr<FColumnMetaData>& MetaData)
+			{
+				return MetaData && MetaData->GetColumnType() == Column;
+			});
+		
+		if (MetaData && *MetaData)
+		{
+			return (*MetaData)->Find(AttributeName);
+		}
+		
 		return FMetaDataEntryView();
 	}
 
diff --git a/Engine/Source/Runtime/TypedElementFramework/Public/Elements/Framework/TypedElementMetaData.h b/Engine/Source/Runtime/TypedElementFramework/Public/Elements/Framework/TypedElementMetaData.h
--- a/Engine/Source/Runtime/TypedElementFramework/Public/Elements/Framework/TypedElementMetaData.h
+++ b/Engine/Source/Runtime/TypedElementFramework/Public/Elements/Framework/TypedElementMetaData.h
@@ -95,6 +95,8 @@ namespace UE::Editor::DataStorage
 		TYPEDELEMENTFRAMEWORK_API FColumnMetaData(const UScriptStruct* InColumnType, EFlags InFlags);
 
 		TYPEDELEMENTFRAMEWORK_API FMetaDataEntryView Find(FName Name) const override;
+
+		TYPEDELEMENTFRAMEWORK_API const UScriptStruct* GetColumnType() const;
 		
 		TYPEDELEMENTFRAMEWORK_API bool IsEmpty() const override;
 
@@ -147,11 +149,13 @@ namespace UE::Editor::DataStorage
 		~FColumnsMetaDataView() override = default;
 
 		TYPEDELEMENTFRAMEWORK_API explicit FColumnsMetaDataView(TConstArrayView<TWeakObjectPtr<const UScriptStruct>> InColumns);
+		TYPEDELEMENTFRAMEWORK_API explicit FColumnsMetaDataView(TConstArrayView<TSharedPtr<FColumnMetaData>> InColumnsMetaData);
 		TYPEDELEMENTFRAMEWORK_API FMetaDataEntryView FindForColumn(
 			TWeakObjectPtr<const UScriptStruct> Column, FName AttributeName) const override;
 
 	private:
 		TConstArrayView<TWeakObjectPtr<const UScriptStruct>> Columns;
+		TConstArrayView<TSharedPtr<FColumnMetaData>> ColumnsMetaData;
 	};
 
 	/**
