OutlinerLabelWidget: Also register the widget for the SceneOutliner.Cell purpose

diff --git a/Engine/Plugins/Experimental/EditorDataStorageFeatures/Source/TedsOutliner/Private/Widgets/OutlinerLabelWidget.cpp b/Engine/Plugins/Experimental/EditorDataStorageFeatures/Source/TedsOutliner/Private/Widgets/OutlinerLabelWidget.cpp
--- a/Engine/Plugins/Experimental/EditorDataStorageFeatures/Source/TedsOutliner/Private/Widgets/OutlinerLabelWidget.cpp
+++ b/Engine/Plugins/Experimental/EditorDataStorageFeatures/Source/TedsOutliner/Private/Widgets/OutlinerLabelWidget.cpp
@@ -20,6 +20,11 @@ void UOutlinerLabelWidgetFactory::RegisterWidgetConstructors(UE::Editor::DataSto
 	DataStorageUi.RegisterWidgetFactory<FOutlinerLabelWidgetConstructor>(
 		DataStorageUi.FindPurpose(UE::Editor::DataStorage::IUiProvider::FPurposeInfo("SceneOutliner", "RowLabel", NAME_None).GeneratePurposeID()),
 		TColumn<FTypedElementLabelColumn>() && (TColumn<FTypedElementClassTypeInfoColumn>() || TColumn<FTypedElementScriptStructTypeInfoColumn>()));
+
+	DataStorageUi.RegisterWidgetFactory<FOutlinerLabelWidgetConstructor>(
+		DataStorageUi.FindPurpose(UE::Editor::DataStorage::IUiProvider::FPurposeInfo("SceneOutliner", "Cell", NAME_None).GeneratePurposeID()),
+		TColumn<FTypedElementLabelColumn>() && (TColumn<FTypedElementClassTypeInfoColumn>() || TColumn<FTypedElementScriptStructTypeInfoColumn>()));
+
 }
 
 void UOutlinerLabelWidgetFactory::RegisterWidgetPurposes(UE::Editor::DataStorage::IUiProvider& DataStorageUi) const
