FQueryNode: Hack to not access data storage in the destructor to work around cases where systems hold onto a query node post data storage destruction

diff --git a/Engine/Plugins/Experimental/EditorDataStorageFeatures/Source/TedsQueryStack/Private/TedsQueryNode.cpp b/Engine/Plugins/Experimental/EditorDataStorageFeatures/Source/TedsQueryStack/Private/TedsQueryNode.cpp
--- a/Engine/Plugins/Experimental/EditorDataStorageFeatures/Source/TedsQueryStack/Private/TedsQueryNode.cpp
+++ b/Engine/Plugins/Experimental/EditorDataStorageFeatures/Source/TedsQueryStack/Private/TedsQueryNode.cpp
@@ -19,7 +19,12 @@ namespace UE::Editor::DataStorage::QueryStack
 
 	FQueryNode::~FQueryNode()
 	{
-		Storage.UnregisterQuery(QueryHandle);
+		// Hack to work around cases where systems that persist after UObject destruction hold onto a query node, which causes a crash since
+		// Storage is a UObject
+		if (UObjectInitialized() && !IsEngineExitRequested())
+		{
+			Storage.UnregisterQuery(QueryHandle);
+		}
 	}
 
 	void FQueryNode::SetQuery(UE::Editor::DataStorage::FQueryDescription Query)
