[TEDS Viewers] Remove GetModel function from hierarchy viewer and add SetQueryStack function to viewers to allow the model to remain private

diff --git a/Engine/Plugins/Experimental/EditorDataStorageFeatures/Source/TedsTableViewer/Private/Widgets/Composite/STedsCompositeHierarchyViewer.cpp b/Engine/Plugins/Experimental/EditorDataStorageFeatures/Source/TedsTableViewer/Private/Widgets/Composite/STedsCompositeHierarchyViewer.cpp
--- a/Engine/Plugins/Experimental/EditorDataStorageFeatures/Source/TedsTableViewer/Private/Widgets/Composite/STedsCompositeHierarchyViewer.cpp
+++ b/Engine/Plugins/Experimental/EditorDataStorageFeatures/Source/TedsTableViewer/Private/Widgets/Composite/STedsCompositeHierarchyViewer.cpp
@@ -45,11 +45,11 @@ namespace UE::Editor::DataStorage
 					.UseSectionsForCategories(InArgs._UseSectionsForCategories)
 					.OnPostFiltersChanged_Lambda([this]()
 					{
-						if (Model && FilterNode.IsValid())
+						if (HierarchyViewer.IsValid() && FilterNode.IsValid())
 						{
 							// Forward the updated FilterNode to the Model QueryStack since it holds a static copy to the old filter node
 							// Only called when the filters have changed, meaning the query stack has to be recomputed
-							Model->SetQueryStack(FilterNode);
+							HierarchyViewer->SetQueryStack(FilterNode);
 						}
 					});
 
@@ -103,9 +103,16 @@ namespace UE::Editor::DataStorage
 		}
 
 		HierarchyViewerArgs.AllNodeProvider(QueryStackNode);
-		HierarchyViewer->Construct(HierarchyViewerArgs, MoveTemp(InHierarchyInterface));
-
-		Model = HierarchyViewer->GetModel();
+		SAssignNew(HierarchyViewer, SHierarchyViewer, MoveTemp(InHierarchyInterface))
+			.AllNodeProvider(HierarchyViewerArgs._AllNodeProvider)
+			.Columns(HierarchyViewerArgs._Columns)
+			.EmptyRowsMessage(HierarchyViewerArgs._EmptyRowsMessage)
+			.CellWidgetPurpose(HierarchyViewerArgs._CellWidgetPurpose)
+			.HeaderWidgetPurpose(HierarchyViewerArgs._HeaderWidgetPurpose)
+			.ColumnMetaData(HierarchyViewerArgs._ColumnMetaData)
+			.ListSelectionMode(HierarchyViewerArgs._ListSelectionMode)
+			.DefaultExpansionState(HierarchyViewerArgs._DefaultExpansionState)
+			.OnSelectionChanged(HierarchyViewerArgs._OnSelectionChanged);
 		
 		ContentWidget->AddSlot()
 		[
@@ -118,6 +125,11 @@ namespace UE::Editor::DataStorage
 		];
 	}
 
+	void STedsCompositeHierarchyViewer::SetQueryStack(TSharedPtr<QueryStack::IRowNode> InRowQueryStack)
+	{
+		HierarchyViewer->SetQueryStack(InRowQueryStack);
+	}
+
 	void STedsCompositeHierarchyViewer::SetColumns(const TArray<TWeakObjectPtr<const UScriptStruct>>& InColumns)
 	{
 		HierarchyViewer->SetColumns(InColumns);
@@ -155,7 +167,7 @@ namespace UE::Editor::DataStorage
 
 	TSharedRef<SWidget> STedsCompositeHierarchyViewer::AsWidget()
 	{
-		return AsShared();
+		return HierarchyViewer->AsWidget();
 	}
 
 	bool STedsCompositeHierarchyViewer::IsSelected(RowHandle InRow) const
diff --git a/Engine/Plugins/Experimental/EditorDataStorageFeatures/Source/TedsTableViewer/Private/Widgets/STedsHierarchyViewer.cpp b/Engine/Plugins/Experimental/EditorDataStorageFeatures/Source/TedsTableViewer/Private/Widgets/STedsHierarchyViewer.cpp
--- a/Engine/Plugins/Experimental/EditorDataStorageFeatures/Source/TedsTableViewer/Private/Widgets/STedsHierarchyViewer.cpp
+++ b/Engine/Plugins/Experimental/EditorDataStorageFeatures/Source/TedsTableViewer/Private/Widgets/STedsHierarchyViewer.cpp
@@ -113,11 +113,6 @@ namespace UE::Editor::DataStorage
 		CreateInternalWidget();
 	}
 
-	TSharedPtr<FTedsTableViewerModel> SHierarchyViewer::GetModel() const
-	{
-		return Model;
-	}
-
 	void SHierarchyViewer::AddWidgetColumns()
 	{
 		if(ICoreProvider* DataStorage = Model->GetDataStorageInterface())
@@ -210,9 +205,14 @@ namespace UE::Editor::DataStorage
 		}
 	}
 
-	void SHierarchyViewer::SetColumns(const TArray<TWeakObjectPtr<const UScriptStruct>>& Columns)
+	void SHierarchyViewer::SetQueryStack(TSharedPtr<QueryStack::IRowNode> InRowQueryStack)
+	{
+		Model->SetQueryStack(InRowQueryStack);
+	}
+
+	void SHierarchyViewer::SetColumns(const TArray<TWeakObjectPtr<const UScriptStruct>>& InColumns)
 	{
-		Model->SetColumns(Columns);
+		Model->SetColumns(InColumns);
 		RefreshColumnWidgets();
 	}
 
diff --git a/Engine/Plugins/Experimental/EditorDataStorageFeatures/Source/TedsTableViewer/Private/Widgets/STedsTableViewer.cpp b/Engine/Plugins/Experimental/EditorDataStorageFeatures/Source/TedsTableViewer/Private/Widgets/STedsTableViewer.cpp
--- a/Engine/Plugins/Experimental/EditorDataStorageFeatures/Source/TedsTableViewer/Private/Widgets/STedsTableViewer.cpp
+++ b/Engine/Plugins/Experimental/EditorDataStorageFeatures/Source/TedsTableViewer/Private/Widgets/STedsTableViewer.cpp
@@ -169,9 +169,14 @@ namespace UE::Editor::DataStorage
 		}
 	}
 
-	void STedsTableViewer::SetColumns(const TArray<TWeakObjectPtr<const UScriptStruct>>& Columns)
+	void STedsTableViewer::SetQueryStack(TSharedPtr<QueryStack::IRowNode> InRowQueryStack)
 	{
-		Model->SetColumns(Columns);
+		Model->SetQueryStack(InRowQueryStack);
+	}
+
+	void STedsTableViewer::SetColumns(const TArray<TWeakObjectPtr<const UScriptStruct>>& InColumns)
+	{
+		Model->SetColumns(InColumns);
 		RefreshColumnWidgets();
 	}
 
diff --git a/Engine/Plugins/Experimental/EditorDataStorageFeatures/Source/TedsTableViewer/Private/Widgets/STedsTileViewer.cpp b/Engine/Plugins/Experimental/EditorDataStorageFeatures/Source/TedsTableViewer/Private/Widgets/STedsTileViewer.cpp
--- a/Engine/Plugins/Experimental/EditorDataStorageFeatures/Source/TedsTableViewer/Private/Widgets/STedsTileViewer.cpp
+++ b/Engine/Plugins/Experimental/EditorDataStorageFeatures/Source/TedsTableViewer/Private/Widgets/STedsTileViewer.cpp
@@ -168,6 +168,11 @@ namespace UE::Editor::DataStorage
 	{
 		return IsSelected(InRow) && TileView->GetNumItemsSelected() == 1;
 	}
+	
+	void STedsTileViewer::SetQueryStack(TSharedPtr<QueryStack::IRowNode> InRowQueryStack)
+	{
+		Model->SetQueryStack(InRowQueryStack);
+	}
 
 	void STedsTileViewer::SetColumns(const TArray<TWeakObjectPtr<const UScriptStruct>>& InColumns)
 	{
diff --git a/Engine/Plugins/Experimental/EditorDataStorageFeatures/Source/TedsTableViewer/Public/ITedsTableViewer.h b/Engine/Plugins/Experimental/EditorDataStorageFeatures/Source/TedsTableViewer/Public/ITedsTableViewer.h
--- a/Engine/Plugins/Experimental/EditorDataStorageFeatures/Source/TedsTableViewer/Public/ITedsTableViewer.h
+++ b/Engine/Plugins/Experimental/EditorDataStorageFeatures/Source/TedsTableViewer/Public/ITedsTableViewer.h
@@ -2,6 +2,7 @@
 #pragma once
 
 #include "DataStorage/Handles.h"
+#include "TedsQueryStackInterfaces.h"
 #include "Widgets/SWidget.h"
 
 namespace ESelectInfo { enum Type : int; }
@@ -43,8 +44,11 @@ namespace UE::Editor::DataStorage
     	// This usually means a new column for column-based views
     	virtual void AddCustomRowWidget(const TSharedRef<FTedsTableViewerColumn>& InRowWidgetFactory) = 0;
 
+    	// Clear the current QueryStack being displayed, set it to the given node, and recreate the sorting nodes
+    	virtual void SetQueryStack(TSharedPtr<QueryStack::IRowNode> InRowQueryStack) = 0;
+
     	// Clear the current list of columns being displayed and set it to the given list
-    	virtual void SetColumns(const TArray<TWeakObjectPtr<const UScriptStruct>>& Columns) = 0;
+    	virtual void SetColumns(const TArray<TWeakObjectPtr<const UScriptStruct>>& InColumns) = 0;
    
     	virtual ~ITableViewer() = default;
     };
diff --git a/Engine/Plugins/Experimental/EditorDataStorageFeatures/Source/TedsTableViewer/Public/Widgets/Composite/STedsCompositeHierarchyViewer.h b/Engine/Plugins/Experimental/EditorDataStorageFeatures/Source/TedsTableViewer/Public/Widgets/Composite/STedsCompositeHierarchyViewer.h
--- a/Engine/Plugins/Experimental/EditorDataStorageFeatures/Source/TedsTableViewer/Public/Widgets/Composite/STedsCompositeHierarchyViewer.h
+++ b/Engine/Plugins/Experimental/EditorDataStorageFeatures/Source/TedsTableViewer/Public/Widgets/Composite/STedsCompositeHierarchyViewer.h
@@ -30,7 +30,7 @@ namespace UE::Editor::DataStorage
      *			.CellWidgetPurpose(PurposeId)
 	 *		.Filters(MyCustomFilterArray);
 	 */
-	class STedsCompositeHierarchyViewer : public SHierarchyViewer
+	class STedsCompositeHierarchyViewer :  public SCompoundWidget, public ITableViewer
 	{
 	public:
 
@@ -66,6 +66,9 @@ namespace UE::Editor::DataStorage
 		
 		TEDSTABLEVIEWER_API void Construct(const FArguments& InArgs, TSharedPtr<IHierarchyViewerDataInterface> InHierarchyInterface);
 
+		// Clear the current QueryStack being displayed, set it to the given node, and recreate the sorting nodes 
+		TEDSTABLEVIEWER_API virtual void SetQueryStack(TSharedPtr<QueryStack::IRowNode> InRowQueryStack) override;
+
 		// Clear the current list of columns being displayed and set it to the given list
 		TEDSTABLEVIEWER_API virtual void SetColumns(const TArray<TWeakObjectPtr<const UScriptStruct>>& InColumns) override;
 
@@ -94,15 +97,6 @@ namespace UE::Editor::DataStorage
 		TEDSTABLEVIEWER_API virtual bool IsSelectedExclusively(RowHandle InRow) const override;
 
 	private:
-		// The actual height of each item
-		TAttribute<float> ItemHeight;
-
-		// The actual padding between each item
-		TAttribute<FMargin> ItemPadding;
-
-		// Our model class
-		TSharedPtr<FTedsTableViewerModel> Model;
-
 		// Row Node passed to the SearchBox that receives the searched result
 		TSharedPtr<QueryStack::IRowNode> SearchNode;
 
@@ -116,12 +110,6 @@ namespace UE::Editor::DataStorage
 		TSharedPtr<STedsFilterBar> FilterBar;
 		
 		// The actual HierarchyViewer widget that displays the rows
-		TSharedPtr<SHierarchyViewer> HierarchyViewer = MakeShared<SHierarchyViewer>();
-
-		// Delegate fired when the selection changes
-		FOnSelectionChanged OnSelectionChanged;
-
-		// The message to show in place of the table viewer when there are no rows provided by the current query stack
-		TAttribute<FText> EmptyRowsMessage;
+		TSharedPtr<SHierarchyViewer> HierarchyViewer;
 	};
 }
\ No newline at end of file
diff --git a/Engine/Plugins/Experimental/EditorDataStorageFeatures/Source/TedsTableViewer/Public/Widgets/STedsHierarchyViewer.h b/Engine/Plugins/Experimental/EditorDataStorageFeatures/Source/TedsTableViewer/Public/Widgets/STedsHierarchyViewer.h
--- a/Engine/Plugins/Experimental/EditorDataStorageFeatures/Source/TedsTableViewer/Public/Widgets/STedsHierarchyViewer.h
+++ b/Engine/Plugins/Experimental/EditorDataStorageFeatures/Source/TedsTableViewer/Public/Widgets/STedsHierarchyViewer.h
@@ -93,8 +93,11 @@ namespace UE::Editor::DataStorage
 			const FArguments& InArgs,
 			TSharedPtr<IHierarchyViewerDataInterface> InHierarchyInterface);
 
+		// Clear the current QueryStack being displayed, set it to the given node, and recreate the sorting nodes 
+		TEDSTABLEVIEWER_API virtual void SetQueryStack(TSharedPtr<QueryStack::IRowNode> InRowQueryStack) override;
+		
 		// Clear the current list of columns being displayed and set it to the given list
-		TEDSTABLEVIEWER_API virtual void SetColumns(const TArray<TWeakObjectPtr<const UScriptStruct>>& Columns) override;
+		TEDSTABLEVIEWER_API virtual void SetColumns(const TArray<TWeakObjectPtr<const UScriptStruct>>& InColumns) override;
 
     	// Add a custom per-row widget to the table viewer (that doesn't necessarily map to a TEDS column)
 		// This means a new column for the hierarchy viewer
@@ -122,8 +125,6 @@ namespace UE::Editor::DataStorage
 		// Get if a specific row is selected exclusively
 		TEDSTABLEVIEWER_API virtual bool IsSelectedExclusively(RowHandle InRow) const override;
 
-		TEDSTABLEVIEWER_API TSharedPtr<FTedsTableViewerModel> GetModel() const;
-
 	protected:
 		
 		TSharedRef<ITableRow> MakeTableRowWidget(TableViewerItemPtr InItem, const TSharedRef<STableViewBase>& OwnerTable) const;
diff --git a/Engine/Plugins/Experimental/EditorDataStorageFeatures/Source/TedsTableViewer/Public/Widgets/STedsTableViewer.h b/Engine/Plugins/Experimental/EditorDataStorageFeatures/Source/TedsTableViewer/Public/Widgets/STedsTableViewer.h
--- a/Engine/Plugins/Experimental/EditorDataStorageFeatures/Source/TedsTableViewer/Public/Widgets/STedsTableViewer.h
+++ b/Engine/Plugins/Experimental/EditorDataStorageFeatures/Source/TedsTableViewer/Public/Widgets/STedsTableViewer.h
@@ -63,7 +63,7 @@ namespace UE::Editor::DataStorage
 		// The widget purpose to use to create the header widgets
 		SLATE_ARGUMENT(IUiProvider::FPurposeID, HeaderWidgetPurpose)
 
-		// Whether or not the user can select the columns generated by right clicking on the header menu.
+		// Whether the user can select the columns generated by right-clicking on the header menu.
 		SLATE_ARGUMENT(bool, CanSelectGeneratedColumn)
 		
 		// Delegate called when the selection changes
@@ -88,8 +88,11 @@ namespace UE::Editor::DataStorage
 		
 		TEDSTABLEVIEWER_API void Construct(const FArguments& InArgs);
 
+		// Clear the current QueryStack being displayed, set it to the given node, and recreate the sorting nodes 
+		TEDSTABLEVIEWER_API virtual void SetQueryStack(TSharedPtr<QueryStack::IRowNode> InRowQueryStack) override;
+		
 		// Clear the current list of columns being displayed and set it to the given list
-		TEDSTABLEVIEWER_API virtual void SetColumns(const TArray<TWeakObjectPtr<const UScriptStruct>>& Columns) override;
+		TEDSTABLEVIEWER_API virtual void SetColumns(const TArray<TWeakObjectPtr<const UScriptStruct>>& InColumns) override;
 
 		// Add a custom per-row widget to the table viewer (that doesn't necessarily map to a TEDS column)
 		// This means a new column for the table viewer
diff --git a/Engine/Plugins/Experimental/EditorDataStorageFeatures/Source/TedsTableViewer/Public/Widgets/STedsTileViewer.h b/Engine/Plugins/Experimental/EditorDataStorageFeatures/Source/TedsTableViewer/Public/Widgets/STedsTileViewer.h
--- a/Engine/Plugins/Experimental/EditorDataStorageFeatures/Source/TedsTableViewer/Public/Widgets/STedsTileViewer.h
+++ b/Engine/Plugins/Experimental/EditorDataStorageFeatures/Source/TedsTableViewer/Public/Widgets/STedsTileViewer.h
@@ -100,9 +100,12 @@ namespace UE::Editor::DataStorage
 		TEDSTABLEVIEWER_API virtual bool IsSelected(RowHandle InRow) const override;
 
 		TEDSTABLEVIEWER_API virtual bool IsSelectedExclusively(RowHandle InRow) const override;
+		
+		// Clear the current QueryStack being displayed, set it to the given node, and recreate the sorting nodes 
+        TEDSTABLEVIEWER_API virtual void SetQueryStack(TSharedPtr<QueryStack::IRowNode> InRowQueryStack) override;
 
 		// Clear the current list of columns being displayed and set it to the given list
-		TEDSTABLEVIEWER_API virtual void SetColumns(const TArray<TWeakObjectPtr<const UScriptStruct>>& Columns) override;
+		TEDSTABLEVIEWER_API virtual void SetColumns(const TArray<TWeakObjectPtr<const UScriptStruct>>& InColumns) override;
 
 		// Add a custom per-row widget to the table viewer (that doesn't necessarily map to a TEDS column)
 		TEDSTABLEVIEWER_API virtual void AddCustomRowWidget(const TSharedRef<FTedsTableViewerColumn>& InColumn) override;
