[DragAndDrop] Introduced asset drag&drop to teds table viewer.

diff --git a/Engine/Plugins/Experimental/EditorDataStorageFeatures/Source/TedsDropOperations/Private/TedsDragDropOp.cpp b/Engine/Plugins/Experimental/EditorDataStorageFeatures/Source/TedsDropOperations/Private/TedsDragDropOp.cpp
--- a/Engine/Plugins/Experimental/EditorDataStorageFeatures/Source/TedsDropOperations/Private/TedsDragDropOp.cpp
+++ b/Engine/Plugins/Experimental/EditorDataStorageFeatures/Source/TedsDropOperations/Private/TedsDragDropOp.cpp
@@ -16,19 +16,4 @@ namespace UE::Editor::DataStorage::Widgets
 		Operation->Construct();
 		return Operation;
 	}
-	
-	FRowHandleArrayView FTedsDragDropOp::GetRows() const
-	{
-		return DraggedRows.GetRows();
-	}
-
-	void FTedsDragDropOp::SetDropStatus(EDropTestResult InDropTestResult)
-	{
-		DropTestResult = InDropTestResult;
-	}
-
-	EDropTestResult FTedsDragDropOp::GetDropStatus() const
-	{
-		return DropTestResult;
-	}
 }
diff --git a/Engine/Plugins/Experimental/EditorDataStorageFeatures/Source/TedsDropOperations/Private/TedsDropOperationsUtility.cpp b/Engine/Plugins/Experimental/EditorDataStorageFeatures/Source/TedsDropOperations/Private/TedsDropOperationsUtility.cpp
new file mode 100644
--- /dev/null
+++ b/Engine/Plugins/Experimental/EditorDataStorageFeatures/Source/TedsDropOperations/Private/TedsDropOperationsUtility.cpp
@@ -0,0 +1,59 @@
+// Copyright Epic Games, Inc. All Rights Reserved.
+
+#include "TedsDropOperationsUtility.h"
+
+#include "DragAndDrop/AssetDragDropOp.h"
+#include "DragAndDrop/ClassDragDropOp.h"
+#include "Elements/Framework/TypedElementRowHandleArray.h"
+#include "Input/DragAndDrop.h"
+#include "TedsDragDropOp.h"
+#include "DataStorage/Features.h"
+#include "DataStorage/MapKey.h"
+#include "Elements/Interfaces/TypedElementDataStorageInterface.h"
+
+namespace UE::Editor::DataStorage::DragAndDrop
+{
+
+void GetRowsFromData(FRowHandleArray& OutRows, const FDragDropOperation& InData)
+{
+	if (InData.IsOfType<Widgets::FTedsDragDropOp>())
+	{
+		OutRows = static_cast<const Widgets::FTedsDragDropOp&>(InData).GetRowArray();
+	}
+	else if (InData.IsOfType<FAssetDragDropOp>())
+	{
+		if (const ICoreProvider* Storage = GetDataStorageFeature<ICoreProvider>(StorageFeatureName))
+		{
+			for (const FAssetData& Element : static_cast<const FAssetDragDropOp&>(InData).GetAssets())
+			{
+				if (Element.IsValid())
+				{
+					RowHandle Row = Storage->LookupMappedRow(/*UE::Editor::AssetData::MappingDomain*/"Asset", FMapKey(Element.GetSoftObjectPath()));
+					if (Row != InvalidRowHandle)
+					{
+						OutRows.Add(Row);
+					}
+				}
+			}
+		}
+	}
+	else if (InData.IsOfType<FClassDragDropOp>())
+	{
+		if (const ICoreProvider* Storage = GetDataStorageFeature<ICoreProvider>(StorageFeatureName))
+		{
+			for (const TWeakObjectPtr<UClass>& ClassPtr : static_cast<const FClassDragDropOp&>(InData).ClassesToDrop)
+			{
+				if (const UClass* Class = ClassPtr.Get())
+				{
+					RowHandle Row = Storage->LookupMappedRow(/*UE::Editor::AssetData::MappingDomain*/"Asset", FMapKey(Class->GetPathName()));
+					if (Row != InvalidRowHandle)
+					{
+						OutRows.Add(Row);
+					}
+				}
+			}
+		}
+	}
+}
+
+}
\ No newline at end of file
diff --git a/Engine/Plugins/Experimental/EditorDataStorageFeatures/Source/TedsDropOperations/Public/TedsDragDropOp.h b/Engine/Plugins/Experimental/EditorDataStorageFeatures/Source/TedsDropOperations/Public/TedsDragDropOp.h
--- a/Engine/Plugins/Experimental/EditorDataStorageFeatures/Source/TedsDropOperations/Public/TedsDragDropOp.h
+++ b/Engine/Plugins/Experimental/EditorDataStorageFeatures/Source/TedsDropOperations/Public/TedsDragDropOp.h
@@ -3,7 +3,6 @@
 #pragma once
 
 #include "DragAndDrop/DecoratedDragDropOp.h"
-#include "Elements/Framework/TypedElementDataStorageWidget.h"
 #include "Elements/Framework/TypedElementRowHandleArray.h"
 #include "Input/DragAndDrop.h"
 
@@ -11,13 +10,6 @@
 
 namespace UE::Editor::DataStorage::Widgets
 {
-	enum EDropTestResult
-	{
-		NoTestsRun, // Default = no drop tests have been run yet
-		SomeRowsPass, // Some but not all  the rows passed the drop test
-		AllRowsPass // All the rows passed the drop test
-	};
-	
 	// TEDS UI TODO: Instead of using a decorated operation use a default one and allow decorator customization using TEDS UI
 	class FTedsDragDropOp : public FDecoratedDragDropOp
 	{
@@ -26,16 +18,16 @@ namespace UE::Editor::DataStorage::Widgets
 
 		UE_API virtual ~FTedsDragDropOp() override = default;
 		UE_API static TSharedRef<FTedsDragDropOp> New(FRowHandleArray InDraggedRows);
-		UE_API FRowHandleArrayView GetRows() const;
-		UE_API void SetDropStatus(EDropTestResult InDropStatus);
-		UE_API EDropTestResult GetDropStatus() const;
+
+		FRowHandleArrayView    GetRows() const     { return DraggedRows.GetRows(); }
+		FRowHandleArray&       GetRowArray()       { return DraggedRows; }
+		const FRowHandleArray& GetRowArray() const { return DraggedRows; }
 		
 	protected:
 		void Init(FRowHandleArray InDraggedRows);
 		
 	protected:
 		FRowHandleArray DraggedRows;
-		EDropTestResult DropTestResult = EDropTestResult::NoTestsRun;
 	};
 }
 
diff --git a/Engine/Plugins/Experimental/EditorDataStorageFeatures/Source/TedsDropOperations/Public/TedsDropOperationsUtility.h b/Engine/Plugins/Experimental/EditorDataStorageFeatures/Source/TedsDropOperations/Public/TedsDropOperationsUtility.h
new file mode 100644
--- /dev/null
+++ b/Engine/Plugins/Experimental/EditorDataStorageFeatures/Source/TedsDropOperations/Public/TedsDropOperationsUtility.h
@@ -0,0 +1,22 @@
+// Copyright Epic Games, Inc. All Rights Reserved.
+
+#pragma once
+
+#include "CoreTypes.h"
+
+#define UE_API TEDSDROPOPERATIONS_API
+
+class FDragDropOperation;
+namespace UE::Editor::DataStorage
+{
+class FRowHandleArray;
+}
+
+namespace UE::Editor::DataStorage::DragAndDrop
+{
+
+UE_API void GetRowsFromData(FRowHandleArray& OutRows, const FDragDropOperation& InData);
+
+}
+
+#undef UE_API
diff --git a/Engine/Plugins/Experimental/EditorDataStorageFeatures/Source/TedsDropOperations/TedsDropOperations.build.cs b/Engine/Plugins/Experimental/EditorDataStorageFeatures/Source/TedsDropOperations/TedsDropOperations.build.cs
--- a/Engine/Plugins/Experimental/EditorDataStorageFeatures/Source/TedsDropOperations/TedsDropOperations.build.cs
+++ b/Engine/Plugins/Experimental/EditorDataStorageFeatures/Source/TedsDropOperations/TedsDropOperations.build.cs
@@ -24,12 +24,11 @@ public class TedsDropOperations : ModuleRules
 				{
 					"Core",
 					"CoreUObject",
-					"Engine",
 					"Slate",
 					"SlateCore",
-					"TypedElementFramework",
 					"TedsQueryStack",
-					"UnrealEd"
+					"TypedElementFramework",
+					"UnrealEd",
 				});
 
 			PublicDependencyModuleNames.AddRange(
diff --git a/Engine/Plugins/Experimental/EditorDataStorageFeatures/Source/TedsOperations/TedsOperations.build.cs b/Engine/Plugins/Experimental/EditorDataStorageFeatures/Source/TedsOperations/TedsOperations.build.cs
--- a/Engine/Plugins/Experimental/EditorDataStorageFeatures/Source/TedsOperations/TedsOperations.build.cs
+++ b/Engine/Plugins/Experimental/EditorDataStorageFeatures/Source/TedsOperations/TedsOperations.build.cs
@@ -24,7 +24,6 @@ public class TedsOperations : ModuleRules
 				{
 					"Core",
 					"CoreUObject",
-					"Engine",
 					"TypedElementFramework",
 					"TedsQueryStack",
 				});
diff --git a/Engine/Plugins/Experimental/EditorDataStorageFeatures/Source/TedsTableViewer/Private/Widgets/STedsTableViewerRow.cpp b/Engine/Plugins/Experimental/EditorDataStorageFeatures/Source/TedsTableViewer/Private/Widgets/STedsTableViewerRow.cpp
--- a/Engine/Plugins/Experimental/EditorDataStorageFeatures/Source/TedsTableViewer/Private/Widgets/STedsTableViewerRow.cpp
+++ b/Engine/Plugins/Experimental/EditorDataStorageFeatures/Source/TedsTableViewer/Private/Widgets/STedsTableViewerRow.cpp
@@ -3,15 +3,57 @@
 
 #include "Widgets/STedsTableViewerRow.h"
 
+#include "DragAndDrop/DecoratedDragDropOp.h"
 #include "Editor.h"
+#include "Elements/Columns/TypedElementHiearchyColumns.h"
 #include "TedsDragDropOp.h"
 #include "TedsDropOperations.h"
+#include "TedsDropOperationsUtility.h"
 #include "TedsTableViewerColumn.h"
-#include "Elements/Columns/TypedElementHiearchyColumns.h"
 #include "Widgets/STedsTableViewer.h"
 
+#define LOCTEXT_NAMESPACE "TedsTableViewerRow"
+
 namespace UE::Editor::DataStorage
 {
+	namespace TedsTableViewerRow_Private
+	{
+		struct FDropResult
+		{
+			bool  bCanDrop = false;
+			FText Hint     = FText();
+		};
+
+		static void UpdateOperationDecorator(FDragDropOperation& Operation, TOptional<FDropResult> Result)
+		{
+			Operation.SetDecoratorVisibility(true);
+			if (Operation.IsOfType<FDecoratedDragDropOp>())
+			{
+				static constexpr TCHAR BrushName_Ok[] = TEXT("Graph.ConnectorFeedback.OK");
+				static constexpr TCHAR BrushName_Error[] = TEXT("Graph.ConnectorFeedback.Error");
+
+				FDecoratedDragDropOp& DecoratedOp = static_cast<FDecoratedDragDropOp&>(Operation);
+				if (Result)
+				{
+					DecoratedOp.SetToolTip(Result->Hint, FAppStyle::GetBrush(Result->bCanDrop ? BrushName_Ok : BrushName_Error));
+				}
+				else
+				{
+					DecoratedOp.ResetToDefaultToolTip();
+				}
+
+			}
+			else if (Result)
+			{
+				Operation.SetCursorOverride(Result->bCanDrop ? EMouseCursor::Default : EMouseCursor::SlashedCircle);
+			}
+			else
+			{
+				Operation.SetCursorOverride({});
+			}
+		}
+	}
+
 	void STedsTableViewerRow::Construct(const FArguments& InArgs, TSharedRef<STableViewBase> OwnerTableView, TSharedRef<FTedsTableViewerModel> InTableViewerModel, TWeakPtr<const ITableViewer> InTableViewer)
 	{
 		Item = InArgs._Item;
@@ -62,83 +104,95 @@ namespace UE::Editor::DataStorage
 
 	void STedsTableViewerRow::HandleOnDragEnter(const FDragDropEvent& DragDropEvent)
 	{
-		if (ICoreProvider* Storage = TableViewerModel->GetDataStorageInterface())
+		ensureMsgf(ValidDropOperation == nullptr, TEXT("Expected last drop operation to be cleared."));
+
+		if (TSharedPtr<FDragDropOperation> Operation = DragDropEvent.GetOperation())
 		{
-			if (UDropOperationsFactory* DropFactory = Storage->FindFactory<UDropOperationsFactory>())
+			TedsTableViewerRow_Private::FDropResult Result;
+			if (ICoreProvider* Storage = TableViewerModel->GetDataStorageInterface())
 			{
-				if (TSharedPtr<Widgets::FTedsDragDropOp> Operation = DragDropEvent.GetOperationAs<Widgets::FTedsDragDropOp>())
+				if (UDropOperationsFactory* DropFactory = Storage->FindFactory<UDropOperationsFactory>())
 				{
 					RowHandle TargetRow = Item;
-		
-					for (RowHandle SourceRow : Operation->GetRows())
+
+					FRowHandleArray DraggedRows;
+					UE::Editor::DataStorage::DragAndDrop::GetRowsFromData(DraggedRows, *Operation);
+					for (RowHandle SourceRow : DraggedRows.GetRows())
 					{
 						FScopedRow InputRow = DropFactory->CreateInputRow(TargetRow, SourceRow, true);
 
-						bool bResult = DropFactory->Test(InputRow.GetRow());
+						Result.bCanDrop = DropFactory->Test(InputRow.GetRow());
 
 						// Currently the description of the last operation will take precedence, we might want support for batched tests in the future
 						if (const FText* Description = DragAndDrop::Helpers::GetDescription(*Storage, InputRow.GetRow()))
 						{
-							Operation->CurrentHoverText = *Description;
+							Result.Hint = *Description;
 						}
 
 						// No need to test further if any operation failed
-						if (!bResult)
+						if (!Result.bCanDrop)
 						{
-							return;
+							break;
 						}
-
-						Operation->SetDropStatus(Widgets::SomeRowsPass);
 					}
-
-					Operation->SetDropStatus(Widgets::AllRowsPass);
 				}
 			}
-		}
 
+			if (Result.bCanDrop)
+			{
+				ValidDropOperation = Operation.Get();
+			}
+			else if (Result.Hint.IsEmpty())
+			{
+				Result.Hint = LOCTEXT("InvalidDropOp", "Invalid drop operation.");
+			}
+
+			TedsTableViewerRow_Private::UpdateOperationDecorator(*Operation, MoveTemp(Result));
+		}
 	}
 
 	void STedsTableViewerRow::HandleOnDragLeave(const FDragDropEvent& DragDropEvent)
 	{
-		if (ICoreProvider* Storage = TableViewerModel->GetDataStorageInterface())
+		TSharedPtr<FDragDropOperation> Operation = DragDropEvent.GetOperation();
+		ensureMsgf(!ValidDropOperation || ValidDropOperation == Operation.Get(), TEXT("Expected to leave correct operation."));
+
+		if (Operation)
 		{
-			if (UDropOperationsFactory* DropFactory = Storage->FindFactory<UDropOperationsFactory>())
+			if (ValidDropOperation == Operation.Get())
 			{
-				if (TSharedPtr<Widgets::FTedsDragDropOp> Operation = DragDropEvent.GetOperationAs<Widgets::FTedsDragDropOp>())
-				{
-					// Reset the drop test result since we are leaving this row
-					Operation->SetDropStatus(Widgets::NoTestsRun);
-				}
+				ValidDropOperation = nullptr;
 			}
+
+			TedsTableViewerRow_Private::UpdateOperationDecorator(*Operation, {});
 		}
 	}
 
 	FReply STedsTableViewerRow::HandleOnAcceptDrop(const FDragDropEvent& DragDropEvent, EItemDropZone DropZone, TableViewerItemPtr InItem)
 	{
-		if (ICoreProvider* Storage = TableViewerModel->GetDataStorageInterface())
+		TSharedPtr<FDragDropOperation> Operation = DragDropEvent.GetOperation();
+		ensureMsgf(!ValidDropOperation || ValidDropOperation == Operation.Get(), TEXT("Expected to drop correct operation."));
+
+		if (ValidDropOperation && ValidDropOperation == Operation.Get())
 		{
-			if (UDropOperationsFactory* DropFactory = Storage->FindFactory<UDropOperationsFactory>())
+			ValidDropOperation = nullptr;
+
+			if (ICoreProvider* Storage = TableViewerModel->GetDataStorageInterface())
 			{
-				if (TSharedPtr<Widgets::FTedsDragDropOp> Operation = DragDropEvent.GetOperationAs<Widgets::FTedsDragDropOp>())
+				if (UDropOperationsFactory* DropFactory = Storage->FindFactory<UDropOperationsFactory>())
 				{
-					// We currently only allow drop if all rows pass the test
-					if (Operation->GetDropStatus() == Widgets::AllRowsPass)
-					{
-						RowHandle TargetRow = Item;
-						FText OutDescription;
-		
-						for (RowHandle SourceRow : Operation->GetRows())
-						{
-							FScopedRow InputRow = DropFactory->CreateInputRow(TargetRow, SourceRow, false);
-							DropFactory->Apply(InputRow.GetRow());
-						}
+					RowHandle TargetRow = Item;
 
-						return FReply::Handled();
+					FRowHandleArray DraggedRows;
+					UE::Editor::DataStorage::DragAndDrop::GetRowsFromData(DraggedRows, *Operation);
+					for (RowHandle SourceRow : DraggedRows.GetRows())
+					{
+						FScopedRow InputRow = DropFactory->CreateInputRow(TargetRow, SourceRow, false);
+						(void)DropFactory->Apply(InputRow.GetRow());
 					}
+					return FReply::Handled();
 				}
 			}
 		}
-
 		return FReply::Unhandled();
 	}
 
@@ -191,3 +245,5 @@ namespace UE::Editor::DataStorage
 		return ActualWidget;
 	}
 } // namespace UE::Editor::DataStorage
+
+#undef LOCTEXT_NAMESPACE
diff --git a/Engine/Plugins/Experimental/EditorDataStorageFeatures/Source/TedsTableViewer/Public/Widgets/STedsTableViewerRow.h b/Engine/Plugins/Experimental/EditorDataStorageFeatures/Source/TedsTableViewer/Public/Widgets/STedsTableViewerRow.h
--- a/Engine/Plugins/Experimental/EditorDataStorageFeatures/Source/TedsTableViewer/Public/Widgets/STedsTableViewerRow.h
+++ b/Engine/Plugins/Experimental/EditorDataStorageFeatures/Source/TedsTableViewer/Public/Widgets/STedsTableViewerRow.h
@@ -63,6 +63,8 @@ namespace UE::Editor::DataStorage
 		TSharedPtr<FTedsTableViewerModel> TableViewerModel;
 		TWeakPtr<const ITableViewer> WeakTableViewer;
 		TableViewerItemPtr Item;
+
+		FDragDropOperation* ValidDropOperation = nullptr; // Not for use, just to denote if a valid drop operation is hovering.
 	};
 
 	/** Widget that represents a row in the hierarchy viewers.  Generates widgets for each column and adds an expander arrow to the very first column. */
