[TEDS Outliner] Add the External Selection Column on Outliner rows and TableViewer rows
- Additionally have the label icon align to the center so it doesn't stretch when rename grows the text area

#rnx

diff --git a/Engine/Plugins/Experimental/EditorDataStorageFeatures/Source/TedsOutliner/Private/TedsOutlinerImpl.cpp b/Engine/Plugins/Experimental/EditorDataStorageFeatures/Source/TedsOutliner/Private/TedsOutlinerImpl.cpp
--- a/Engine/Plugins/Experimental/EditorDataStorageFeatures/Source/TedsOutliner/Private/TedsOutlinerImpl.cpp
+++ b/Engine/Plugins/Experimental/EditorDataStorageFeatures/Source/TedsOutliner/Private/TedsOutlinerImpl.cpp
@@ -160,10 +160,9 @@ namespace LabelWidgetUtils
 		TSharedPtr<SWidget> Widget = StorageUi.ConstructWidget(WidgetRow, WidgetConstructor, QueryMetaDataView);
 		if (Widget)
 		{
-			if (FExternalWidgetSelectionColumn* ExternalWidgetSelectionColumn = Storage.GetColumn<FExternalWidgetSelectionColumn>(WidgetRow))
-			{
-				ExternalWidgetSelectionColumn->IsSelected = FIsSelected::CreateSP(&RowItem, &STableRow<FSceneOutlinerTreeItemPtr>::IsSelectedExclusively);
-			}
+			Storage.AddColumn<FExternalWidgetSelectionColumn>(WidgetRow, FExternalWidgetSelectionColumn{
+				.IsSelected = FIsSelected::CreateSP(&RowItem, &STableRow<FSceneOutlinerTreeItemPtr>::IsSelected)
+			});
 
 			if (bInteractable)
 			{
diff --git a/Engine/Plugins/Experimental/EditorDataStorageFeatures/Source/TedsOutliner/Private/Widgets/OutlinerLabelWidget.cpp b/Engine/Plugins/Experimental/EditorDataStorageFeatures/Source/TedsOutliner/Private/Widgets/OutlinerLabelWidget.cpp
--- a/Engine/Plugins/Experimental/EditorDataStorageFeatures/Source/TedsOutliner/Private/Widgets/OutlinerLabelWidget.cpp
+++ b/Engine/Plugins/Experimental/EditorDataStorageFeatures/Source/TedsOutliner/Private/Widgets/OutlinerLabelWidget.cpp
@@ -115,6 +115,7 @@ TSharedPtr<SWidget> FOutlinerLabelWidgetConstructor::CreateWidget(UE::Editor::Da
 		return SNew(SHorizontalBox)
 			+ SHorizontalBox::Slot()
 			.AutoWidth()
+			.VAlign(VAlign_Center)
 			[
 				IconWidget.ToSharedRef()
 			]
@@ -126,6 +127,7 @@ TSharedPtr<SWidget> FOutlinerLabelWidgetConstructor::CreateWidget(UE::Editor::Da
 			]
 			+ SHorizontalBox::Slot()
 			.FillWidth(1.0f)
+			.VAlign(VAlign_Center)
 			[
 				TextWidget.ToSharedRef()
 			];
diff --git a/Engine/Plugins/Experimental/EditorDataStorageFeatures/Source/TedsTableViewer/Private/Widgets/STedsTableViewerRow.cpp b/Engine/Plugins/Experimental/EditorDataStorageFeatures/Source/TedsTableViewer/Private/Widgets/STedsTableViewerRow.cpp
--- a/Engine/Plugins/Experimental/EditorDataStorageFeatures/Source/TedsTableViewer/Private/Widgets/STedsTableViewerRow.cpp
+++ b/Engine/Plugins/Experimental/EditorDataStorageFeatures/Source/TedsTableViewer/Private/Widgets/STedsTableViewerRow.cpp
@@ -3,6 +3,7 @@
 
 #include "Widgets/STedsTableViewerRow.h"
 
+#include "Columns/SlateDelegateColumns.h"
 #include "DragAndDrop/DecoratedDragDropOp.h"
 #include "Editor.h"
 #include "Elements/Columns/TypedElementHiearchyColumns.h"
@@ -80,6 +81,9 @@ namespace UE::Editor::DataStorage
 			auto WidgetRowSetupDelegate = [this] (ICoreProvider& InStorage, const RowHandle& InUIRowHandle)
 			{
 				InStorage.AddColumn(InUIRowHandle, FTableRowParentColumn{ .Parent = ParentWidgetRowHandle });
+				// Currently only this column is needed (not FExternalWidgetExclusiveSelectionColumn) since the double-click action will automatically
+				// deselect all and select only the desired row. 
+				InStorage.AddColumn(InUIRowHandle, FExternalWidgetSelectionColumn { .IsSelected = FIsSelected::CreateSP(this, &STedsTableViewerRow::IsSelected)});
 			};
 
 			if (TSharedPtr<SWidget> RowWidget = Column->ConstructRowWidget(Item, WidgetRowSetupDelegate))
@@ -109,10 +113,10 @@ namespace UE::Editor::DataStorage
 		if (TSharedPtr<FDragDropOperation> Operation = DragDropEvent.GetOperation())
 		{
 			TedsTableViewerRow_Private::FDropResult Result;
-			if (ICoreProvider* Storage = TableViewerModel->GetDataStorageInterface())
+		if (ICoreProvider* Storage = TableViewerModel->GetDataStorageInterface())
+		{
+			if (UDropOperationsFactory* DropFactory = Storage->FindFactory<UDropOperationsFactory>())
 			{
-				if (UDropOperationsFactory* DropFactory = Storage->FindFactory<UDropOperationsFactory>())
-				{
 					RowHandle TargetRow = Item;
 
 					FRowHandleArray DraggedRows;
@@ -135,13 +139,13 @@ namespace UE::Editor::DataStorage
 							break;
 						}
 					}
-				}
-			}
+						}
+					}
 
 			if (Result.bCanDrop)
 			{
 				ValidDropOperation = Operation.Get();
-			}
+				}
 			else if (Result.Hint.IsEmpty())
 			{
 				Result.Hint = LOCTEXT("InvalidDropOp", "Invalid drop operation.");
@@ -157,9 +161,9 @@ namespace UE::Editor::DataStorage
 		ensureMsgf(!ValidDropOperation || ValidDropOperation == Operation.Get(), TEXT("Expected to leave correct operation."));
 
 		if (Operation)
-		{
-			if (ValidDropOperation == Operation.Get())
 			{
+			if (ValidDropOperation == Operation.Get())
+				{
 				ValidDropOperation = nullptr;
 			}
 
@@ -176,23 +180,23 @@ namespace UE::Editor::DataStorage
 		{
 			ValidDropOperation = nullptr;
 
-			if (ICoreProvider* Storage = TableViewerModel->GetDataStorageInterface())
+		if (ICoreProvider* Storage = TableViewerModel->GetDataStorageInterface())
+		{
+			if (UDropOperationsFactory* DropFactory = Storage->FindFactory<UDropOperationsFactory>())
 			{
-				if (UDropOperationsFactory* DropFactory = Storage->FindFactory<UDropOperationsFactory>())
-				{
-					RowHandle TargetRow = Item;
+						RowHandle TargetRow = Item;
 
 					FRowHandleArray DraggedRows;
 					UE::Editor::DataStorage::DragAndDrop::GetRowsFromData(DraggedRows, *Operation);
 					for (RowHandle SourceRow : DraggedRows.GetRows())
-					{
-						FScopedRow InputRow = DropFactory->CreateInputRow(TargetRow, SourceRow, false);
+						{
+							FScopedRow InputRow = DropFactory->CreateInputRow(TargetRow, SourceRow, false);
 						(void)DropFactory->Apply(InputRow.GetRow());
+						}
+						return FReply::Handled();
 					}
-					return FReply::Handled();
 				}
 			}
-		}
 		return FReply::Unhandled();
 	}
 
