Teds Outliner: Replace bespoke drag and drop logic with the new Teds drop operations

diff --git a/Engine/Plugins/Experimental/EditorDataStorageFeatures/Source/TedsOutliner/Private/TedsOutlinerMode.cpp b/Engine/Plugins/Experimental/EditorDataStorageFeatures/Source/TedsOutliner/Private/TedsOutlinerMode.cpp
--- a/Engine/Plugins/Experimental/EditorDataStorageFeatures/Source/TedsOutliner/Private/TedsOutlinerMode.cpp
+++ b/Engine/Plugins/Experimental/EditorDataStorageFeatures/Source/TedsOutliner/Private/TedsOutlinerMode.cpp
@@ -2,17 +2,15 @@
 
 #include "TedsOutlinerMode.h"
 
-#include "TedsOutlinerFilter.h"
-#include "Elements/Framework/TypedElementRegistry.h"
-#include "TedsOutlinerHierarchy.h"
-#include "TedsOutlinerItem.h"
 #include "Elements/Columns/TypedElementCompatibilityColumns.h"
-#include "Elements/Columns/TypedElementHiearchyColumns.h"
-#include "Elements/Columns/TypedElementMiscColumns.h"
 #include "Elements/Columns/TypedElementSelectionColumns.h"
 #include "FolderTreeItem.h"
-#include "Elements/Columns/TypedElementLabelColumns.h"
-#include "Elements/Columns/TypedElementTypeInfoColumns.h"
+#include "TedsDragDropOp.h"
+#include "TedsDropOperationsColumnsHelpers.h"
+#include "TedsDropOperationsFactory.h"
+#include "TedsDropOperationsUtility.h"
+#include "TedsOutlinerHierarchy.h"
+#include "TedsOutlinerItem.h"
 #include "ToolMenus.h"
 
 #include UE_INLINE_GENERATED_CPP_BY_NAME(TedsOutlinerMode)
@@ -23,9 +21,6 @@ namespace UE::Editor::Outliner
 {
 	namespace Private
 	{
-		// Drag drop currently disabled as we are missing data marshaling for hierarchies from TEDS to the world
-		static bool TedsOutlinerDragDropEnabled = false;
-		static FAutoConsoleVariableRef TedsOutlinerDragDropEnabledCvar(TEXT("TEDS.UI.EnableTEDSOutlinerDragDrop"), TedsOutlinerDragDropEnabled, TEXT("Enable drag/drop for the generic TEDS Outliner."));
 		static FName ContextMenuName("TedsOutlinerContextMenu");
 	} // namespace Private
 
@@ -110,15 +105,7 @@ void FTedsOutlinerMode::OnItemSelectionChanged(FSceneOutlinerTreeItemPtr Item, E
 
 TSharedPtr<FDragDropOperation> FTedsOutlinerMode::CreateDragDropOperation(const FPointerEvent& MouseEvent, const TArray<FSceneOutlinerTreeItemPtr>& InTreeItems) const
 {
-	const TOptional<FTedsOutlinerHierarchyData>& HierarchyData = TedsOutlinerImpl->GetHierarchyData();
-
-	// We don't want drag/drop if this TEDS Outliner isn't showing any hierarchy data
-	if(!HierarchyData.IsSet())
-	{
-		return nullptr;
-	}
-	
-	TArray<DataStorage::RowHandle> DraggedRowHandles;
+	FRowHandleArray DraggedRowHandles;
 
 	for(const FSceneOutlinerTreeItemPtr& Item :InTreeItems)
 	{
@@ -129,99 +116,71 @@ TSharedPtr<FDragDropOperation> FTedsOutlinerMode::CreateDragDropOperation(const
 		}
 	}
 
-	return FTedsRowDragDropOp::New(DraggedRowHandles);
+	return Widgets::FTedsDragDropOp::New(MoveTemp(DraggedRowHandles));
 }
 
 bool FTedsOutlinerMode::ParseDragDrop(FSceneOutlinerDragDropPayload& OutPayload, const FDragDropOperation& Operation) const
 {
-	if (Operation.IsOfType<FTedsRowDragDropOp>())
+	if (Operation.IsOfType<Widgets::FTedsDragDropOp>())
 	{
-		const FTedsRowDragDropOp& TedsOp = static_cast<const FTedsRowDragDropOp&>(Operation);
+		const Widgets::FTedsDragDropOp& TedsOperation = static_cast<const Widgets::FTedsDragDropOp&>(Operation);
 
-		for(DataStorage::RowHandle RowHandle : TedsOp.DraggedRows)
+		for (RowHandle SourceRow : TedsOperation.GetRows())
 		{
-			OutPayload.DraggedItems.Add(SceneOutliner->GetTreeItem(RowHandle));
+			if (FSceneOutlinerTreeItemPtr Item = SceneOutliner->GetTreeItem(SourceRow))
+			{
+				OutPayload.DraggedItems.Add(Item);
+			}
 		}
-		return true;
 	}
-	return false;
+	return true; // Always return true since we support external drops as well
 }
 
 FSceneOutlinerDragValidationInfo FTedsOutlinerMode::ValidateDrop(const ISceneOutlinerTreeItem& DropTarget, const FSceneOutlinerDragDropPayload& Payload) const
 {
-	const TOptional<FTedsOutlinerHierarchyData>& HierarchyData = TedsOutlinerImpl->GetHierarchyData();
-	DataStorage::ICoreProvider* Storage = TedsOutlinerImpl->GetStorage();
+	ICoreProvider* Storage = TedsOutlinerImpl->GetStorage();
+	FRowHandleArray DraggedRowHandles;
 
-	// We don't want drag/drop if this TEDS Outliner isn't showing any hierarchy data
-	if(!HierarchyData.IsSet())
+	// External drop
+	if (Payload.DraggedItems.IsEmpty())
 	{
-		return FSceneOutlinerDragValidationInfo(ESceneOutlinerDropCompatibility::IncompatibleGeneric,
-			LOCTEXT("DropDisabled", "Drag/Drop is disabled due to missing hierarchy data!"));
+		DragAndDrop::GetRowsFromData(DraggedRowHandles, Payload.SourceOperation);
 	}
-	
-	TArray<DataStorage::RowHandle> DraggedRowHandles;
-
-	Payload.ForEachItem<FTedsOutlinerTreeItem>([&DraggedRowHandles](FTedsOutlinerTreeItem& TedsItem)
+	else
+	{
+		Payload.ForEachItem<FTedsOutlinerTreeItem>([&DraggedRowHandles](FTedsOutlinerTreeItem& TedsItem)
 		{
 			DraggedRowHandles.Add(TedsItem.GetRowHandle());
 		});
-
-	// Dropping onto another item
-	// TEDS-Outliner TODO: Need better drag/drop validation and better place for this, TEDS-Outliner does not know about what types these rows are and all types that exist and what attachment is valid
-	if(const FTedsOutlinerTreeItem* TedsItem = DropTarget.CastTo<FTedsOutlinerTreeItem>())
-	{
-		DataStorage::RowHandle DropTargetRowHandle = TedsItem->GetRowHandle();
-
-		FTypedElementClassTypeInfoColumn* DropTargetTypeInfoColumn = Storage->GetColumn<FTypedElementClassTypeInfoColumn>(DropTargetRowHandle);
-
-		// For now only allow attachment to same type
-		if(!DropTargetTypeInfoColumn || !DropTargetTypeInfoColumn->TypeInfo.Get())
-		{
-			return FSceneOutlinerDragValidationInfo(ESceneOutlinerDropCompatibility::IncompatibleGeneric, LOCTEXT("DropTargetInvalidType", "Invalid Drop target"));
-		}
-
-		
-		// TEDS-Outliner TODO: Currently we detect parent changes by removing the column and then adding the column back with the new parent 
-		for(DataStorage::RowHandle RowHandle : DraggedRowHandles)
-		{
-			FTypedElementClassTypeInfoColumn* TypeInfoColumn = Storage->GetColumn<FTypedElementClassTypeInfoColumn>(RowHandle);
-
-			if(!TypeInfoColumn || !TypeInfoColumn->TypeInfo.Get())
-			{
-				return FSceneOutlinerDragValidationInfo(ESceneOutlinerDropCompatibility::IncompatibleGeneric, LOCTEXT("DragItemInvalidType", "Invalid Drag item"));
-			}
-
-			// TEDS-Outliner TOOD UE-205438: Proper drag/drop validation
-			if(TypeInfoColumn->TypeInfo.Get() != DropTargetTypeInfoColumn->TypeInfo.Get())
-			{
-				return FSceneOutlinerDragValidationInfo(ESceneOutlinerDropCompatibility::IncompatibleGeneric, FText::Format(LOCTEXT("DragDropTypeMismatch", "Cannot drag a {0} into a {1}"), FText::FromName(TypeInfoColumn->TypeInfo.Get()->GetFName()), FText::FromName(DropTargetTypeInfoColumn->TypeInfo.Get()->GetFName())));
-			}
-		}
-
-		return FSceneOutlinerDragValidationInfo(ESceneOutlinerDropCompatibility::CompatibleAttach, LOCTEXT("ValidDrop", "Valid Drop"));
 	}
-	// Dropping onto root, remove parent
-	else if (const FFolderTreeItem* FolderItem = DropTarget.CastTo<FFolderTreeItem>())
-	{
-		const FFolder DestinationPath = FolderItem->GetFolder();
 
-		if(DestinationPath.IsNone())
+	// TODO: The Outliner considers a drop onto empty area as a drop onto an empty FFolderTreeItem, we need to handle that case by providing
+	// that info to the operations somehow
+	if (const FTedsOutlinerTreeItem* TedsDropTarget = DropTarget.CastTo<FTedsOutlinerTreeItem>())
+	{
+		if (UDropOperationsFactory* DropFactory = Storage->FindFactory<UDropOperationsFactory>())
 		{
-			bool bValidDetach = false;
+			FText Description;
 			
-			for(DataStorage::RowHandle RowHandle : DraggedRowHandles)
+			for (RowHandle SourceRow : DraggedRowHandles.GetRows())
 			{
-				if(Storage->HasColumns(RowHandle, MakeArrayView({HierarchyData.GetValue().HierarchyColumn})))
+				FScopedRow InputRow = DropFactory->CreateInputRow(TedsDropTarget->GetRowHandle(), SourceRow, true);
+ 
+				bool bResult = DropFactory->Test(InputRow.GetRow());
+
+				// Currently the description of the last operation will take precedence, we might want support for batched tests in the future
+				if (FText* DescriptionPtr = DragAndDrop::Helpers::GetDescription(*Storage, InputRow.GetRow()))
 				{
-					bValidDetach = true;
-					break;
+					Description = *DescriptionPtr;
 				}
-			}
 
-			if(bValidDetach)
-			{
-				return FSceneOutlinerDragValidationInfo(ESceneOutlinerDropCompatibility::CompatibleDetach, LOCTEXT("MoveToRoot", "Move to root"));
+				// Currently only accept drops if all the rows are compatible with an operation
+				if (!bResult)
+				{
+					return FSceneOutlinerDragValidationInfo(ESceneOutlinerDropCompatibility::IncompatibleGeneric, Description);
+				}
 			}
+			return FSceneOutlinerDragValidationInfo(ESceneOutlinerDropCompatibility::CompatibleGeneric, Description);
 		}
 	}
 
@@ -230,43 +189,34 @@ FSceneOutlinerDragValidationInfo FTedsOutlinerMode::ValidateDrop(const ISceneOut
 
 void FTedsOutlinerMode::OnDrop(ISceneOutlinerTreeItem& DropTarget, const FSceneOutlinerDragDropPayload& Payload, const FSceneOutlinerDragValidationInfo& ValidationInfo) const
 {
-	const TOptional<FTedsOutlinerHierarchyData>& HierarchyData = TedsOutlinerImpl->GetHierarchyData();
-	DataStorage::ICoreProvider* Storage = TedsOutlinerImpl->GetStorage();
+	ICoreProvider* Storage = TedsOutlinerImpl->GetStorage();
 
-	if(!Private::TedsOutlinerDragDropEnabledCvar->GetBool() || !HierarchyData.IsSet())
-	{
-		return;
-	}
-	
-	TArray<DataStorage::RowHandle> DraggedRowHandles;
+	FRowHandleArray DraggedRowHandles;
 
-	Payload.ForEachItem<FTedsOutlinerTreeItem>([&DraggedRowHandles](FTedsOutlinerTreeItem& TedsItem)
+	if (Payload.DraggedItems.IsEmpty())
 	{
-		DraggedRowHandles.Add(TedsItem.GetRowHandle());
-	});
-
-	if(ValidationInfo.CompatibilityType == ESceneOutlinerDropCompatibility::CompatibleDetach)
+		DragAndDrop::GetRowsFromData(DraggedRowHandles, Payload.SourceOperation);
+	}
+	else
 	{
-		for(DataStorage::RowHandle RowHandle : DraggedRowHandles)
+		Payload.ForEachItem<FTedsOutlinerTreeItem>([&DraggedRowHandles](FTedsOutlinerTreeItem& TedsItem)
 		{
-			Storage->RemoveColumn(RowHandle, HierarchyData.GetValue().HierarchyColumn);
-			Storage->AddColumn<FTypedElementSyncBackToWorldTag>(RowHandle);
-		}
+			DraggedRowHandles.Add(TedsItem.GetRowHandle());
+		});
 	}
-	
-	if(const FTedsOutlinerTreeItem* TedsItem = DropTarget.CastTo<FTedsOutlinerTreeItem>())
+
+	if (const FTedsOutlinerTreeItem* TedsItem = DropTarget.CastTo<FTedsOutlinerTreeItem>())
 	{
-		DataStorage::RowHandle DropTargetRowHandle = TedsItem->GetRowHandle();
-		
-		for(DataStorage::RowHandle RowHandle : DraggedRowHandles)
+		if (ValidationInfo.CompatibilityType == ESceneOutlinerDropCompatibility::CompatibleGeneric)
 		{
-			// Add the column
-			Storage->AddColumn(RowHandle, HierarchyData.GetValue().HierarchyColumn);
-
-			// Let the hierarchy data fill in the parent row into the column
-			HierarchyData.GetValue().SetParent.Execute(Storage->GetColumnData(RowHandle, HierarchyData.GetValue().HierarchyColumn), DropTargetRowHandle);
-			
-			Storage->AddColumn<FTypedElementSyncBackToWorldTag>(RowHandle);
+			if (UDropOperationsFactory* DropFactory = Storage->FindFactory<UDropOperationsFactory>())
+			{
+				for (RowHandle SourceRow : DraggedRowHandles.GetRows())
+				{
+					FScopedRow InputRow = DropFactory->CreateInputRow(TedsItem->GetRowHandle(), SourceRow, false);
+					DropFactory->Apply(InputRow.GetRow());
+				}
+			}
 		}
 	}
 }
diff --git a/Engine/Plugins/Experimental/EditorDataStorageFeatures/Source/TedsOutliner/Public/TedsOutlinerMode.h b/Engine/Plugins/Experimental/EditorDataStorageFeatures/Source/TedsOutliner/Public/TedsOutlinerMode.h
--- a/Engine/Plugins/Experimental/EditorDataStorageFeatures/Source/TedsOutliner/Public/TedsOutlinerMode.h
+++ b/Engine/Plugins/Experimental/EditorDataStorageFeatures/Source/TedsOutliner/Public/TedsOutlinerMode.h
@@ -13,32 +13,6 @@
 
 namespace UE::Editor::Outliner
 {
-// TEDS-Outliner TODO: This can probably be moved to a more generic location for all TEDS related drag drops?
-class FTedsRowDragDropOp : public FDecoratedDragDropOp
-{
-public:
-	DRAG_DROP_OPERATOR_TYPE(FTedsRowDragDropOp, FDecoratedDragDropOp)
-
-	/** Rows we are dragging */
-	TArray<DataStorage::RowHandle> DraggedRows;
-
-	void Init(const TArray<DataStorage::RowHandle>& InRowHandles)
-	{
-		DraggedRows = InRowHandles;
-	}
-
-	static TSharedRef<FTedsRowDragDropOp> New(const TArray<DataStorage::RowHandle>& InRowHandles)
-	{
-		TSharedRef<FTedsRowDragDropOp> Operation = MakeShareable(new FTedsRowDragDropOp);
-		
-		Operation->Init(InRowHandles);
-		Operation->SetupDefaults();
-		Operation->Construct();
-		
-		return Operation;
-	}
-};
-
 /*
  * TEDS driven Outliner mode where the Outliner is populated using the results of the RowHandleQueries passed in.
  * See CreateGenericTEDSOutliner() for example usage
diff --git a/Engine/Plugins/Experimental/EditorDataStorageFeatures/Source/TedsOutliner/TedsOutliner.Build.cs b/Engine/Plugins/Experimental/EditorDataStorageFeatures/Source/TedsOutliner/TedsOutliner.Build.cs
--- a/Engine/Plugins/Experimental/EditorDataStorageFeatures/Source/TedsOutliner/TedsOutliner.Build.cs
+++ b/Engine/Plugins/Experimental/EditorDataStorageFeatures/Source/TedsOutliner/TedsOutliner.Build.cs
@@ -41,7 +41,9 @@ public class TedsOutliner : ModuleRules
 					"InputCore",
 					"LevelEditor",
 					"TedsQueryStack",
-					"TedsActorCompatibility"
+					"TedsActorCompatibility",
+					"TedsOperations",
+					"TedsDropOperations"
 				});
 
 			DynamicallyLoadedModuleNames.AddRange(new string[] {});
