Update Teds AssetData to pause processing while Asset Registry is gathering

diff --git a/Engine/Plugins/Experimental/EditorDataStorageFeatures/Source/TedsAssetData/Private/TedsAssetData.cpp b/Engine/Plugins/Experimental/EditorDataStorageFeatures/Source/TedsAssetData/Private/TedsAssetData.cpp
--- a/Engine/Plugins/Experimental/EditorDataStorageFeatures/Source/TedsAssetData/Private/TedsAssetData.cpp
+++ b/Engine/Plugins/Experimental/EditorDataStorageFeatures/Source/TedsAssetData/Private/TedsAssetData.cpp
@@ -188,8 +188,72 @@
 FTedsAssetData::FTedsAssetData(UE::Editor::DataStorage::ICoreProvider& InDatabase)
 	: Database(InDatabase)
 {
+	IAssetRegistry& AssetRegistry = IAssetRegistry::GetChecked();
+
+	AssetRegistry.OnScanStarted().AddRaw(this, &FTedsAssetData::OnARScanStarted);
+	AssetRegistry.OnScanEnded().AddRaw(this, &FTedsAssetData::OnARScanEnded);
+
+	// Start processing immediately if Asset Registry is not gathering and has not completed its initial scan
+	// Otherwise wait for OnScanEnded
+	if (AssetRegistry.IsSearchAllAssets() && !AssetRegistry.IsGathering())
+	{
+		StartProcessing();
+	}
+}
+
+FTedsAssetData::~FTedsAssetData()
+{
+	TRACE_CPUPROFILER_EVENT_SCOPE(FTedsAssetData::~FTedsAssetData);
+
+#if TRACK_TEDSASSETDATA_MEMORY
+	LLM_SCOPE_BYNAME(TEXT("FTedsAssetData"))
+#endif
+
+	// Not needed on a editor shut down
+	if (!IsEngineExitRequested())
+	{
+		Database.UnregisterCooperativeUpdate(TEXT("ProcessAssetTasksCooperatively"));
+
+		Database.OnUpdate().RemoveAll(this);
+
+		if (IAssetRegistry* AssetRegistry = IAssetRegistry::Get())
+		{
+			Database.UnregisterQuery(RemoveUpdatedAssetDataTagQuery);
+			Database.UnregisterQuery(RemoveUpdatedPathTagQuery);
+
+			AssetRegistry->OnScanStarted().RemoveAll(this);
+			AssetRegistry->OnScanEnded().RemoveAll(this);
+			AssetRegistry->OnAssetsAdded().RemoveAll(this);
+			AssetRegistry->OnAssetsRemoved().RemoveAll(this);
+			AssetRegistry->OnAssetsUpdated().RemoveAll(this);
+			AssetRegistry->OnAssetsUpdatedOnDisk().RemoveAll(this);
+			AssetRegistry->OnAssetRenamed().RemoveAll(this);
+			AssetRegistry->OnPathsAdded().RemoveAll(this);
+			AssetRegistry->OnPathsRemoved().RemoveAll(this);
+	
+			AssetRegistry->EnumerateAllCachedPaths([this](FName InPath)
+				{
+					const FMapKeyView PathKey = FMapKeyView(InPath);
+					const RowHandle Row = Database.LookupMappedRow(MappingDomain, PathKey);
+					Database.RemoveRow(Row);
+					return true;
+				});
+
+			AssetRegistry->EnumerateAllAssets([this](const FAssetData& InAssetData)
+				{
+					const FMapKey AssetPathKey = FMapKey(InAssetData.GetSoftObjectPath());
+					const RowHandle Row = Database.LookupMappedRow(MappingDomain, AssetPathKey);
+					Database.RemoveRow(Row);
+					return true;
+				});
+		}
+	}
+}
+
+void FTedsAssetData::Initialize()
+{
 	using namespace UE::Editor::DataStorage::Queries;
-	TRACE_CPUPROFILER_EVENT_SCOPE(FTedsAssetData::FTedsAssetData);
+	TRACE_CPUPROFILER_EVENT_SCOPE(FTedsAssetData::Initialize);
 
 #if TRACK_TEDSASSETDATA_MEMORY
 	LLM_SCOPE_BYNAME(TEXT("FTedsAssetData"))
@@ -238,7 +302,7 @@
 				Context.RemoveColumns<FUpdatedPathTag>(TConstArrayView<RowHandle>(Rows, Context.GetRowCount()));
 			})
 		.Where()
-			.All<FUpdatedPathTag>()
+		.All<FUpdatedPathTag>()
 		.Compile());
 
 	RemoveUpdatedAssetDataTagQuery = Database.RegisterQuery(
@@ -250,7 +314,7 @@
 				Context.RemoveColumns<FUpdatedAssetDataTag>(TConstArrayView<RowHandle>(Rows, Context.GetRowCount()));
 			})
 		.Where()
-			.All<FUpdatedAssetDataTag>()
+		.All<FUpdatedAssetDataTag>()
 		.Compile());
 
 	// Init with the data existing at moment in asset registry
@@ -274,53 +338,23 @@
 			return true;
 		});
 	CachedPaths.MoveToLinearArray(AddedPaths);
+
+	bIsInitialized = true;
 }
 
-FTedsAssetData::~FTedsAssetData()
+void FTedsAssetData::StartProcessing()
 {
-	TRACE_CPUPROFILER_EVENT_SCOPE(FTedsAssetData::~FTedsAssetData);
-
-#if TRACK_TEDSASSETDATA_MEMORY
-	LLM_SCOPE_BYNAME(TEXT("FTedsAssetData"))
-#endif
-
-	// Not needed on a editor shut down
-	if (!IsEngineExitRequested())
+	if (!bIsInitialized)
 	{
-		Database.UnregisterCooperativeUpdate(TEXT("ProcessAssetTasksCooperatively"));
-
-		Database.OnUpdate().RemoveAll(this);
-
-		if (IAssetRegistry* AssetRegistry = IAssetRegistry::Get())
-		{
-			Database.UnregisterQuery(RemoveUpdatedAssetDataTagQuery);
-			Database.UnregisterQuery(RemoveUpdatedPathTagQuery);
-
-			AssetRegistry->OnAssetsAdded().RemoveAll(this);
-			AssetRegistry->OnAssetsRemoved().RemoveAll(this);
-			AssetRegistry->OnAssetsUpdated().RemoveAll(this);
-			AssetRegistry->OnAssetsUpdatedOnDisk().RemoveAll(this);
-			AssetRegistry->OnAssetRenamed().RemoveAll(this);
-			AssetRegistry->OnPathsAdded().RemoveAll(this);
-			AssetRegistry->OnPathsRemoved().RemoveAll(this);
-	
-			AssetRegistry->EnumerateAllCachedPaths([this](FName InPath)
-				{
-					const FMapKeyView PathKey = FMapKeyView(InPath);
-					const RowHandle Row = Database.LookupMappedRow(MappingDomain, PathKey);
-					Database.RemoveRow(Row);
-					return true;
-				});
-
-			AssetRegistry->EnumerateAllAssets([this](const FAssetData& InAssetData)
-				{
-					const FMapKey AssetPathKey = FMapKey(InAssetData.GetSoftObjectPath());
-					const RowHandle Row = Database.LookupMappedRow(MappingDomain, AssetPathKey);
-					Database.RemoveRow(Row);
-					return true;
-				});
-		}
+		// Do initial bulk asset grab and register to Asset Registry events
+		Initialize();
 	}
+	bIsPaused = false;
 }
 
+void FTedsAssetData::PauseProcessing()
+{
+	bIsPaused = true;
+}
+
 void FTedsAssetData::ProcessAllEvents()
@@ -331,6 +365,16 @@
 	}
 }
 
+void FTedsAssetData::OnARScanStarted()
+{
+	PauseProcessing();
+}
+
+void FTedsAssetData::OnARScanEnded()
+{
+	StartProcessing();
+}
+
 void FTedsAssetData::OnAssetsAdded(TConstArrayView<FAssetData> InAssetsAdded)
 {
 	TRACE_CPUPROFILER_EVENT_SCOPE(FTedsAssetData::OnAssetsAdded);
@@ -692,6 +736,12 @@
 {
 	TRACE_CPUPROFILER_EVENT_SCOPE(FTedsAssetData::OnDatabaseUpdate);
 
+	// Asset Registry is currently gathering. Don't process batches at this time to avoid competing on the Game Thread
+	if (bIsPaused)
+	{
+		return;
+	}
+
 	UpdateAvgProcessedTime(AvgTimePerProcessedPathAdd, ProcessPathAddStartTime, ProcessPathAddEndTime, LastAddedPathBatchSize);
 	UpdateAvgProcessedTime(AvgTimePerProcessedAssetAdd, ProcessAssetAddStartTime, ProcessAssetAddEndTime, LastAddedAssetBatchSize);
 	UpdateAvgProcessedTime(AvgTimePerProcessedAssetUpdate, ProcessAssetUpdateStartTime, ProcessAssetUpdateEndTime, LastUpdatedAssetBatchSize);
@@ -791,8 +841,8 @@
 	TRACE_COUNTER_SET(COUNTER_TedsAssetData_AssetUpdatesRemaining, UpdatedAssets.Num());
 
 	SET_DWORD_STAT(STAT_TedsAssetData_PathAddBatchSize, LastAddedPathBatchSize);
-	SET_DWORD_STAT(STAT_TedsAssetData_PathAddBatchSize, LastAddedAssetBatchSize);
-	SET_DWORD_STAT(STAT_TedsAssetData_PathAddBatchSize, LastUpdatedAssetBatchSize);
+	SET_DWORD_STAT(STAT_TedsAssetData_AssetAddBatchSize, LastAddedAssetBatchSize);
+	SET_DWORD_STAT(STAT_TedsAssetData_AssetUpdateBatchSize, LastUpdatedAssetBatchSize);
 
 	TRACE_COUNTER_SET(COUNTER_TedsAssetData_PathAddBatchSize, LastAddedPathBatchSize);
 	TRACE_COUNTER_SET(COUNTER_TedsAssetData_AssetAddBatchSize, LastAddedAssetBatchSize);
diff --git a/Engine/Plugins/Experimental/EditorDataStorageFeatures/Source/TedsAssetData/Private/TedsAssetData.h b/Engine/Plugins/Experimental/EditorDataStorageFeatures/Source/TedsAssetData/Private/TedsAssetData.h
--- a/Engine/Plugins/Experimental/EditorDataStorageFeatures/Source/TedsAssetData/Private/TedsAssetData.h
+++ b/Engine/Plugins/Experimental/EditorDataStorageFeatures/Source/TedsAssetData/Private/TedsAssetData.h
@@ -27,9 +27,16 @@ namespace UE::Editor::AssetData::Private
 
 		~FTedsAssetData();
 
+		void Initialize();
+		void StartProcessing();
+		void PauseProcessing();
+
 		void ProcessAllEvents();
 
 	private:
+		void OnARScanStarted();
+		void OnARScanEnded();
+
 		void OnAssetsAdded(TConstArrayView<FAssetData> InAssetsAdded);
 		FTimespan ProcessAddedAssets(const FTimespan& AllowedTime);
 
@@ -55,6 +62,9 @@ namespace UE::Editor::AssetData::Private
 
 		void RecordAssetDataStats();
 
+		bool bIsInitialized = false;
+		bool bIsPaused = true;
+
 		TArray<FName> AddedPaths;
 		TArray<FAssetData> AddedAssets;
 		TArray<FAssetData> UpdatedAssets;
