[TEDS] Change the Query Stack nodes to use defined values for initial revision IDs and change the FilterNode to initialize with zero instead of the parent's ID. Also, change the ViewNode to initialize with one and its tests to match the new expected value

diff --git a/Engine/Plugins/Experimental/EditorDataStorageFeatures/Source/TedsQueryStack/Private/TedsRowFilterNode.cpp b/Engine/Plugins/Experimental/EditorDataStorageFeatures/Source/TedsQueryStack/Private/TedsRowFilterNode.cpp
--- a/Engine/Plugins/Experimental/EditorDataStorageFeatures/Source/TedsQueryStack/Private/TedsRowFilterNode.cpp
+++ b/Engine/Plugins/Experimental/EditorDataStorageFeatures/Source/TedsQueryStack/Private/TedsRowFilterNode.cpp
@@ -37,6 +37,6 @@ namespace UE::Editor::DataStorage::QueryStack
 
 	void FRowFilterNode::ForceRefresh()
 	{
-		CachedParentRevisionID = 0;
+		CachedParentRevisionID = UninitializedRevisionId;
 	}
 } // namespace UE::Editor::DataStorage::QueryStack
diff --git a/Engine/Plugins/Experimental/EditorDataStorageFeatures/Source/TedsQueryStack/Private/TedsRowMonitorNode.cpp b/Engine/Plugins/Experimental/EditorDataStorageFeatures/Source/TedsQueryStack/Private/TedsRowMonitorNode.cpp
--- a/Engine/Plugins/Experimental/EditorDataStorageFeatures/Source/TedsQueryStack/Private/TedsRowMonitorNode.cpp
+++ b/Engine/Plugins/Experimental/EditorDataStorageFeatures/Source/TedsQueryStack/Private/TedsRowMonitorNode.cpp
@@ -227,7 +227,7 @@ namespace UE::Editor::DataStorage::QueryStack
 			else
 			{
 				MonitoredColumns.Empty();
-				QueryRevision = 0;
+				QueryRevision = UninitializedRevisionId;
 			}
 		}
 		else
diff --git a/Engine/Plugins/Experimental/EditorDataStorageFeatures/Source/TedsQueryStack/Public/TedsHierarchyNode.h b/Engine/Plugins/Experimental/EditorDataStorageFeatures/Source/TedsQueryStack/Public/TedsHierarchyNode.h
--- a/Engine/Plugins/Experimental/EditorDataStorageFeatures/Source/TedsQueryStack/Public/TedsHierarchyNode.h
+++ b/Engine/Plugins/Experimental/EditorDataStorageFeatures/Source/TedsQueryStack/Public/TedsHierarchyNode.h
@@ -49,12 +49,12 @@ namespace UE::Editor::DataStorage
 		FRowHandleArray Rows;
 		FRowHandleArray NewRows;
 		TSharedPtr<QueryStack::IRowNode> TopLevelRows;
-		RevisionId ParentRevision = 0;
+		RevisionId ParentRevision = UninitializedRevisionId;
 		
 		ICoreProvider& Storage;
 		FHierarchyHandle HierarchyHandle;
 		
-		RevisionId Revision = 0;
+		RevisionId Revision = UninitializedRevisionId;
 		ESyncFlags SyncFlags;
 	};
 
diff --git a/Engine/Plugins/Experimental/EditorDataStorageFeatures/Source/TedsQueryStack/Public/TedsQueryHandleNode.h b/Engine/Plugins/Experimental/EditorDataStorageFeatures/Source/TedsQueryStack/Public/TedsQueryHandleNode.h
--- a/Engine/Plugins/Experimental/EditorDataStorageFeatures/Source/TedsQueryStack/Public/TedsQueryHandleNode.h
+++ b/Engine/Plugins/Experimental/EditorDataStorageFeatures/Source/TedsQueryStack/Public/TedsQueryHandleNode.h
@@ -27,6 +27,6 @@ namespace UE::Editor::DataStorage::QueryStack
 
 	private:
 		QueryHandle Query = InvalidQueryHandle;
-		RevisionId Revision = 0;
+		RevisionId Revision = UninitializedRevisionId;
 	};
 } // namespace UE::Editor::DataStorage::QueryStack
diff --git a/Engine/Plugins/Experimental/EditorDataStorageFeatures/Source/TedsQueryStack/Public/TedsQueryMergeNode.h b/Engine/Plugins/Experimental/EditorDataStorageFeatures/Source/TedsQueryStack/Public/TedsQueryMergeNode.h
--- a/Engine/Plugins/Experimental/EditorDataStorageFeatures/Source/TedsQueryStack/Public/TedsQueryMergeNode.h
+++ b/Engine/Plugins/Experimental/EditorDataStorageFeatures/Source/TedsQueryStack/Public/TedsQueryMergeNode.h
@@ -41,6 +41,6 @@ namespace UE::Editor::DataStorage::QueryStack
 		QueryHandle QueryHandle = InvalidQueryHandle;
 		TArray<FParentInfo> Parents;
 		ICoreProvider& Storage;
-		RevisionId Revision = 0;
+		RevisionId Revision = UninitializedRevisionId;
 	};
 } // namespace UE::Editor::DataStorage::QueryStack
diff --git a/Engine/Plugins/Experimental/EditorDataStorageFeatures/Source/TedsQueryStack/Public/TedsQueryNode.h b/Engine/Plugins/Experimental/EditorDataStorageFeatures/Source/TedsQueryStack/Public/TedsQueryNode.h
--- a/Engine/Plugins/Experimental/EditorDataStorageFeatures/Source/TedsQueryStack/Public/TedsQueryNode.h
+++ b/Engine/Plugins/Experimental/EditorDataStorageFeatures/Source/TedsQueryStack/Public/TedsQueryNode.h
@@ -35,6 +35,6 @@ namespace UE::Editor::DataStorage::QueryStack
 	private:
 		QueryHandle QueryHandle = InvalidQueryHandle;
 		ICoreProvider& Storage;
-		RevisionId Revision = 0;
+		RevisionId Revision = UninitializedRevisionId;
 	};
 } // namespace UE::Editor::DataStorage::QueryStack
diff --git a/Engine/Plugins/Experimental/EditorDataStorageFeatures/Source/TedsQueryStack/Public/TedsQueryStackInterfaces.h b/Engine/Plugins/Experimental/EditorDataStorageFeatures/Source/TedsQueryStack/Public/TedsQueryStackInterfaces.h
--- a/Engine/Plugins/Experimental/EditorDataStorageFeatures/Source/TedsQueryStack/Public/TedsQueryStackInterfaces.h
+++ b/Engine/Plugins/Experimental/EditorDataStorageFeatures/Source/TedsQueryStack/Public/TedsQueryStackInterfaces.h
@@ -18,6 +18,9 @@ namespace UE::Editor::DataStorage::QueryStack
 		friend class FExecutorBase;
 
 		using RevisionId = uint32;
+		// The default uninitialized revision Id is zero to signify it is empty and needs to be updated/rows need to be initialized
+		static constexpr RevisionId UninitializedRevisionId = 0;
+		
 		using ParentListCallback = TFunctionRef<void(const TSharedPtr<INode>& Parent)>;
 		static inline FTimespan UnlimitedTime = FTimespan::MaxValue();
 
diff --git a/Engine/Plugins/Experimental/EditorDataStorageFeatures/Source/TedsQueryStack/Public/TedsRowArrayNode.h b/Engine/Plugins/Experimental/EditorDataStorageFeatures/Source/TedsQueryStack/Public/TedsRowArrayNode.h
--- a/Engine/Plugins/Experimental/EditorDataStorageFeatures/Source/TedsQueryStack/Public/TedsRowArrayNode.h
+++ b/Engine/Plugins/Experimental/EditorDataStorageFeatures/Source/TedsQueryStack/Public/TedsRowArrayNode.h
@@ -32,6 +32,6 @@ namespace UE::Editor::DataStorage::QueryStack
 	
 	private:
 		FRowHandleArray Rows;
-		RevisionId Revision = 0;
+		RevisionId Revision = UninitializedRevisionId;
 	};
 } // namespace UE::Editor::DataStorage::QueryStack
diff --git a/Engine/Plugins/Experimental/EditorDataStorageFeatures/Source/TedsQueryStack/Public/TedsRowCopyNode.h b/Engine/Plugins/Experimental/EditorDataStorageFeatures/Source/TedsQueryStack/Public/TedsRowCopyNode.h
--- a/Engine/Plugins/Experimental/EditorDataStorageFeatures/Source/TedsQueryStack/Public/TedsRowCopyNode.h
+++ b/Engine/Plugins/Experimental/EditorDataStorageFeatures/Source/TedsQueryStack/Public/TedsRowCopyNode.h
@@ -34,7 +34,7 @@ namespace UE::Editor::DataStorage::QueryStack
 	private:
 		FRowHandleArray Rows;
 		TSharedPtr<IRowNode> Parent;
-		RevisionId ParentRevision = 0;
-		RevisionId Revision = 0;
+		RevisionId ParentRevision = UninitializedRevisionId;
+		RevisionId Revision = UninitializedRevisionId;
 	};
 } // namespace UE::Editor::DataStorage::QueryStack
diff --git a/Engine/Plugins/Experimental/EditorDataStorageFeatures/Source/TedsQueryStack/Public/TedsRowFilterNode.h b/Engine/Plugins/Experimental/EditorDataStorageFeatures/Source/TedsQueryStack/Public/TedsRowFilterNode.h
--- a/Engine/Plugins/Experimental/EditorDataStorageFeatures/Source/TedsQueryStack/Public/TedsRowFilterNode.h
+++ b/Engine/Plugins/Experimental/EditorDataStorageFeatures/Source/TedsQueryStack/Public/TedsRowFilterNode.h
@@ -27,7 +27,6 @@ namespace UE::Editor::DataStorage::QueryStack
 			: Storage(Storage)
 			, Filter(Filter)
 			, ParentRowNode(ParentRowNode)
-			, CachedParentRevisionID(ParentRowNode->GetRevision())
 			, Options(Options)
 		{
 		}
@@ -85,7 +84,7 @@ namespace UE::Editor::DataStorage::QueryStack
 		Queries::TQueryFunction<bool> Filter;
 		FRowHandleArray Rows;
 		TSharedPtr<IRowNode> ParentRowNode;
-		RevisionId CachedParentRevisionID;
+		RevisionId CachedParentRevisionID = UninitializedRevisionId;
 		ICoreProvider::EFilterOptions Options;
 	};
 }
diff --git a/Engine/Plugins/Experimental/EditorDataStorageFeatures/Source/TedsQueryStack/Public/TedsRowHandleSortNode.h b/Engine/Plugins/Experimental/EditorDataStorageFeatures/Source/TedsQueryStack/Public/TedsRowHandleSortNode.h
--- a/Engine/Plugins/Experimental/EditorDataStorageFeatures/Source/TedsQueryStack/Public/TedsRowHandleSortNode.h
+++ b/Engine/Plugins/Experimental/EditorDataStorageFeatures/Source/TedsQueryStack/Public/TedsRowHandleSortNode.h
@@ -28,6 +28,6 @@ namespace UE::Editor::DataStorage::QueryStack
 	
 	private:
 		TSharedPtr<IRowNode> Parent;
-		RevisionId ParentRevision = 0;
+		RevisionId ParentRevision = UninitializedRevisionId;
 	};
 } // namespace UE::Editor::DataStorage::QueryStack
diff --git a/Engine/Plugins/Experimental/EditorDataStorageFeatures/Source/TedsQueryStack/Public/TedsRowMergeNode.h b/Engine/Plugins/Experimental/EditorDataStorageFeatures/Source/TedsQueryStack/Public/TedsRowMergeNode.h
--- a/Engine/Plugins/Experimental/EditorDataStorageFeatures/Source/TedsQueryStack/Public/TedsRowMergeNode.h
+++ b/Engine/Plugins/Experimental/EditorDataStorageFeatures/Source/TedsQueryStack/Public/TedsRowMergeNode.h
@@ -37,12 +37,12 @@ namespace UE::Editor::DataStorage::QueryStack
 		struct FParentInfo
 		{
 			TSharedPtr<IRowNode> Parent;
-			RevisionId Revision = 0;
+			RevisionId Revision = UninitializedRevisionId;
 		};
 
 		TArray<FParentInfo> Parents;
 		FRowHandleArray Rows;
-		RevisionId Revision = 0;
+		RevisionId Revision = UninitializedRevisionId;
 		EMergeApproach MergeApproach;
 	};
 } // namespace UE::Editor::DataStorage::QueryStack
diff --git a/Engine/Plugins/Experimental/EditorDataStorageFeatures/Source/TedsQueryStack/Public/TedsRowMonitorNode.h b/Engine/Plugins/Experimental/EditorDataStorageFeatures/Source/TedsQueryStack/Public/TedsRowMonitorNode.h
--- a/Engine/Plugins/Experimental/EditorDataStorageFeatures/Source/TedsQueryStack/Public/TedsRowMonitorNode.h
+++ b/Engine/Plugins/Experimental/EditorDataStorageFeatures/Source/TedsQueryStack/Public/TedsRowMonitorNode.h
@@ -71,9 +71,9 @@ namespace UE::Editor::DataStorage::QueryStack
 		TSharedPtr<IQueryNode> QueryNode;
 		TSharedPtr<IRowNode> ParentNode;
 		ICoreProvider& Storage;
-		RevisionId QueryRevision = 0;
-		RevisionId ParentRevision = 0;
-		RevisionId Revision = 0;
+		RevisionId QueryRevision = UninitializedRevisionId;
+		RevisionId ParentRevision = UninitializedRevisionId;
+		RevisionId Revision = UninitializedRevisionId;
 		bool bFixedColumns = false;
 		
 		FOnRowsChanged RowsAddedEvent;
diff --git a/Engine/Plugins/Experimental/EditorDataStorageFeatures/Source/TedsQueryStack/Public/TedsRowOrderInversionNode.h b/Engine/Plugins/Experimental/EditorDataStorageFeatures/Source/TedsQueryStack/Public/TedsRowOrderInversionNode.h
--- a/Engine/Plugins/Experimental/EditorDataStorageFeatures/Source/TedsQueryStack/Public/TedsRowOrderInversionNode.h
+++ b/Engine/Plugins/Experimental/EditorDataStorageFeatures/Source/TedsQueryStack/Public/TedsRowOrderInversionNode.h
@@ -34,8 +34,8 @@ namespace UE::Editor::DataStorage::QueryStack
 		void ApplyOrder();
 
 		TSharedPtr<IRowNode> Parent;
-		RevisionId ParentRevision = 0;
-		RevisionId Revision = 0;
+		RevisionId ParentRevision = UninitializedRevisionId;
+		RevisionId Revision = UninitializedRevisionId;
 		bool bIsEnabled = true;
 		bool bRequiresSync = false;
 	};
diff --git a/Engine/Plugins/Experimental/EditorDataStorageFeatures/Source/TedsQueryStack/Public/TedsRowQueryResultsNode.h b/Engine/Plugins/Experimental/EditorDataStorageFeatures/Source/TedsQueryStack/Public/TedsRowQueryResultsNode.h
--- a/Engine/Plugins/Experimental/EditorDataStorageFeatures/Source/TedsQueryStack/Public/TedsRowQueryResultsNode.h
+++ b/Engine/Plugins/Experimental/EditorDataStorageFeatures/Source/TedsQueryStack/Public/TedsRowQueryResultsNode.h
@@ -57,8 +57,8 @@ namespace UE::Editor::DataStorage::QueryStack
 		FRowHandleArray Rows;
 		TSharedPtr<IQueryNode> QueryNode;
 		ICoreProvider& Storage;
-		RevisionId QueryRevision = 0;
-		RevisionId Revision = 0;
+		RevisionId QueryRevision = UninitializedRevisionId;
+		RevisionId Revision = UninitializedRevisionId;
 		ESyncFlags SyncFlags = ESyncFlags::None;
 	};
 
diff --git a/Engine/Plugins/Experimental/EditorDataStorageFeatures/Source/TedsQueryStack/Public/TedsRowSortNode.h b/Engine/Plugins/Experimental/EditorDataStorageFeatures/Source/TedsQueryStack/Public/TedsRowSortNode.h
--- a/Engine/Plugins/Experimental/EditorDataStorageFeatures/Source/TedsQueryStack/Public/TedsRowSortNode.h
+++ b/Engine/Plugins/Experimental/EditorDataStorageFeatures/Source/TedsQueryStack/Public/TedsRowSortNode.h
@@ -62,7 +62,7 @@ namespace UE::Editor::DataStorage::QueryStack
 
 		FTimespan MaxFrameDuration;
 		RevisionId LastUpdatedRevision = TNumericLimits<RevisionId>::Max();
-		RevisionId Revision = 0;
+		RevisionId Revision = UninitializedRevisionId;
 		
 		ICoreProvider& Storage;
 		
diff --git a/Engine/Plugins/Experimental/EditorDataStorageFeatures/Source/TedsQueryStack/Public/TedsRowViewNode.h b/Engine/Plugins/Experimental/EditorDataStorageFeatures/Source/TedsQueryStack/Public/TedsRowViewNode.h
--- a/Engine/Plugins/Experimental/EditorDataStorageFeatures/Source/TedsQueryStack/Public/TedsRowViewNode.h
+++ b/Engine/Plugins/Experimental/EditorDataStorageFeatures/Source/TedsQueryStack/Public/TedsRowViewNode.h
@@ -30,6 +30,6 @@ namespace UE::Editor::DataStorage::QueryStack
 	
 	private:
 		FRowHandleArray* Rows = nullptr;
-		RevisionId Revision = 0;
+		RevisionId Revision = 1;
 	};
 } // namespace UE::Editor::DataStorage::QueryStack
diff --git a/Engine/Plugins/Experimental/EditorDataStorageFeatures/Source/TedsQueryStack/Tests/TedsRowViewNodeTests.cpp b/Engine/Plugins/Experimental/EditorDataStorageFeatures/Source/TedsQueryStack/Tests/TedsRowViewNodeTests.cpp
--- a/Engine/Plugins/Experimental/EditorDataStorageFeatures/Source/TedsQueryStack/Tests/TedsRowViewNodeTests.cpp
+++ b/Engine/Plugins/Experimental/EditorDataStorageFeatures/Source/TedsQueryStack/Tests/TedsRowViewNodeTests.cpp
@@ -15,7 +15,7 @@ TEST_CASE_NAMED(TEDS_QueryStack_RowViewNode_Tests, "Editor::DataStorage::QuerySt
 		FRowHandleArray RowHandles;
 		RowHandles.Append({ 1, 2, 3 }, FRowHandleArray::EFlags::IsSorted | FRowHandleArray::EFlags::IsUnique);
 		FRowViewNode View(RowHandles);
-		CHECK_EQUALS(TEXT("Revision"), View.GetRevision(), 0);
+		CHECK_EQUALS(TEXT("Revision"), View.GetRevision(), 1);
 		
 		FRowHandleArrayView Rows = View.GetRows();
 		CHECK_MESSAGE(TEXT("Expected rows to be sorted."), Rows.IsSorted());
@@ -23,16 +23,16 @@ TEST_CASE_NAMED(TEDS_QueryStack_RowViewNode_Tests, "Editor::DataStorage::QuerySt
 		CHECK_EQUALS(TEXT("Rows[0]"), Rows[0], 1llu);
 		CHECK_EQUALS(TEXT("Rows[1]"), Rows[1], 2llu);
 		CHECK_EQUALS(TEXT("Rows[2]"), Rows[2], 3llu);
-	}
+	} 
 
 	SECTION("Mark dirty")
 	{
 		FRowHandleArray RowHandles;
 		FRowViewNode View(RowHandles);
-		CHECK_EQUALS(TEXT("Revision"), View.GetRevision(), 0);
+		CHECK_EQUALS(TEXT("Revision"), View.GetRevision(), 1);
 		
 		View.MarkDirty();
-		CHECK_EQUALS(TEXT("Revision"), View.GetRevision(), 1);
+		CHECK_EQUALS(TEXT("Revision"), View.GetRevision(), 2);
 	}
 }
 
