[TEDS] Fixed a potential null pointer crash.

A while back the rules about needing either a parent or having a local array were enforced, but the monitor node didn't get the memo. This has now been addressed by requiring a parent at all time. If the intend is to only work on rows reported by the observers then a parent node can be provided that just stores the rows.

diff --git a/Engine/Plugins/Experimental/EditorDataStorageFeatures/Source/TedsQueryStack/Private/TedsRowMonitorNode.cpp b/Engine/Plugins/Experimental/EditorDataStorageFeatures/Source/TedsQueryStack/Private/TedsRowMonitorNode.cpp
--- a/Engine/Plugins/Experimental/EditorDataStorageFeatures/Source/TedsQueryStack/Private/TedsRowMonitorNode.cpp
+++ b/Engine/Plugins/Experimental/EditorDataStorageFeatures/Source/TedsQueryStack/Private/TedsRowMonitorNode.cpp
@@ -8,59 +8,36 @@
 namespace UE::Editor::DataStorage::QueryStack
 {
 	FRowMonitorNode::FRowMonitorNode(
-		ICoreProvider& Storage, TSharedPtr<IRowNode> ParentNode, TSharedPtr<IQueryNode> QueryNode)
-		: QueryNode(MoveTemp(QueryNode))
-		, ParentNode(MoveTemp(ParentNode))
-		, Storage(Storage)
+		ICoreProvider& InStorage, TSharedPtr<IRowNode> InParentNode, TSharedPtr<IQueryNode> InQueryNode)
+		: QueryNode(MoveTemp(InQueryNode))
+		, ParentNode(MoveTemp(InParentNode))
+		, Storage(InStorage)
 	{
+		checkf(ParentNode, TEXT("FRowMonitorNode requires a valid ParentNode."));
 		UpdateColumnsFromQuery();
 		UpdateMonitoredColumns();
 	}
 
 	FRowMonitorNode::FRowMonitorNode(
-		ICoreProvider& Storage, TSharedPtr<IRowNode> ParentNode, TArray<TObjectPtr<const UScriptStruct>> Columns)
-		: MonitoredColumns(MoveTemp(Columns))
-		, ParentNode(MoveTemp(ParentNode))
-		, Storage(Storage)
+		ICoreProvider& InStorage, TSharedPtr<IRowNode> InParentNode, TArray<TObjectPtr<const UScriptStruct>> InColumns)
+		: MonitoredColumns(MoveTemp(InColumns))
+		, ParentNode(MoveTemp(InParentNode))
+		, Storage(InStorage)
 	{
+		checkf(ParentNode, TEXT("FRowMonitorNode requires a valid ParentNode."));
 		UpdateMonitoredColumns();
 	}
 
 	FRowMonitorNode::FRowMonitorNode(
-		ICoreProvider& Storage, TSharedPtr<IQueryNode> QueryNode, TSharedPtr<IRowNode> ParentNode,
-		TArray<TObjectPtr<const UScriptStruct>> MonitoredColumns)
-		: MonitoredColumns(MoveTemp(MonitoredColumns))
-		, QueryNode(MoveTemp(QueryNode))
-		, ParentNode(MoveTemp(ParentNode))
-		, Storage(Storage)
-		, bFixedColumns(true)
-	{
-		UpdateColumnsFromQuery();
-		UpdateMonitoredColumns();
-	}
-
-	FRowMonitorNode::FRowMonitorNode(ICoreProvider& Storage, TSharedPtr<IQueryNode> QueryNode)
-		: QueryNode(MoveTemp(QueryNode))
-		, Storage(Storage)
-	{
-		UpdateColumnsFromQuery();
-		UpdateMonitoredColumns();
-	}
-	
-	FRowMonitorNode::FRowMonitorNode(ICoreProvider& Storage, TArray<TObjectPtr<const UScriptStruct>> Columns)
-		: MonitoredColumns(MoveTemp(Columns))
-		, Storage(Storage)
-	{
-		UpdateMonitoredColumns();
-	}
-
-	FRowMonitorNode::FRowMonitorNode(
-		ICoreProvider& Storage, TSharedPtr<IQueryNode> QueryNode, TArray<TObjectPtr<const UScriptStruct>> MonitoredColumns)
-		: MonitoredColumns(MoveTemp(MonitoredColumns))
-		, QueryNode(MoveTemp(QueryNode))
-		, Storage(Storage)
+		ICoreProvider& InStorage, TSharedPtr<IQueryNode> InQueryNode, TSharedPtr<IRowNode> InParentNode,
+		TArray<TObjectPtr<const UScriptStruct>> InMonitoredColumns)
+		: MonitoredColumns(MoveTemp(InMonitoredColumns))
+		, QueryNode(MoveTemp(InQueryNode))
+		, ParentNode(MoveTemp(InParentNode))
+		, Storage(InStorage)
 		, bFixedColumns(true)
 	{
+		checkf(ParentNode, TEXT("FRowMonitorNode requires a valid ParentNode."));
 		UpdateColumnsFromQuery();
 		UpdateMonitoredColumns();
 	}
@@ -135,10 +112,7 @@ namespace UE::Editor::DataStorage::QueryStack
 		{
 			Callback(QueryNode);
 		}
-		if (ParentNode)
-		{
-			Callback(ParentNode);
-		}
+		Callback(ParentNode);
 	}
 
 	FRowHandleArrayView FRowMonitorNode::GetRows() const
@@ -230,10 +204,6 @@ namespace UE::Editor::DataStorage::QueryStack
 				QueryRevision = UninitializedRevisionId;
 			}
 		}
-		else
-		{
-			QueryRevision = QueryNode->GetRevision();
-		}
 	}
 
 	void FRowMonitorNode::UpdateMonitoredColumns()
@@ -247,7 +217,7 @@ namespace UE::Editor::DataStorage::QueryStack
 			Storage.UnregisterQuery(Observer);
 		}
 
-		Observers.Reserve(MonitoredColumns.Num() * 2); // Add double, one for the OnAdd and one for the OnRemove.
+		Observers.Reset(MonitoredColumns.Num() * 2); // Add double, one for the OnAdd and one for the OnRemove.
 
 		if (QueryNode)
 		{
diff --git a/Engine/Plugins/Experimental/EditorDataStorageFeatures/Source/TedsQueryStack/Public/TedsRowMonitorNode.h b/Engine/Plugins/Experimental/EditorDataStorageFeatures/Source/TedsQueryStack/Public/TedsRowMonitorNode.h
--- a/Engine/Plugins/Experimental/EditorDataStorageFeatures/Source/TedsQueryStack/Public/TedsRowMonitorNode.h
+++ b/Engine/Plugins/Experimental/EditorDataStorageFeatures/Source/TedsQueryStack/Public/TedsRowMonitorNode.h
@@ -24,18 +24,12 @@ namespace UE::Editor::DataStorage::QueryStack
 		DECLARE_MULTICAST_DELEGATE_OneParam(FOnRowsChanged, FRowHandleArrayView /* InRows */);
 		
 		TEDSQUERYSTACK_API FRowMonitorNode(
-			ICoreProvider& Storage, TSharedPtr<IRowNode> ParentNode, TSharedPtr<IQueryNode> QueryNode);
+			ICoreProvider& InStorage, TSharedPtr<IRowNode> InParentNode, TSharedPtr<IQueryNode> InQueryNode);
 		TEDSQUERYSTACK_API FRowMonitorNode(
-			ICoreProvider& Storage, TSharedPtr<IRowNode> ParentNode, TArray<TObjectPtr<const UScriptStruct>> Columns);
+			ICoreProvider& InStorage, TSharedPtr<IRowNode> InParentNode, TArray<TObjectPtr<const UScriptStruct>> InColumns);
 		TEDSQUERYSTACK_API FRowMonitorNode(
-			ICoreProvider& Storage, TSharedPtr<IQueryNode> QueryNode, TSharedPtr<IRowNode> ParentNode,
-			TArray<TObjectPtr<const UScriptStruct>> MonitoredColumns);
-		TEDSQUERYSTACK_API FRowMonitorNode(
-			ICoreProvider& Storage, TSharedPtr<IQueryNode> QueryNode);
-		TEDSQUERYSTACK_API FRowMonitorNode(
-			ICoreProvider& Storage, TArray<TObjectPtr<const UScriptStruct>> Columns);
-		TEDSQUERYSTACK_API FRowMonitorNode(
-			ICoreProvider& Storage, TSharedPtr<IQueryNode> QueryNode, TArray<TObjectPtr<const UScriptStruct>> MonitoredColumns);
+			ICoreProvider& InStorage, TSharedPtr<IQueryNode> InQueryNode, TSharedPtr<IRowNode> InParentNode,
+			TArray<TObjectPtr<const UScriptStruct>> InMonitoredColumns);
 
 		TEDSQUERYSTACK_API virtual ~FRowMonitorNode() override;
 
