TEDS Hierarchy Viewer:

Update the hierarchy interface to add a depth first walk function and an event that is broadcast on parent change
Remove the legacy hierarchy data
Add a new hierarchy data implementation to collate multiple hierarchies into one view

TEDS Hierarchies:
Add an function to the TEDS interface to get the parent changed column for a hierarchy

Everything Picker:
Update code to remove usage of legacy hierarchy viewer data

diff --git a/Engine/Plugins/Experimental/EditorDataStorage/Source/TedsCore/Private/Hierarchy/EditorDataStorageHierarchyImplementation.cpp b/Engine/Plugins/Experimental/EditorDataStorage/Source/TedsCore/Private/Hierarchy/EditorDataStorageHierarchyImplementation.cpp
--- a/Engine/Plugins/Experimental/EditorDataStorage/Source/TedsCore/Private/Hierarchy/EditorDataStorageHierarchyImplementation.cpp
+++ b/Engine/Plugins/Experimental/EditorDataStorage/Source/TedsCore/Private/Hierarchy/EditorDataStorageHierarchyImplementation.cpp
@@ -365,6 +365,11 @@ namespace UE::Editor::DataStorage
 		return UnresolvedParentColumnType;
 	}
 
+	const UScriptStruct* FTedsHierarchyAccessInterface::GetParentChangedColumnType() const
+	{
+		return ParentChangedColumn;
+	}
+
 	bool FTedsHierarchyAccessInterface::HasChildren(const ICoreProvider& Context, RowHandle Row) const
 	{
 		const FEditorDataHierarchyData_Template* Hierarchy = static_cast<const FEditorDataHierarchyData_Template*>(Context.GetColumnData(Row, HierarchyDataColumnType));
diff --git a/Engine/Plugins/Experimental/EditorDataStorage/Source/TedsCore/Private/Hierarchy/EditorDataStorageHierarchyImplementation.h b/Engine/Plugins/Experimental/EditorDataStorage/Source/TedsCore/Private/Hierarchy/EditorDataStorageHierarchyImplementation.h
--- a/Engine/Plugins/Experimental/EditorDataStorage/Source/TedsCore/Private/Hierarchy/EditorDataStorageHierarchyImplementation.h
+++ b/Engine/Plugins/Experimental/EditorDataStorage/Source/TedsCore/Private/Hierarchy/EditorDataStorageHierarchyImplementation.h
@@ -70,6 +70,8 @@ namespace UE::Editor::DataStorage
 		const UScriptStruct* GetParentTagType() const;
 		const UScriptStruct* GetHierarchyDataColumnType() const;
 		const UScriptStruct* GetUnresolvedParentColumnType() const;
+		const UScriptStruct* GetParentChangedColumnType() const;
+
 
 		void SetParentRow(ICoreProvider* CoreProvider, RowHandle Target, RowHandle Parent) const;
 		void SetUnresolvedParent(ICoreProvider* CoreProvider, RowHandle Target, FMapKey ParentId, FName MappingDomain) const;
diff --git a/Engine/Plugins/Experimental/EditorDataStorage/Source/TedsCore/Private/TypedElementDatabase.cpp b/Engine/Plugins/Experimental/EditorDataStorage/Source/TedsCore/Private/TypedElementDatabase.cpp
--- a/Engine/Plugins/Experimental/EditorDataStorage/Source/TedsCore/Private/TypedElementDatabase.cpp
+++ b/Engine/Plugins/Experimental/EditorDataStorage/Source/TedsCore/Private/TypedElementDatabase.cpp
@@ -1127,6 +1127,15 @@ const UScriptStruct* UEditorDataStorage::GetHierarchyDataColumnType(UE::Editor::
 	return nullptr;
 }
 
+const UScriptStruct* UEditorDataStorage::GetParentChangedColumnType(UE::Editor::DataStorage::FHierarchyHandle InHierarchyHandle) const
+{
+	if (const UE::Editor::DataStorage::FTedsHierarchyAccessInterface* AccessInterface = Environment->GetHierarchyRegistrar().GetAccessInterface(InHierarchyHandle))
+	{
+		return AccessInterface->GetParentChangedColumnType();
+	}
+	return nullptr;
+}
+
 void UEditorDataStorage::ListHierarchyNames(TFunctionRef<void(const FName& HierarchyName)> Callback) const
 {
 	Environment->GetHierarchyRegistrar().ListHierarchyNames(Callback);
@@ -1182,7 +1191,7 @@ bool UEditorDataStorage::HasChildren(UE::Editor::DataStorage::FHierarchyHandle I
 void UEditorDataStorage::WalkDepthFirst(
 	UE::Editor::DataStorage::FHierarchyHandle InHierarchyHandle,
 	UE::Editor::DataStorage::RowHandle Row,
-	TFunction<void(const ICoreProvider& Context, UE::Editor::DataStorage::RowHandle Owner, UE::Editor::DataStorage::RowHandle Target)> VisitFn) const
+	FHierarchyIterationCallback VisitFn) const
 {
 	if (const UE::Editor::DataStorage::FTedsHierarchyAccessInterface* AccessInterface = Environment->GetHierarchyRegistrar().GetAccessInterface(InHierarchyHandle))
 	{
diff --git a/Engine/Plugins/Experimental/EditorDataStorage/Source/TedsCore/Private/TypedElementDatabase.h b/Engine/Plugins/Experimental/EditorDataStorage/Source/TedsCore/Private/TypedElementDatabase.h
--- a/Engine/Plugins/Experimental/EditorDataStorage/Source/TedsCore/Private/TypedElementDatabase.h
+++ b/Engine/Plugins/Experimental/EditorDataStorage/Source/TedsCore/Private/TypedElementDatabase.h
@@ -149,6 +149,8 @@ public:
 		UE::Editor::DataStorage::FHierarchyHandle) const override;
 	UE_API virtual const UScriptStruct* GetHierarchyDataColumnType(
 		UE::Editor::DataStorage::FHierarchyHandle) const override;
+	UE_API virtual const UScriptStruct* GetParentChangedColumnType(
+		UE::Editor::DataStorage::FHierarchyHandle InHierarchyHandle) const override;
 	UE_API virtual void ListHierarchyNames(TFunctionRef<void(const FName& HierarchyName)> Callback) const override;
 	UE_API virtual void SetParentRow(
 		UE::Editor::DataStorage::FHierarchyHandle, UE::Editor::DataStorage::RowHandle Target,	UE::Editor::DataStorage::RowHandle Parent) override;
diff --git a/Engine/Plugins/Experimental/EditorDataStorageFeatures/Source/TedsActorCompatibility/Private/ActorComponentDebugHierarchyWidget/SActorComponentTreeWidget.cpp b/Engine/Plugins/Experimental/EditorDataStorageFeatures/Source/TedsActorCompatibility/Private/ActorComponentDebugHierarchyWidget/SActorComponentTreeWidget.cpp
--- a/Engine/Plugins/Experimental/EditorDataStorageFeatures/Source/TedsActorCompatibility/Private/ActorComponentDebugHierarchyWidget/SActorComponentTreeWidget.cpp
+++ b/Engine/Plugins/Experimental/EditorDataStorageFeatures/Source/TedsActorCompatibility/Private/ActorComponentDebugHierarchyWidget/SActorComponentTreeWidget.cpp
@@ -23,20 +23,23 @@ void SActorComponentTreeWidget::Construct(
 	using namespace UE::Editor::DataStorage;
 	using namespace UE::Editor::DataStorage::Queries;
 
-	TSharedPtr<FHierarchyViewerData> HierarchyData = MakeShared<FHierarchyViewerData>(InHierarchyHandle);
+	if (ICoreProvider* Storage = GetMutableDataStorageFeature<ICoreProvider>(StorageFeatureName))
+	{
+		TSharedPtr<FHierarchyViewerData> HierarchyData = MakeShared<FHierarchyViewerData>(*Storage,InHierarchyHandle);
 	
-	TSharedRef<SHierarchyViewer> HierarchyViewer =
-		SNew(SHierarchyViewer,
-			MoveTemp(HierarchyData))
-			.AllNodeProvider(InRowProvider)
-			.EmptyRowsMessage(LOCTEXT("EmptyRowsMessage", "The provided query has no results."))
-			.CellWidgetPurpose(IUiProvider::FPurposeInfo("General", "RowLabel", NAME_None).GeneratePurposeID())
-			.Columns({FTypedElementLabelColumn::StaticStruct()});
-
-	ChildSlot
-	[
-		HierarchyViewer
-	];
+		TSharedRef<SHierarchyViewer> HierarchyViewer =
+			SNew(SHierarchyViewer,
+				MoveTemp(HierarchyData))
+				.AllNodeProvider(InRowProvider)
+				.EmptyRowsMessage(LOCTEXT("EmptyRowsMessage", "The provided query has no results."))
+				.CellWidgetPurpose(IUiProvider::FPurposeInfo("General", "RowLabel", NAME_None).GeneratePurposeID())
+				.Columns({FTypedElementLabelColumn::StaticStruct()});
+
+		ChildSlot
+		[
+			HierarchyViewer
+		];
+	}
 }
 
 #undef LOCTEXT_NAMESPACE
diff --git a/Engine/Plugins/Experimental/EditorDataStorageFeatures/Source/TedsDebugger/Private/Widgets/QueryEditor/TedsQueryEditorResultsView.cpp b/Engine/Plugins/Experimental/EditorDataStorageFeatures/Source/TedsDebugger/Private/Widgets/QueryEditor/TedsQueryEditorResultsView.cpp
--- a/Engine/Plugins/Experimental/EditorDataStorageFeatures/Source/TedsDebugger/Private/Widgets/QueryEditor/TedsQueryEditorResultsView.cpp
+++ b/Engine/Plugins/Experimental/EditorDataStorageFeatures/Source/TedsDebugger/Private/Widgets/QueryEditor/TedsQueryEditorResultsView.cpp
@@ -228,7 +228,7 @@ namespace UE::Editor::DataStorage::Debug::QueryEditor
 		
 		case ETableViewMode::Tree:
 
-			HierarchyData = MakeShared<FHierarchyViewerData>(HierarchyHandle);
+			HierarchyData = MakeShared<FHierarchyViewerData>(Model->GetTedsInterface(), HierarchyHandle);
 
 			TableViewer =
 				SNew(SHierarchyViewer, HierarchyData)
diff --git a/Engine/Plugins/Experimental/EditorDataStorageFeatures/Source/TedsEverythingPicker/Private/Context/TedsPickerContextUtil.cpp b/Engine/Plugins/Experimental/EditorDataStorageFeatures/Source/TedsEverythingPicker/Private/Context/TedsPickerContextUtil.cpp
--- a/Engine/Plugins/Experimental/EditorDataStorageFeatures/Source/TedsEverythingPicker/Private/Context/TedsPickerContextUtil.cpp
+++ b/Engine/Plugins/Experimental/EditorDataStorageFeatures/Source/TedsEverythingPicker/Private/Context/TedsPickerContextUtil.cpp
@@ -83,15 +83,8 @@ namespace UE::Editor::DataStorage::Picker
 		// Using the Teds Outliner Purpose for matching customization for now
 		IUiProvider::FPurposeID PurposeId = IUiProvider::FPurposeInfo("SceneOutliner", "RowLabel", NAME_None).GeneratePurposeID();
 
-		auto ParentExtractionFunction = [](const void* Column, const UScriptStruct*)->RowHandle
-		{
-			const FTableRowParentColumn* ParentColumn = static_cast<const FTableRowParentColumn*>(Column);
-			return ParentColumn->Parent;
-			
-		};
-
-		TSharedPtr<FHierarchyViewerLegacyData> HierarchyData = MakeShared<FHierarchyViewerLegacyData>(FTableRowParentColumn::StaticStruct(), MoveTemp(ParentExtractionFunction));
-
+		TSharedPtr<FHierarchyViewerData> HierarchyData = MakeShared<FHierarchyViewerData>(*Storage, InArgs._HierarchyHandle);
+		
 		TSharedPtr<SHierarchyViewer> ReferenceView = SNew(SHierarchyViewer, HierarchyData)
 			.AllNodeProvider(FinalNode)
 			.Columns({ FTypedElementLabelColumn::StaticStruct(), FTypedElementClassTypeInfoColumn::StaticStruct() })
diff --git a/Engine/Plugins/Experimental/EditorDataStorageFeatures/Source/TedsEverythingPicker/Public/Context/TedsPickerContextUtil.h b/Engine/Plugins/Experimental/EditorDataStorageFeatures/Source/TedsEverythingPicker/Public/Context/TedsPickerContextUtil.h
--- a/Engine/Plugins/Experimental/EditorDataStorageFeatures/Source/TedsEverythingPicker/Public/Context/TedsPickerContextUtil.h
+++ b/Engine/Plugins/Experimental/EditorDataStorageFeatures/Source/TedsEverythingPicker/Public/Context/TedsPickerContextUtil.h
@@ -29,6 +29,7 @@ namespace UE::Editor::DataStorage::Picker
 			SLATE_ARGUMENT(FQueryDescription, Query)
 			SLATE_ARGUMENT(TSharedPtr<QueryStack::IRowNode>, QueryStack) // Advanced: QueryStack allows finer control over displayed data. Will be used instead of Query if specified
 			SLATE_ARGUMENT(FOnSelectionChanged, OnSelectionChanged)
+			SLATE_ARGUMENT(FHierarchyHandle, HierarchyHandle)
 			SLATE_ARGUMENT(bool, SearchingEnabled)
 		SLATE_END_ARGS()
 
diff --git a/Engine/Plugins/Experimental/EditorDataStorageFeatures/Source/TedsTableViewer/Private/HierarchyViewerIntefaces.cpp b/Engine/Plugins/Experimental/EditorDataStorageFeatures/Source/TedsTableViewer/Private/HierarchyViewerIntefaces.cpp
--- a/Engine/Plugins/Experimental/EditorDataStorageFeatures/Source/TedsTableViewer/Private/HierarchyViewerIntefaces.cpp
+++ b/Engine/Plugins/Experimental/EditorDataStorageFeatures/Source/TedsTableViewer/Private/HierarchyViewerIntefaces.cpp
@@ -2,11 +2,23 @@
 
 #include "HierarchyViewerIntefaces.h"
 
+#include "DataStorage/Features.h"
+#include "Elements/Framework/TypedElementQueryBuilder.h"
+
 namespace UE::Editor::DataStorage
 {
-	FHierarchyViewerData::FHierarchyViewerData(FHierarchyHandle InHierarchyHandle)
+	FHierarchyViewerData::FHierarchyViewerData(ICoreProvider& Storage, FHierarchyHandle InHierarchyHandle)
 		: HierarchyHandle(InHierarchyHandle)
 	{
+		RegisterQueries(Storage);
+	}
+
+	FHierarchyViewerData::~FHierarchyViewerData()
+	{
+		if (ICoreProvider* Storage = GetMutableDataStorageFeature<ICoreProvider>(StorageFeatureName))
+		{
+			UnRegisterQueries(*Storage);
+		}
 	}
 
 	RowHandle FHierarchyViewerData::GetParent(const ICoreProvider& Storage, RowHandle InRow) const
@@ -14,22 +26,91 @@ namespace UE::Editor::DataStorage
 		return Storage.GetParentRow(HierarchyHandle, InRow);
 	}
 
-	FHierarchyViewerLegacyData::FHierarchyViewerLegacyData(const UScriptStruct* InHierarchyDataColumnType,
-	                                                       TFunction<RowHandle(const void*, const UScriptStruct*)> InGetRowHandleFn)
-		: HierarchyDataColumnType(InHierarchyDataColumnType)
-		, GetRowHandleFn(InGetRowHandleFn)
+	void FHierarchyViewerData::WalkDepthFirst(const ICoreProvider& Storage, RowHandle InRow, ICoreProvider::FHierarchyIterationCallback VisitFn) const
+	{
+		Storage.WalkDepthFirst(HierarchyHandle, InRow, VisitFn);
+	}
+
+	void FHierarchyViewerData::RegisterQueries(ICoreProvider& Storage)
+	{
+		using namespace UE::Editor::DataStorage::Queries;
+		
+		if (const UScriptStruct* ParentChangedColumn = Storage.GetParentChangedColumnType(HierarchyHandle))
+		{
+			ParentAddedQuery =
+				Storage.RegisterQuery(
+					Select(
+						TEXT("On parent added to row"),
+						FObserver(FObserver::EEvent::Add, ParentChangedColumn).SetExecutionMode(EExecutionMode::GameThread),
+						[this](IQueryContext& Context, RowHandle Row)
+						{
+								ParentChangedEvent.Broadcast(Row);
+						})
+					.Compile());
+
+			ParentRemovedQuery =
+				Storage.RegisterQuery(
+					Select(
+						TEXT("On parent removed from row"),
+						FObserver(FObserver::EEvent::Remove, ParentChangedColumn).SetExecutionMode(EExecutionMode::GameThread),
+						[this](IQueryContext& Context, RowHandle Row)
+						{
+								ParentChangedEvent.Broadcast(Row);
+						})
+					.Compile());
+		}
+	}
+
+	void FHierarchyViewerData::UnRegisterQueries(ICoreProvider& Storage)
 	{
-		checkf(HierarchyDataColumnType != nullptr, TEXT("Nullptr ParentColumnType"));
+		Storage.UnregisterQuery(ParentAddedQuery);
+		Storage.UnregisterQuery(ParentRemovedQuery);
+	}
+
+	FHierarchyViewerMultiData::FHierarchyViewerMultiData(ICoreProvider& Storage, TArray<FHierarchyHandle> OrderedHierarchyHandles)
+	{
+		for (FHierarchyHandle Handle : OrderedHierarchyHandles)
+		{  
+			TSharedRef<FHierarchyViewerData> Data = MakeShared<FHierarchyViewerData>(Storage, Handle);
+			Data->OnHierarchyChangedEvent().AddLambda([this](RowHandle Row)
+				{
+					OnHierarchyChangedEvent().Broadcast(Row);
+				});
+
+			Hierarchies.Add(MoveTemp(Data));
+		}
 	}
 
-	RowHandle FHierarchyViewerLegacyData::GetParent(const ICoreProvider& Storage, RowHandle InRow) const
+	RowHandle FHierarchyViewerMultiData::GetParent(const ICoreProvider& Storage, RowHandle InRow) const
 	{
-		if (const void* ParentColumn = Storage.GetColumnData(InRow, HierarchyDataColumnType))
+		for (const TSharedPtr<FHierarchyViewerData>& Data : Hierarchies)
 		{
-			RowHandle ParentRow = GetRowHandleFn(ParentColumn, HierarchyDataColumnType);
-			return ParentRow;
+			RowHandle Parent = Data->GetParent(Storage, InRow);
+
+			if (Storage.IsRowAvailable(Parent))
+			{
+				return Parent;
+			}
 		}
 
 		return InvalidRowHandle;
 	}
-}
\ No newline at end of file
+
+	void FHierarchyViewerMultiData::WalkDepthFirst(const ICoreProvider& Storage, RowHandle InRow,
+		ICoreProvider::FHierarchyIterationCallback VisitFn) const
+	{
+		TSet<RowHandle> Visited;
+		
+		for (const TSharedPtr<FHierarchyViewerData>& Data : Hierarchies)
+		{
+			Data->WalkDepthFirst(Storage, InRow, [&Visited, VisitFn](const ICoreProvider& Context, RowHandle Owner, RowHandle Target)
+			{
+				if (!Visited.Contains(Target))
+				{
+					VisitFn(Context, Owner, Target);
+					Visited.Add(Target);
+				}
+			});
+		}
+	}
+}
diff --git a/Engine/Plugins/Experimental/EditorDataStorageFeatures/Source/TedsTableViewer/Public/HierarchyViewerIntefaces.h b/Engine/Plugins/Experimental/EditorDataStorageFeatures/Source/TedsTableViewer/Public/HierarchyViewerIntefaces.h
--- a/Engine/Plugins/Experimental/EditorDataStorageFeatures/Source/TedsTableViewer/Public/HierarchyViewerIntefaces.h
+++ b/Engine/Plugins/Experimental/EditorDataStorageFeatures/Source/TedsTableViewer/Public/HierarchyViewerIntefaces.h
@@ -15,30 +15,53 @@ namespace UE::Editor::DataStorage
 	class IHierarchyViewerDataInterface
 	{
 	public:
+
+		// Delegate that is fired when the parent of a row on this hierarchy changes
+		DECLARE_MULTICAST_DELEGATE_OneParam(FOnParentChanged, RowHandle);
+		
 		UE_API virtual ~IHierarchyViewerDataInterface() = default;
 		UE_API virtual RowHandle GetParent(const ICoreProvider& Storage, RowHandle InRow) const = 0;
+		UE_API virtual void WalkDepthFirst(const ICoreProvider& Storage, RowHandle InRow, ICoreProvider::FHierarchyIterationCallback VisitFn) const = 0;
+		UE_API FOnParentChanged& OnHierarchyChangedEvent()
+		{
+			return ParentChangedEvent;
+		}
+
+	protected:
+		FOnParentChanged ParentChangedEvent;
 	};
 
 	// Data interface for the hierarchy viewer that operates on a single FHierarchyHandle 
 	class FHierarchyViewerData : public IHierarchyViewerDataInterface
 	{
 	public:
-		UE_API explicit FHierarchyViewerData(FHierarchyHandle InHierarchyHandle);
+		UE_API explicit FHierarchyViewerData(ICoreProvider& Storage, FHierarchyHandle InHierarchyHandle);
+		UE_API virtual ~FHierarchyViewerData() override;
 		UE_API virtual RowHandle GetParent(const ICoreProvider& Storage, RowHandle InRow) const override;
-	private:
+		UE_API virtual void WalkDepthFirst(const ICoreProvider& Storage, RowHandle InRow, ICoreProvider::FHierarchyIterationCallback VisitFn) const override;
+
+	protected:
+		void RegisterQueries(ICoreProvider& Storage);
+		void UnRegisterQueries(ICoreProvider& Storage);
+		
+	protected:
 		FHierarchyHandle HierarchyHandle;
+		QueryHandle ParentAddedQuery = InvalidQueryHandle;
+		QueryHandle ParentRemovedQuery = InvalidQueryHandle;
 	};
 
-	// Legacy data interface for the hierarchy viewer that supports arbitrary hierarchy column types alongside a function to extract the row handle from the column data
-	class FHierarchyViewerLegacyData : public IHierarchyViewerDataInterface
+	// Accepts an ordered list of hierarchies to combine into one singular view.
+	class FHierarchyViewerMultiData : public IHierarchyViewerDataInterface
 	{
 	public:
-		UE_API FHierarchyViewerLegacyData(const UScriptStruct* InHierarchyDataColumnType, TFunction<RowHandle(const void*, const UScriptStruct*)> InGetRowHandleFn);
+		UE_API FHierarchyViewerMultiData(ICoreProvider& Storage, TArray<FHierarchyHandle> OrderedHierarchyHandles);
+		// Returns the first parent it finds among the list of hierarchies
 		UE_API virtual RowHandle GetParent(const ICoreProvider& Storage, RowHandle InRow) const override;
+		// Does a DFS of each hierarchy one-by-one. If a child is encountered multiple times in different hierarchies, it is only reported in the first instance
+		UE_API virtual void WalkDepthFirst(const ICoreProvider& Storage, RowHandle InRow, ICoreProvider::FHierarchyIterationCallback VisitFn) const override;
 		
 	protected:
-		const UScriptStruct* HierarchyDataColumnType;
-		TFunction<RowHandle(const void*, const UScriptStruct*)> GetRowHandleFn;
+		TArray<TSharedPtr<FHierarchyViewerData>> Hierarchies;
 	};
 }
 
diff --git a/Engine/Source/Runtime/TypedElementFramework/Public/Elements/Interfaces/TypedElementDataStorageInterface.h b/Engine/Source/Runtime/TypedElementFramework/Public/Elements/Interfaces/TypedElementDataStorageInterface.h
--- a/Engine/Source/Runtime/TypedElementFramework/Public/Elements/Interfaces/TypedElementDataStorageInterface.h
+++ b/Engine/Source/Runtime/TypedElementFramework/Public/Elements/Interfaces/TypedElementDataStorageInterface.h
@@ -393,6 +393,10 @@ public:
 	// and is used internally.  It is the caller's responsibility to provide at least Read access
 	// to this column type in query requirements in order to use the HierarchyAccessInterface on the row
 	virtual const UScriptStruct* GetHierarchyDataColumnType(FHierarchyHandle InHierarchyHandle) const = 0;
+
+	// Gets the column type that is added to a row for a frame when the parent changes (if specified in RegisterHierarchy)
+	virtual const UScriptStruct* GetParentChangedColumnType(FHierarchyHandle InHierarchyHandle) const = 0;
+
 	// List all the currently known hierarchies by name
 	virtual void ListHierarchyNames(TFunctionRef<void(const FName& HierarchyName)> Callback) const = 0;
 
@@ -426,12 +430,14 @@ public:
 	
 	virtual bool HasChildren(FHierarchyHandle InHierarchyHandle, RowHandle Row) const = 0;
 
+	using FHierarchyIterationCallback = TFunction<void(const ICoreProvider& Context, RowHandle Owner, RowHandle Target)>;
+	
 	// Iterates depth first from the passed in Row, calling the VisitFn on each row
 	// The Row passed in will have VisitFn called on it
 	virtual void WalkDepthFirst(
 		FHierarchyHandle InHierarchyHandle,
 		RowHandle Row,
-		TFunction<void(const ICoreProvider& Context, RowHandle Owner, RowHandle Target)> VisitFn) const = 0;
+		FHierarchyIterationCallback VisitFn) const = 0;
 	
 	/**
 	 * Outputs the registered query callbacks to the given output device for debugging purposes.
