TypeInfoWidget: Convert to use FSimpleWidgetConstructor and attribute bindings instead of the old classes and UICapabilities
TableViewer: Add a new widget hierarchy to replace uses of the old FTableRowParentColumn. Cell widgets in the table viewer are now parented to the table viewer widget row.

diff --git a/Engine/Plugins/Experimental/EditorDataStorageFeatures/Source/TedsContentBrowser/Private/Widgets/ContentBrowserTileViewWidget.cpp b/Engine/Plugins/Experimental/EditorDataStorageFeatures/Source/TedsContentBrowser/Private/Widgets/ContentBrowserTileViewWidget.cpp
--- a/Engine/Plugins/Experimental/EditorDataStorageFeatures/Source/TedsContentBrowser/Private/Widgets/ContentBrowserTileViewWidget.cpp
+++ b/Engine/Plugins/Experimental/EditorDataStorageFeatures/Source/TedsContentBrowser/Private/Widgets/ContentBrowserTileViewWidget.cpp
@@ -8,6 +8,7 @@
 #include "TedsAssetDataColumns.h"
 #include "TedsAssetDataHelper.h"
 #include "TedsAssetDataWidgetColumns.h"
+#include "TedsTableViewerUtils.h"
 #include "TedsTableViewerWidgetColumns.h"
 #include "Columns/SlateDelegateColumns.h"
 #include "Elements/Columns/TypedElementFolderColumns.h"
@@ -78,11 +79,10 @@ TSharedPtr<SWidget> FContentBrowserTileViewWidgetConstructor::CreateWidget(UE::E
 	{
 		FAttributeBinder WidgetRowBinder(WidgetRow, DataStorage);
 		FAttributeBinder RowBinder(TargetRow, DataStorage);
-		RowHandle ParentWidgetRowHandle = InvalidRowHandle;
-		if (FTableRowParentColumn* ParentWidgetRow = DataStorage->GetColumn<FTableRowParentColumn>(WidgetRow))
-		{
-			ParentWidgetRowHandle = ParentWidgetRow->Parent;
-		}
+
+		FHierarchyHandle WidgetHierarchy = DataStorage->FindHierarchyByName(TableViewerUtils::GetWidgetHierarchyName());
+		RowHandle ParentWidgetRowHandle = DataStorage->GetParentRow(WidgetHierarchy, WidgetRow);
+		
 		FAttributeBinder ParentWidgetRowBinder(ParentWidgetRowHandle, DataStorage);
 
 		constexpr float BorderPadding = 1.f;
diff --git a/Engine/Plugins/Experimental/EditorDataStorageFeatures/Source/TedsOutliner/Private/Widgets/OutlinerTextWidget.cpp b/Engine/Plugins/Experimental/EditorDataStorageFeatures/Source/TedsOutliner/Private/Widgets/OutlinerTextWidget.cpp
--- a/Engine/Plugins/Experimental/EditorDataStorageFeatures/Source/TedsOutliner/Private/Widgets/OutlinerTextWidget.cpp
+++ b/Engine/Plugins/Experimental/EditorDataStorageFeatures/Source/TedsOutliner/Private/Widgets/OutlinerTextWidget.cpp
@@ -4,6 +4,7 @@
 
 #include "ActorEditorUtils.h"
 #include "ITedsTableViewer.h"
+#include "TedsTableViewerUtils.h"
 #include "TedsTableViewerWidgetColumns.h"
 #include "TedsTableViewerWidgetColumns.h"
 #include "Columns/SlateDelegateColumns.h"
@@ -78,18 +79,18 @@ namespace UE::Editor::Outliner::TextWidget::Private
 			});
 
 			// The table viewer, if present, is the parent widget of our widget row
-			if (FTableRowParentColumn* ParentColumn = DataStorage->GetColumn<FTableRowParentColumn>(WidgetRow))
+
+			DataStorage::FHierarchyHandle WidgetHierarchy = DataStorage->FindHierarchyByName(DataStorage::TableViewerUtils::GetWidgetHierarchyName());
+			DataStorage::RowHandle ParentRow = DataStorage->GetParentRow(WidgetHierarchy, WidgetRow);
+
+			if (DataStorage->IsRowAvailable(ParentRow))
 			{
-				if (DataStorage->IsRowAvailable(ParentColumn->Parent))
+				// Bind to selection change on the table viewer
+				if (DataStorage::FTableViewerSelectionChangedColumn* SelectionChangedColumn =
+					DataStorage->GetColumn<DataStorage::FTableViewerSelectionChangedColumn>(ParentRow))
 				{
-					// Bind to selection change on the table viewer
-					if (DataStorage::FTableViewerSelectionChangedColumn* SelectionChangedColumn =
-						DataStorage->GetColumn<DataStorage::FTableViewerSelectionChangedColumn>(ParentColumn->Parent))
-					{
-						SelectionChangedColumn->OnSelectionChanged.AddSP(this, &SOutlinerEditableTextWidget::OnTableViewerSelectionChanged, ParentColumn->Parent);
-					}
+					SelectionChangedColumn->OnSelectionChanged.AddSP(this, &SOutlinerEditableTextWidget::OnTableViewerSelectionChanged, ParentRow);
 				}
-				
 			}
 		}
 
diff --git a/Engine/Plugins/Experimental/EditorDataStorageFeatures/Source/TedsOutliner/Private/Widgets/TypeInfoWidget.cpp b/Engine/Plugins/Experimental/EditorDataStorageFeatures/Source/TedsOutliner/Private/Widgets/TypeInfoWidget.cpp
--- a/Engine/Plugins/Experimental/EditorDataStorageFeatures/Source/TedsOutliner/Private/Widgets/TypeInfoWidget.cpp
+++ b/Engine/Plugins/Experimental/EditorDataStorageFeatures/Source/TedsOutliner/Private/Widgets/TypeInfoWidget.cpp
@@ -11,6 +11,7 @@
 #include "Elements/Interfaces/TypedElementDataStorageUiInterface.h"
 #include "Styling/SlateIconFinder.h"
 #include "TedsTableViewerUtils.h"
+#include "Elements/Framework/TypedElementAttributeBinding.h"
 #include "Widgets/Text/STextBlock.h"
 #include "Widgets/Images/SImage.h"
 
@@ -49,21 +50,20 @@ void UTypeInfoWidgetFactory::RegisterWidgetConstructors(UE::Editor::DataStorage:
 
 FTypeInfoWidgetConstructor::FTypeInfoWidgetConstructor()
 	: Super(FTypeInfoWidgetConstructor::StaticStruct())
-	, bUseIcon(false)
 {
 }
 
 FTypeInfoWidgetConstructor::FTypeInfoWidgetConstructor(const UScriptStruct* InTypeInfo)
 	: Super(InTypeInfo)
-	, bUseIcon(false)
 {
 	
 }
 
-TSharedPtr<SWidget> FTypeInfoWidgetConstructor::CreateWidget(
-	const UE::Editor::DataStorage::FMetaDataView& Arguments	)
+TSharedPtr<SWidget> FTypeInfoWidgetConstructor::CreateWidget(UE::Editor::DataStorage::ICoreProvider* DataStorage,
+	UE::Editor::DataStorage::IUiProvider* DataStorageUi, UE::Editor::DataStorage::RowHandle TargetRow, UE::Editor::DataStorage::RowHandle WidgetRow,
+	const UE::Editor::DataStorage::FMetaDataView& Arguments)
 {
-	bUseIcon = false;
+	bool bUseIcon = false;
 	
 	// Check if the caller provided metadata to use an icon widget
 	UE::Editor::DataStorage::FMetaDataEntryView MetaDataEntryView = Arguments.FindGeneric("TypeInfoWidget_bUseIcon");
@@ -73,83 +73,28 @@ TSharedPtr<SWidget> FTypeInfoWidgetConstructor::CreateWidget(
 
 		bUseIcon = *MetaDataEntryView.TryGetExact<bool>();
 	}
-
+	
 	if(bUseIcon)
 	{
 		return SNew(SImage)
 					.DesiredSizeOverride(FVector2D(16.f, 16.f))
-					.ColorAndOpacity(FSlateColor::UseForeground());
+					.ColorAndOpacity(FSlateColor::UseForeground())
+					// Icon doesn't really change post creation so this doesn't need to be an attribute
+					.Image(UE::Editor::DataStorage::TableViewerUtils::GetIconForRow(DataStorage, TargetRow));
 	}
 	else
 	{
-		return SNew(SHorizontalBox);
+		return SNew(SHorizontalBox)
+				+SHorizontalBox::Slot()
+					.AutoWidth()
+					.VAlign(VAlign_Center)
+					.Padding(8, 0, 0, 0)
+					[
+						CreateTextWidget(DataStorage, TargetRow)
+					];
 	}
 }
 
-bool FTypeInfoWidgetConstructor::FinalizeWidget(UE::Editor::DataStorage::ICoreProvider* DataStorage,
-	UE::Editor::DataStorage::IUiProvider* DataStorageUi,
-	UE::Editor::DataStorage::RowHandle Row, const TSharedPtr<SWidget>& Widget)
-{
-	checkf(Widget, TEXT("Referenced widget is not valid. A constructed widget may not have been cleaned up. This can "
-		"also happen if this processor is running in the same phase as the processors responsible for cleaning up old "
-		"references."));
-
-	UE::Editor::DataStorage::RowHandle TargetRow = DataStorage->GetColumn<FTypedElementRowReferenceColumn>(Row)->Row;
-
-	if (const FTypedElementClassTypeInfoColumn* TypeInfoColumn = DataStorage->GetColumn<FTypedElementClassTypeInfoColumn>(TargetRow))
-	{
-		if(bUseIcon)
-		{
-			checkf(Widget->GetType() == SImage::StaticWidgetClass().GetWidgetType(),
-				TEXT("Stored widget with FTypedElementLabelWidgetConstructor doesn't match type %s, but was a %s."),
-				*(SImage::StaticWidgetClass().GetWidgetType().ToString()),
-				*(Widget->GetTypeAsString()));
-
-			SImage* WidgetInstance = static_cast<SImage*>(Widget.Get());
-			
-			WidgetInstance->SetImage(UE::Editor::DataStorage::TableViewerUtils::GetIconForRow(DataStorage, Row));
-		}
-		else
-		{
-			checkf(Widget->GetType() == SHorizontalBox::StaticWidgetClass().GetWidgetType(),
-				TEXT("Stored widget with FTypedElementLabelWidgetConstructor doesn't match type %s, but was a %s."),
-				*(SHorizontalBox::StaticWidgetClass().GetWidgetType().ToString()),
-				*(Widget->GetTypeAsString()));
-				
-			SHorizontalBox* WidgetInstance = static_cast<SHorizontalBox*>(Widget.Get());
-
-			TSharedPtr<SWidget> ActualWidget;
-
-			// Check if we have a hyperlink for this object
-			if (const FTypedElementUObjectColumn* ObjectColumn = DataStorage->GetColumn<FTypedElementUObjectColumn>(TargetRow))
-			{
-				ActualWidget = SceneOutliner::FSceneOutlinerHelpers::GetClassHyperlink(ObjectColumn->Object.Get());
-			}
-
-			// If not, we simply show a text block with the type
-			if(!ActualWidget)
-			{
-				TSharedPtr<STextBlock> TextBlock = SNew(STextBlock)
-								.ColorAndOpacity(FSlateColor::UseSubduedForeground())
-								.Text(FText::FromString(TypeInfoColumn->TypeInfo.Get()->GetName()));
-
-				TextBlock->AddMetadata(MakeShared<TTypedElementUiTextCapability<STextBlock>>(*TextBlock));
-				ActualWidget = TextBlock;
-			}
-
-			WidgetInstance->AddSlot()
-						.AutoWidth()
-						.VAlign(VAlign_Center)
-						.Padding(8, 0, 0, 0)
-						[
-							ActualWidget.ToSharedRef()
-						];
-		}
-	}
-
-	return true;
-}
-
 TArray<TSharedPtr<const UE::Editor::DataStorage::FColumnSorterInterface>> FTypeInfoWidgetConstructor::ConstructColumnSorters(
 	UE::Editor::DataStorage::ICoreProvider* DataStorage,
 	UE::Editor::DataStorage::IUiProvider* DataStorageUi,
@@ -162,4 +107,33 @@ TArray<TSharedPtr<const UE::Editor::DataStorage::FColumnSorterInterface>> FTypeI
 		});
 }
 
-#undef LOCTEXT_NAMESPACE
\ No newline at end of file
+TSharedRef<SWidget> FTypeInfoWidgetConstructor::CreateTextWidget(UE::Editor::DataStorage::ICoreProvider* DataStorage,
+	UE::Editor::DataStorage::RowHandle TargetRow)
+{
+	// Check if we have a hyperlink for this object
+	if (const FTypedElementUObjectColumn* ObjectColumn = DataStorage->GetColumn<FTypedElementUObjectColumn>(TargetRow))
+	{
+		if (TSharedPtr<SWidget> HyperlinkWidget = SceneOutliner::FSceneOutlinerHelpers::GetClassHyperlink(ObjectColumn->Object.Get()))
+		{
+			return HyperlinkWidget.ToSharedRef();
+		}
+	}
+
+	UE::Editor::DataStorage::FAttributeBinder Binder(TargetRow, DataStorage);
+
+	// If not, we simply show a text block with the type
+	TSharedPtr<STextBlock> TextBlock = SNew(STextBlock)
+						.ColorAndOpacity(FSlateColor::UseSubduedForeground())
+						.Text(Binder.BindData(&FTypedElementClassTypeInfoColumn::TypeInfo, [](TWeakObjectPtr<const UClass> InTypeInfo)
+							{
+								if (TObjectPtr<const UClass> TypeInfoPtr = InTypeInfo.Get())
+								{
+									return FText::FromString(TypeInfoPtr->GetName());
+								}
+								return FText::GetEmpty();
+							}));
+	
+	return TextBlock.ToSharedRef();
+}
+
+#undef LOCTEXT_NAMESPACE
diff --git a/Engine/Plugins/Experimental/EditorDataStorageFeatures/Source/TedsOutliner/Public/Widgets/TypeInfoWidget.h b/Engine/Plugins/Experimental/EditorDataStorageFeatures/Source/TedsOutliner/Public/Widgets/TypeInfoWidget.h
--- a/Engine/Plugins/Experimental/EditorDataStorageFeatures/Source/TedsOutliner/Public/Widgets/TypeInfoWidget.h
+++ b/Engine/Plugins/Experimental/EditorDataStorageFeatures/Source/TedsOutliner/Public/Widgets/TypeInfoWidget.h
@@ -36,7 +36,7 @@ public:
 };
 
 USTRUCT()
-struct FTypeInfoWidgetConstructor : public FTypedElementWidgetConstructor
+struct FTypeInfoWidgetConstructor : public FSimpleWidgetConstructor
 {
 	GENERATED_BODY()
 
@@ -46,10 +46,14 @@ public:
 
 protected:
 	UE_API explicit FTypeInfoWidgetConstructor(const UScriptStruct* InTypeInfo);
-	UE_API TSharedPtr<SWidget> CreateWidget(const UE::Editor::DataStorage::FMetaDataView& Arguments) override;
-	UE_API bool FinalizeWidget(UE::Editor::DataStorage::ICoreProvider* DataStorage,
+	
+	UE_API virtual TSharedPtr<SWidget> CreateWidget(
+		UE::Editor::DataStorage::ICoreProvider* DataStorage,
 		UE::Editor::DataStorage::IUiProvider* DataStorageUi,
-		UE::Editor::DataStorage::RowHandle Row, const TSharedPtr<SWidget>& Widget) override;
+		UE::Editor::DataStorage::RowHandle TargetRow,
+		UE::Editor::DataStorage::RowHandle WidgetRow,
+		const UE::Editor::DataStorage::FMetaDataView& Arguments) override;
+	
 	UE_API TArray<TSharedPtr<const UE::Editor::DataStorage::FColumnSorterInterface>> ConstructColumnSorters(
 		UE::Editor::DataStorage::ICoreProvider* DataStorage,
 		UE::Editor::DataStorage::IUiProvider* DataStorageUi,
@@ -57,8 +61,7 @@ protected:
 
 protected:
 
-	// Whether the widget created by this constructor should be icon or text
-	bool bUseIcon;
+	static TSharedRef<SWidget> CreateTextWidget(UE::Editor::DataStorage::ICoreProvider* DataStorage, UE::Editor::DataStorage::RowHandle TargetRow);
 };
 
 #undef UE_API
diff --git a/Engine/Plugins/Experimental/EditorDataStorageFeatures/Source/TedsTableViewer/Private/TedsTableViewerUtils.cpp b/Engine/Plugins/Experimental/EditorDataStorageFeatures/Source/TedsTableViewer/Private/TedsTableViewerUtils.cpp
--- a/Engine/Plugins/Experimental/EditorDataStorageFeatures/Source/TedsTableViewer/Private/TedsTableViewerUtils.cpp
+++ b/Engine/Plugins/Experimental/EditorDataStorageFeatures/Source/TedsTableViewer/Private/TedsTableViewerUtils.cpp
@@ -16,11 +16,17 @@
 namespace UE::Editor::DataStorage::TableViewerUtils
 {
 	static FName TableViewerWidgetTableName("Editor_TableViewerWidgetTable");
+	static FName TableViewerWidgetHierarchyName("TableViewerWidgetHierarchy");
 
 	FName GetWidgetTableName()
 	{
 		return TableViewerWidgetTableName;
 	}
+
+	FName GetWidgetHierarchyName()
+	{
+		return TableViewerWidgetHierarchyName;
+	}
 	
 	// TEDS UI TODO: Maybe the widget can specify a user facing name derived from the matched columns instead of trying to find the longest matching name
 	FName FindLongestMatchingName(TConstArrayView<TWeakObjectPtr<const UScriptStruct>> ColumnTypes, int32 DefaultNameIndex)
@@ -173,7 +179,7 @@ namespace UE::Editor::DataStorage::TableViewerUtils
 	}
 }
 
-void UTypedElementTableViewerFactory::RegisterTables(UE::Editor::DataStorage::ICoreProvider& DataStorage)
+void UTedsTableViewerFactory::RegisterTables(UE::Editor::DataStorage::ICoreProvider& DataStorage)
 {
 	using namespace UE::Editor::DataStorage;
 	const TableHandle BaseWidgetTable = DataStorage.FindTable(FName(TEXT("Editor_WidgetTable")));
@@ -186,3 +192,15 @@ void UTypedElementTableViewerFactory::RegisterTables(UE::Editor::DataStorage::IC
 			}, TableViewerUtils::TableViewerWidgetTableName);
 	}
 }
+
+void UTedsTableViewerFactory::RegisterHierarchies(UE::Editor::DataStorage::ICoreProvider& DataStorage)
+{
+	using namespace UE::Editor::DataStorage;
+
+	const FHierarchyRegistrationParams Params
+	{
+		.Name = TableViewerUtils::GetWidgetHierarchyName()
+	};
+
+	DataStorage.RegisterHierarchy(Params);
+}
diff --git a/Engine/Plugins/Experimental/EditorDataStorageFeatures/Source/TedsTableViewer/Private/Widgets/STedsTableViewerRow.cpp b/Engine/Plugins/Experimental/EditorDataStorageFeatures/Source/TedsTableViewer/Private/Widgets/STedsTableViewerRow.cpp
--- a/Engine/Plugins/Experimental/EditorDataStorageFeatures/Source/TedsTableViewer/Private/Widgets/STedsTableViewerRow.cpp
+++ b/Engine/Plugins/Experimental/EditorDataStorageFeatures/Source/TedsTableViewer/Private/Widgets/STedsTableViewerRow.cpp
@@ -11,6 +11,7 @@
 #include "TedsDropOperations.h"
 #include "TedsDropOperationsUtility.h"
 #include "TedsTableViewerColumn.h"
+#include "TedsTableViewerUtils.h"
 #include "Widgets/STedsTableViewer.h"
 
 #define LOCTEXT_NAMESPACE "TedsTableViewerRow"
@@ -80,10 +81,11 @@ namespace UE::Editor::DataStorage
 		{
 			auto WidgetRowSetupDelegate = [this] (ICoreProvider& InStorage, const RowHandle& InUIRowHandle)
 			{
-				InStorage.AddColumn(InUIRowHandle, FTableRowParentColumn{ .Parent = ParentWidgetRowHandle });
 				// Currently only this column is needed (not FExternalWidgetExclusiveSelectionColumn) since the double-click action will automatically
 				// deselect all and select only the desired row. 
 				InStorage.AddColumn(InUIRowHandle, FExternalWidgetSelectionColumn { .IsSelected = FIsSelected::CreateSP(this, &STedsTableViewerRow::IsSelected)});
+				const FHierarchyHandle WidgetHierarchy = InStorage.FindHierarchyByName(TableViewerUtils::GetWidgetHierarchyName());
+				InStorage.SetParentRow(WidgetHierarchy, InUIRowHandle, ParentWidgetRowHandle);
 			};
 
 			if (TSharedPtr<SWidget> RowWidget = Column->ConstructRowWidget(Item, WidgetRowSetupDelegate))
diff --git a/Engine/Plugins/Experimental/EditorDataStorageFeatures/Source/TedsTableViewer/Private/Widgets/STedsTileViewer.cpp b/Engine/Plugins/Experimental/EditorDataStorageFeatures/Source/TedsTableViewer/Private/Widgets/STedsTileViewer.cpp
--- a/Engine/Plugins/Experimental/EditorDataStorageFeatures/Source/TedsTableViewer/Private/Widgets/STedsTileViewer.cpp
+++ b/Engine/Plugins/Experimental/EditorDataStorageFeatures/Source/TedsTableViewer/Private/Widgets/STedsTileViewer.cpp
@@ -10,6 +10,7 @@
 #include "Elements/Framework/TypedElementAttributeBinding.h"
 #include "Elements/Framework/TypedElementDataStorageWidget.h"
 #include "TedsTableViewerColumn.h"
+#include "TedsTableViewerUtils.h"
 #include "TedsTableViewerWidgetColumns.h"
 #include "Elements/Columns/TypedElementHiearchyColumns.h"
 #include "Styling/StyleColors.h"
@@ -29,8 +30,11 @@ namespace UE::Editor::DataStorage
 		// TODO: Have a better "default" widget for tiles instead of the default cell widget purpose
 		if (!WidgetPurpose.IsSet())
 		{
-			IUiProvider* StorageUi = GetMutableDataStorageFeature<IUiProvider>(UiFeatureName);
-			WidgetPurpose = StorageUi->GetGeneralWidgetPurposeID();
+			if (IUiProvider* StorageUi = GetMutableDataStorageFeature<IUiProvider>(UiFeatureName))
+			{
+				WidgetPurpose = StorageUi->GetGeneralWidgetPurposeID();
+			}
+			
 		}
 		
 		// While we might have columns for the widget constructor, we don't want to display any UI columns since that doesn't make sense
@@ -99,20 +103,21 @@ namespace UE::Editor::DataStorage
 			return false;
 		};
 	
-		IUiProvider* StorageUi = GetMutableDataStorageFeature<IUiProvider>(UiFeatureName);
-		
-		RowHandle WidgetPurposeRow = StorageUi->FindPurpose(WidgetPurpose);
-
-		if (Columns.IsEmpty())
-		{
-			StorageUi->CreateWidgetConstructors(WidgetPurposeRow, FMetaDataView(), AssignWidgetToColumn);
-		}
-		else
+		if (IUiProvider* StorageUi = GetMutableDataStorageFeature<IUiProvider>(UiFeatureName))
 		{
-			StorageUi->CreateWidgetConstructors(WidgetPurposeRow, IUiProvider::EMatchApproach::ExactMatch, Columns,
-				FMetaDataView(), AssignWidgetToColumn);
-		}
+			RowHandle WidgetPurposeRow = StorageUi->FindPurpose(WidgetPurpose);
 
+			if (Columns.IsEmpty())
+			{
+				StorageUi->CreateWidgetConstructors(WidgetPurposeRow, FMetaDataView(), AssignWidgetToColumn);
+			}
+			else
+			{
+				StorageUi->CreateWidgetConstructors(WidgetPurposeRow, IUiProvider::EMatchApproach::ExactMatch, Columns,
+					FMetaDataView(), AssignWidgetToColumn);
+			}
+		}
+		
 		if (Column)
 		{
 			Model->AddCustomRowWidget(Column.ToSharedRef());
@@ -197,7 +202,8 @@ namespace UE::Editor::DataStorage
 			RowWidget = Column->ConstructRowWidget(InItem,
 				[&](ICoreProvider& DataStorage, const RowHandle& WidgetRow)
 				{
-					DataStorage.AddColumn(WidgetRow, FTableRowParentColumn{ .Parent = GetWidgetRowHandle() });
+					FHierarchyHandle WidgetHierarchy = DataStorage.FindHierarchyByName(TableViewerUtils::GetWidgetHierarchyName());
+					DataStorage.SetParentRow(WidgetHierarchy, WidgetRow, GetWidgetRowHandle());
 					DataStorage.AddColumn(WidgetRow, FExternalWidgetSelectionColumn{ .IsSelected = FIsSelected::CreateSP(this, &STedsTileViewer::IsSelected, InItem.RowHandle) });
 					DataStorage.AddColumn(WidgetRow, FExternalWidgetExclusiveSelectionColumn{ .IsSelectedExclusively = FIsSelected::CreateSP(this, &STedsTileViewer::IsSelectedExclusively, InItem.RowHandle) });
 				});
diff --git a/Engine/Plugins/Experimental/EditorDataStorageFeatures/Source/TedsTableViewer/Public/TedsTableViewerUtils.h b/Engine/Plugins/Experimental/EditorDataStorageFeatures/Source/TedsTableViewer/Public/TedsTableViewerUtils.h
--- a/Engine/Plugins/Experimental/EditorDataStorageFeatures/Source/TedsTableViewer/Public/TedsTableViewerUtils.h
+++ b/Engine/Plugins/Experimental/EditorDataStorageFeatures/Source/TedsTableViewer/Public/TedsTableViewerUtils.h
@@ -23,6 +23,10 @@ struct FSlateBrush;
 namespace UE::Editor::DataStorage::TableViewerUtils
 {
 	TEDSTABLEVIEWER_API FName GetWidgetTableName();
+
+	// Get the name of the widget hierarchy for table viewer widgets
+	// This hierarchy currently contains the relationship between the row containing the table viewer itself and individual cell widgets in that table viewer
+	TEDSTABLEVIEWER_API FName GetWidgetHierarchyName();
 	
 	// Find the longest matching common column name given a list of columns
 	TEDSTABLEVIEWER_API FName FindLongestMatchingName(TConstArrayView<TWeakObjectPtr<const UScriptStruct>> ColumnTypes, int32 DefaultNameIndex);
@@ -39,13 +43,14 @@ namespace UE::Editor::DataStorage::TableViewerUtils
 } // namespace UE::Editor::DataStorage::TableViewerUtils
 
 UCLASS()
-class UTypedElementTableViewerFactory : public UEditorDataStorageFactory
+class UTedsTableViewerFactory : public UEditorDataStorageFactory
 {
 	GENERATED_BODY()
 
 public:
-	~UTypedElementTableViewerFactory() override = default;
+	~UTedsTableViewerFactory() override = default;
 
 	void RegisterTables(UE::Editor::DataStorage::ICoreProvider& DataStorage) override;
+	void RegisterHierarchies(UE::Editor::DataStorage::ICoreProvider& DataStorage) override;
 };
 
diff --git a/Engine/Plugins/Experimental/EditorDataStorageFeatures/Source/TedsTableViewer/Public/TedsTableViewerWidgetColumns.h b/Engine/Plugins/Experimental/EditorDataStorageFeatures/Source/TedsTableViewer/Public/TedsTableViewerWidgetColumns.h
--- a/Engine/Plugins/Experimental/EditorDataStorageFeatures/Source/TedsTableViewer/Public/TedsTableViewerWidgetColumns.h
+++ b/Engine/Plugins/Experimental/EditorDataStorageFeatures/Source/TedsTableViewer/Public/TedsTableViewerWidgetColumns.h
@@ -52,7 +52,7 @@ struct FTedsTableViewerSelectionChangedColumn : public FEditorDataStorageColumn
  * Column on the table viewer row with the delegate that executes to rename the selected rows
  * TEDS UI TODO: Can we use a generic column for actions/events instead of one hardcoded for rename?
  */
-USTRUCT(meta = (DisplayName = "On Table Viewer selection changed"))
+USTRUCT(meta = (DisplayName = "Table Viewer rename event"))
 struct FTedsTableViewerRenameSelectionColumn : public FEditorDataStorageColumn
 {
 	GENERATED_BODY()
