Adding IDO columns (HasUnknownPropertiesTag) to TEDS type rows

diff --git a/Engine/Plugins/Experimental/EditorDataStorageFeatures/Source/TedsTypeInfo/Private/TedsTypeInfoFactory.cpp b/Engine/Plugins/Experimental/EditorDataStorageFeatures/Source/TedsTypeInfo/Private/TedsTypeInfoFactory.cpp
--- a/Engine/Plugins/Experimental/EditorDataStorageFeatures/Source/TedsTypeInfo/Private/TedsTypeInfoFactory.cpp
+++ b/Engine/Plugins/Experimental/EditorDataStorageFeatures/Source/TedsTypeInfo/Private/TedsTypeInfoFactory.cpp
@@ -16,6 +16,7 @@
 #include "Elements/Framework/TypedElementQueryBuilder.h"
 
 #include "HAL/IConsoleManager.h"
+#include "UObject/PropertyBagRepository.h"
 
 #include "UObject/UObjectIterator.h"
 #include "VerseVM/VVMVerseClass.h"
@@ -165,6 +166,8 @@ bool UTypeInfoFactory::TryAddTypeInfoRow(const UStruct* TypeInfo)
 			if (const UVerseClass* VerseClassInfo = Cast<UVerseClass>(ClassInfo))
 			{
 				AddVerseColumns(*DataStorage, TypeRowHandle, VerseClassInfo);
+				
+				AddIdoColumns(*DataStorage, TypeRowHandle, VerseClassInfo);
 			}
 
 			DataStorage->MapRow(TypeInfo::TypeMappingDomain, FMapKey(ClassInfo), TypeRowHandle);
@@ -265,6 +268,19 @@ void UTypeInfoFactory::AddVerseColumns(UE::Editor::DataStorage::ICoreProvider& D
 	}
 }
 
+void UTypeInfoFactory::AddIdoColumns(UE::Editor::DataStorage::ICoreProvider& DataStorage, UE::Editor::DataStorage::RowHandle TypeRowHandle, const UVerseClass* InVerseTypeInfo)
+{
+	using namespace UE::Editor::DataStorage;
+
+	checkf(InVerseTypeInfo, TEXT("Attempted to add type info columns with null VerseTypeInfo"));
+
+	// Try find an InstanceDataObject from type
+	if (UE::FPropertyBagRepository::Get().FindFirstInstanceThatNeedsFixup(InVerseTypeInfo) != nullptr)
+	{
+		DataStorage.AddColumn<FHasUnknownPropertiesTag>(TypeRowHandle);
+	}
+}
+
 template <class TypeToClear>
 void UTypeInfoFactory::ClearTypeInfoByTag()
 {
@@ -274,4 +290,4 @@ void UTypeInfoFactory::ClearTypeInfoByTag()
 	checkf(DataStorage, TEXT("Attempted to clear type info from Teds with a null DataStorage"));
 
 	DataStorage->RemoveAllRowsWith<TypeToClear>();
-}
\ No newline at end of file
+}
diff --git a/Engine/Plugins/Experimental/EditorDataStorageFeatures/Source/TedsTypeInfo/Private/TedsTypeInfoFactory.h b/Engine/Plugins/Experimental/EditorDataStorageFeatures/Source/TedsTypeInfo/Private/TedsTypeInfoFactory.h
--- a/Engine/Plugins/Experimental/EditorDataStorageFeatures/Source/TedsTypeInfo/Private/TedsTypeInfoFactory.h
+++ b/Engine/Plugins/Experimental/EditorDataStorageFeatures/Source/TedsTypeInfo/Private/TedsTypeInfoFactory.h
@@ -72,6 +72,7 @@ private:
 	void AddStructColumns(UE::Editor::DataStorage::ICoreProvider& DataStorage, UE::Editor::DataStorage::RowHandle TypeRowHandle, const UScriptStruct* InStructInfo);
 	void AddClassColumns(UE::Editor::DataStorage::ICoreProvider& DataStorage, UE::Editor::DataStorage::RowHandle TypeRowHandle, const UClass* InClassInfo);
 	void AddVerseColumns(UE::Editor::DataStorage::ICoreProvider& DataStorage, UE::Editor::DataStorage::RowHandle TypeRowHandle, const UVerseClass* InVerseTypeInfo);
+	void AddIdoColumns(UE::Editor::DataStorage::ICoreProvider& DataStorage, UE::Editor::DataStorage::RowHandle TypeRowHandle, const UVerseClass* InVerseTypeInfo);
 
 	bool bRefreshTypeInfoQueued = false;
 
diff --git a/Engine/Plugins/Experimental/EditorDataStorageFeatures/Source/TedsTypeInfo/Public/Elements/Columns/TedsTypeInfoColumns.h b/Engine/Plugins/Experimental/EditorDataStorageFeatures/Source/TedsTypeInfo/Public/Elements/Columns/TedsTypeInfoColumns.h
--- a/Engine/Plugins/Experimental/EditorDataStorageFeatures/Source/TedsTypeInfo/Public/Elements/Columns/TedsTypeInfoColumns.h
+++ b/Engine/Plugins/Experimental/EditorDataStorageFeatures/Source/TedsTypeInfo/Public/Elements/Columns/TedsTypeInfoColumns.h
@@ -42,6 +42,12 @@ struct FDataStorageVerseTypeInfoAccessLevel : public FEditorDataStorageTag
 	GENERATED_BODY()
 };
 
+USTRUCT(meta = (DisplayName = "Row type has unknown properties"))
+struct FDataStorageHasUnknownPropertiesTag : public FEditorDataStorageTag
+{
+	GENERATED_BODY()
+};
+
 namespace UE::Editor::DataStorage
 {
 	using FTypeInfoTag = FDataStorageTypeInfoTag;
@@ -50,6 +56,7 @@ namespace UE::Editor::DataStorage
 	using FTypeInfoInterfaceTag = FDataStorageTypeInfoInterfaceTag;
 	using FVerseTypeInfoTag = FDataStorageVerseTypeInfoTag;
 	using FVerseTypeInfoAccessLevel = FDataStorageVerseTypeInfoAccessLevel;
+	using FHasUnknownPropertiesTag = FDataStorageHasUnknownPropertiesTag;
 }
 
 namespace UE::Editor::DataStorage::TypeInfo
diff --git a/Engine/Source/Runtime/CoreUObject/Internal/UObject/PropertyBagRepository.h b/Engine/Source/Runtime/CoreUObject/Internal/UObject/PropertyBagRepository.h
--- a/Engine/Source/Runtime/CoreUObject/Internal/UObject/PropertyBagRepository.h
+++ b/Engine/Source/Runtime/CoreUObject/Internal/UObject/PropertyBagRepository.h
@@ -7,6 +7,7 @@
 #include "Containers/Map.h"
 #include "Templates/FunctionFwd.h"
 #include "UObject/ObjectPtr.h"
+#include "UObject/UObjectHash.h"
 
 class UObject;
 class FMemoryArchive;
@@ -124,6 +125,8 @@ public:
 	COREUOBJECT_API void FindNestedInstanceDataObject(const UObject* Owner, bool bRequiresFixupOnly, TFunctionRef<void(UObject*)> Callback);
 	UE_INTERNAL void AddReferencedInstanceDataObject(const UObject* Owner, FReferenceCollector& Collector);
 
+	COREUOBJECT_API UObject* FindFirstInstanceThatNeedsFixup(const UClass* Class);
+
 	COREUOBJECT_API const UObject* FindInstanceForDataObject(const UObject* InstanceDataObject) const;
 
 	#if STATS
diff --git a/Engine/Source/Runtime/CoreUObject/Private/UObject/PropertyBagRepository.cpp b/Engine/Source/Runtime/CoreUObject/Private/UObject/PropertyBagRepository.cpp
--- a/Engine/Source/Runtime/CoreUObject/Private/UObject/PropertyBagRepository.cpp
+++ b/Engine/Source/Runtime/CoreUObject/Private/UObject/PropertyBagRepository.cpp
@@ -1083,6 +1083,27 @@ UStruct* FPropertyBagRepository::CreatePropertyBagPlaceholderType(UObject* Outer
 	return PlaceholderType;
 }
 
+UObject* FPropertyBagRepository::FindFirstInstanceThatNeedsFixup(const UClass* Class)
+{
+	UObject* FoundInstance = nullptr;
+
+	ForEachObjectOfClass(Class, [this, &FoundInstance](UObject* Instance)
+		{
+			// Quickly return if already found (No breakable version of ForEachObjectOfClass).
+			if (FoundInstance != nullptr)
+			{
+				return;
+			}
+
+			if (RequiresFixup(Instance))
+			{
+				FoundInstance = Instance;
+			}
+		}, /*bIncludeDerivedClasses =*/ false);
+
+	return FoundInstance;
+}
+
 #if STATS
 void FPropertyBagRepository::GatherStats(FPropertyBagRepositoryStats& Stats)
 {
@@ -1155,4 +1176,4 @@ void FPropertyBagRepository::DumpPropertyBagPlaceholderTypes(FOutputDevice& Outp
 
 } // UE
 
-#endif // WITH_EDITORONLY_DATA
+#endif // WITH_EDITORONLY_DATA
\ No newline at end of file
