SHierarchyViewer : Add a slate argument to set CanSelectGeneratedColumn on the SHeaderRow

The same argument was added to the STedsTableViewer in 45581071

diff --git a/Engine/Plugins/Experimental/EditorDataStorageFeatures/Source/TedsTableViewer/Private/Widgets/STedsHierarchyViewer.cpp b/Engine/Plugins/Experimental/EditorDataStorageFeatures/Source/TedsTableViewer/Private/Widgets/STedsHierarchyViewer.cpp
--- a/Engine/Plugins/Experimental/EditorDataStorageFeatures/Source/TedsTableViewer/Private/Widgets/STedsHierarchyViewer.cpp
+++ b/Engine/Plugins/Experimental/EditorDataStorageFeatures/Source/TedsTableViewer/Private/Widgets/STedsHierarchyViewer.cpp
@@ -59,7 +59,7 @@ namespace UE::Editor::DataStorage
 		QueryStackExecutor = QueryStack::FExplicitUpdateExecutor(FName("HierarchyViewer", InstanceCounter++), TopLevelRowsNode);
 		
 		HeaderRowWidget = SNew( SHeaderRow )
-							.CanSelectGeneratedColumn(true);
+							.CanSelectGeneratedColumn(InArgs._CanSelectGeneratedColumn);
 		
 		TedsWidget = StorageUi->CreateContainerTedsWidget(InvalidRowHandle);
 			
diff --git a/Engine/Plugins/Experimental/EditorDataStorageFeatures/Source/TedsTableViewer/Public/Widgets/STedsHierarchyViewer.h b/Engine/Plugins/Experimental/EditorDataStorageFeatures/Source/TedsTableViewer/Public/Widgets/STedsHierarchyViewer.h
--- a/Engine/Plugins/Experimental/EditorDataStorageFeatures/Source/TedsTableViewer/Public/Widgets/STedsHierarchyViewer.h
+++ b/Engine/Plugins/Experimental/EditorDataStorageFeatures/Source/TedsTableViewer/Public/Widgets/STedsHierarchyViewer.h
@@ -52,7 +52,8 @@ namespace UE::Editor::DataStorage
 		DECLARE_DELEGATE_OneParam(FOnSelectionChanged, RowHandle)
 
 		SLATE_BEGIN_ARGS(SHierarchyViewer)
-			: _ListSelectionMode(ESelectionMode::Type::Single)
+			: _CanSelectGeneratedColumn(true)
+			, _ListSelectionMode(ESelectionMode::Type::Single)
 		{
 			
 		}
@@ -72,6 +73,9 @@ namespace UE::Editor::DataStorage
 		// The widget purpose to use to create the header widgets
 		SLATE_ARGUMENT(IUiProvider::FPurposeID, HeaderWidgetPurpose)
 		
+		// Whether or not the user can select the columns generated by right clicking on the header menu.
+		SLATE_ARGUMENT(bool, CanSelectGeneratedColumn)
+
 		// Delegate called when the selection changes
 		SLATE_EVENT(FOnSelectionChanged, OnSelectionChanged)
 
