Hierarchy Viewer: Remove broken default expansion state argument and replace with a new explicit ExpandAll() call

diff --git a/Engine/Plugins/Experimental/EditorDataStorageFeatures/Source/TedsEverythingPicker/Private/Context/TedsPickerContextUtil.cpp b/Engine/Plugins/Experimental/EditorDataStorageFeatures/Source/TedsEverythingPicker/Private/Context/TedsPickerContextUtil.cpp
--- a/Engine/Plugins/Experimental/EditorDataStorageFeatures/Source/TedsEverythingPicker/Private/Context/TedsPickerContextUtil.cpp
+++ b/Engine/Plugins/Experimental/EditorDataStorageFeatures/Source/TedsEverythingPicker/Private/Context/TedsPickerContextUtil.cpp
@@ -89,9 +89,10 @@ namespace UE::Editor::DataStorage::Picker
 			.AllNodeProvider(FinalNode)
 			.Columns({ FTypedElementLabelColumn::StaticStruct(), FTypedElementClassTypeInfoColumn::StaticStruct() })
 			.CellWidgetPurpose(PurposeId)
-			.OnSelectionChanged(InArgs._OnSelectionChanged)
-			.DefaultExpansionState(SHierarchyViewer::EExpansionState::Expanded);
+			.OnSelectionChanged(InArgs._OnSelectionChanged);
 
+		ReferenceView->ExpandAll();
+		
 		ChildSlot
 			[
 				SNew(SVerticalBox)
diff --git a/Engine/Plugins/Experimental/EditorDataStorageFeatures/Source/TedsTableViewer/Private/Widgets/Composite/STedsCompositeHierarchyViewer.cpp b/Engine/Plugins/Experimental/EditorDataStorageFeatures/Source/TedsTableViewer/Private/Widgets/Composite/STedsCompositeHierarchyViewer.cpp
--- a/Engine/Plugins/Experimental/EditorDataStorageFeatures/Source/TedsTableViewer/Private/Widgets/Composite/STedsCompositeHierarchyViewer.cpp
+++ b/Engine/Plugins/Experimental/EditorDataStorageFeatures/Source/TedsTableViewer/Private/Widgets/Composite/STedsCompositeHierarchyViewer.cpp
@@ -111,7 +111,6 @@ namespace UE::Editor::DataStorage
 			.HeaderWidgetPurpose(HierarchyViewerArgs._HeaderWidgetPurpose)
 			.ColumnMetaData(HierarchyViewerArgs._ColumnMetaData)
 			.ListSelectionMode(HierarchyViewerArgs._ListSelectionMode)
-			.DefaultExpansionState(HierarchyViewerArgs._DefaultExpansionState)
 			.OnSelectionChanged(HierarchyViewerArgs._OnSelectionChanged);
 		
 		ContentWidget->AddSlot()
@@ -180,4 +179,13 @@ namespace UE::Editor::DataStorage
 		return HierarchyViewer->IsSelectedExclusively(InRow);
 	}
 
+	void STedsCompositeHierarchyViewer::ExpandAll() const
+	{
+		HierarchyViewer->ExpandAll();
+	}
+
+	void STedsCompositeHierarchyViewer::CollapseAll() const
+	{
+		HierarchyViewer->CollapseAll();
+	}
 } // namespace UE::Editor::DataStorage
diff --git a/Engine/Plugins/Experimental/EditorDataStorageFeatures/Source/TedsTableViewer/Private/Widgets/STedsHierarchyViewer.cpp b/Engine/Plugins/Experimental/EditorDataStorageFeatures/Source/TedsTableViewer/Private/Widgets/STedsHierarchyViewer.cpp
--- a/Engine/Plugins/Experimental/EditorDataStorageFeatures/Source/TedsTableViewer/Private/Widgets/STedsHierarchyViewer.cpp
+++ b/Engine/Plugins/Experimental/EditorDataStorageFeatures/Source/TedsTableViewer/Private/Widgets/STedsHierarchyViewer.cpp
@@ -27,7 +27,6 @@ namespace UE::Editor::DataStorage
 		HierarchyInterface = MoveTemp(InHierarchyInterface);
 		OnSelectionChanged = InArgs._OnSelectionChanged;
 		EmptyRowsMessage = InArgs._EmptyRowsMessage;
-		DefaultExpansionState = InArgs._DefaultExpansionState;
 		
 		IUiProvider::FPurposeID CellWidgetPurpose = InArgs._CellWidgetPurpose;
 
@@ -108,11 +107,29 @@ namespace UE::Editor::DataStorage
 				TopLevelRows.Add(TedsRowHandle);
 			}
 		}
+
+		if (bExpansionStateChangeRequested.IsSet())
+		{
+			if (TSharedPtr<QueryStack::IRowNode> AllRowsNode = Model->GetRowNode())
+			{
+				SetExpansionState_Internal(AllRowsNode->GetRows(), bExpansionStateChangeRequested.GetValue());
+			}
+			bExpansionStateChangeRequested.Reset();
+		}
 		
 		TreeView->RequestListRefresh();
 		CreateInternalWidget();
 	}
 
+	void SHierarchyViewer::SetExpansionState_Internal(FRowHandleArrayView Rows, bool bExpanded) const
+	{
+		for (RowHandle Row : Rows)
+		{
+			FTedsRowHandle TedsRowHandle{ .RowHandle = Row };
+			TreeView->SetItemExpansion(TedsRowHandle, bExpanded);
+		}
+	}
+
 	void SHierarchyViewer::AddWidgetColumns()
 	{
 		if(ICoreProvider* DataStorage = Model->GetDataStorageInterface())
@@ -271,6 +288,39 @@ namespace UE::Editor::DataStorage
 		return IsSelected(InRow) && TreeView->GetNumItemsSelected() == 1;
 	}
 
+	void SHierarchyViewer::ExpandAll() const
+	{
+		if (TSharedPtr<QueryStack::IRowNode> RowNode = Model->GetRowNode())
+		{
+			// If the query stack hasn't been initialized yet, defer the expansion until next refresh
+			if (RowNode->GetRevision() == QueryStack::INode::UninitializedRevisionId)
+			{
+				bExpansionStateChangeRequested = true;
+			}
+			else
+			{
+				SetExpansionState_Internal(RowNode->GetRows(), true);
+			}
+			
+		}
+	}
+
+	void SHierarchyViewer::CollapseAll() const
+	{
+		// If the query stack hasn't been initialized yet, defer the expansion until next refresh
+		if (TSharedPtr<QueryStack::IRowNode> RowNode = Model->GetRowNode())
+		{
+			if (RowNode->GetRevision() == QueryStack::INode::UninitializedRevisionId)
+			{
+				bExpansionStateChangeRequested = false;
+			}
+			else
+			{
+				SetExpansionState_Internal(RowNode->GetRows(), false);
+			}
+		}
+	}
+
 	bool SHierarchyViewer::IsItemVisible(TableViewerItemPtr InItem) const
 	{
 		return TreeView->IsItemVisible(InItem);
@@ -278,8 +328,6 @@ namespace UE::Editor::DataStorage
 
 	TSharedRef<ITableRow> SHierarchyViewer::MakeTableRowWidget(TableViewerItemPtr InItem, const TSharedRef<STableViewBase>& OwnerTable) const
 	{
-		TreeView->SetItemExpansion(InItem, DefaultExpansionState.Get() == EExpansionState::Expanded);
-
 		return SNew(SHierarchyViewerRow, OwnerTable, Model.ToSharedRef(), SharedThis(this))
 				.Padding(FMargin(0))
 				.ParentWidgetRowHandle(GetWidgetRowHandle())
diff --git a/Engine/Plugins/Experimental/EditorDataStorageFeatures/Source/TedsTableViewer/Public/Widgets/Composite/STedsCompositeHierarchyViewer.h b/Engine/Plugins/Experimental/EditorDataStorageFeatures/Source/TedsTableViewer/Public/Widgets/Composite/STedsCompositeHierarchyViewer.h
--- a/Engine/Plugins/Experimental/EditorDataStorageFeatures/Source/TedsTableViewer/Public/Widgets/Composite/STedsCompositeHierarchyViewer.h
+++ b/Engine/Plugins/Experimental/EditorDataStorageFeatures/Source/TedsTableViewer/Public/Widgets/Composite/STedsCompositeHierarchyViewer.h
@@ -96,6 +96,12 @@ namespace UE::Editor::DataStorage
 
 		TEDSTABLEVIEWER_API virtual bool IsSelectedExclusively(RowHandle InRow) const override;
 
+		// Recursively Expand all rows currently present in the viewer
+		TEDSTABLEVIEWER_API void ExpandAll() const;
+
+		// Recursively Collapse all rows currently present in the viewer
+		TEDSTABLEVIEWER_API void CollapseAll() const;
+
 	private:
 		// Row Node passed to the SearchBox that receives the searched result
 		TSharedPtr<QueryStack::IRowNode> SearchNode;
diff --git a/Engine/Plugins/Experimental/EditorDataStorageFeatures/Source/TedsTableViewer/Public/Widgets/STedsHierarchyViewer.h b/Engine/Plugins/Experimental/EditorDataStorageFeatures/Source/TedsTableViewer/Public/Widgets/STedsHierarchyViewer.h
--- a/Engine/Plugins/Experimental/EditorDataStorageFeatures/Source/TedsTableViewer/Public/Widgets/STedsHierarchyViewer.h
+++ b/Engine/Plugins/Experimental/EditorDataStorageFeatures/Source/TedsTableViewer/Public/Widgets/STedsHierarchyViewer.h
@@ -86,9 +86,6 @@ namespace UE::Editor::DataStorage
 		// Empty = simply show the column headers instead of a message
 		SLATE_ATTRIBUTE(FText, EmptyRowsMessage)
 
-		// Whether all rows start as collapsed or expanded by default
-		SLATE_ARGUMENT(EExpansionState, DefaultExpansionState)
-		
 		SLATE_END_ARGS()
 
 	public:
@@ -129,6 +126,12 @@ namespace UE::Editor::DataStorage
 		// Get if a specific row is selected exclusively
 		TEDSTABLEVIEWER_API virtual bool IsSelectedExclusively(RowHandle InRow) const override;
 
+		// Recursively Expand all rows currently present in the viewer
+		TEDSTABLEVIEWER_API void ExpandAll() const;
+
+		// Recursively Collapse all rows currently present in the viewer
+		TEDSTABLEVIEWER_API void CollapseAll() const;
+
 	protected:
 		
 		TSharedRef<ITableRow> MakeTableRowWidget(TableViewerItemPtr InItem, const TSharedRef<STableViewBase>& OwnerTable) const;
@@ -147,6 +150,8 @@ namespace UE::Editor::DataStorage
 
 		void OnModelChanged();
 
+		void SetExpansionState_Internal(FRowHandleArrayView Rows, bool bExpanded) const;
+
 	private:
 
 		// The actual ListView widget that displays the rows
@@ -167,8 +172,6 @@ namespace UE::Editor::DataStorage
 		// The message to show in place of the table viewer when there are no rows provided by the current query stack
 		TAttribute<FText> EmptyRowsMessage;
 
-		TAttribute<EExpansionState> DefaultExpansionState;
-
 		// A specialized query stack node to only contain the top level rows present in the actual query stack we are using.
 		TSharedPtr<QueryStack::IRowNode> TopLevelRowsNode;
 
@@ -181,5 +184,8 @@ namespace UE::Editor::DataStorage
 		// Hierarchy data that is used to extract information about the parent of a row
 		TSharedPtr<IHierarchyViewerDataInterface> HierarchyInterface;
 
+		// If we need to update the expansion state of all rows the next time the model changes
+		mutable TOptional<bool> bExpansionStateChangeRequested;
+
 	};
 }
\ No newline at end of file
