Fix for verse interfaces in the TEDS picker showing with a mangled name

diff --git a/Engine/Plugins/Experimental/EditorDataStorageFeatures/Source/TedsOutliner/Private/Widgets/TypeInfoWidget.cpp b/Engine/Plugins/Experimental/EditorDataStorageFeatures/Source/TedsOutliner/Private/Widgets/TypeInfoWidget.cpp
--- a/Engine/Plugins/Experimental/EditorDataStorageFeatures/Source/TedsOutliner/Private/Widgets/TypeInfoWidget.cpp
+++ b/Engine/Plugins/Experimental/EditorDataStorageFeatures/Source/TedsOutliner/Private/Widgets/TypeInfoWidget.cpp
@@ -128,7 +128,7 @@ TSharedRef<SWidget> FTypeInfoWidgetConstructor::CreateTextWidget(UE::Editor::Dat
 							{
 								if (TObjectPtr<const UClass> TypeInfoPtr = InTypeInfo.Get())
 								{
-									return FText::FromString(TypeInfoPtr->GetName());
+									return TypeInfoPtr->GetDisplayNameText();
 								}
 								return FText::GetEmpty();
 							}));
