[TEDS Hierarchy Viewer] Have the TopLevelNode reset on a query stack change + update the executor in the model on change

diff --git a/Engine/Plugins/Experimental/EditorDataStorageFeatures/Source/TedsTableViewer/Private/TedsTableViewerModel.cpp b/Engine/Plugins/Experimental/EditorDataStorageFeatures/Source/TedsTableViewer/Private/TedsTableViewerModel.cpp
--- a/Engine/Plugins/Experimental/EditorDataStorageFeatures/Source/TedsTableViewer/Private/TedsTableViewerModel.cpp
+++ b/Engine/Plugins/Experimental/EditorDataStorageFeatures/Source/TedsTableViewer/Private/TedsTableViewerModel.cpp
@@ -170,6 +170,7 @@ namespace UE::Editor::DataStorage
 		static int InstanceCounter = 0;
 		QueryStackExecutor = QueryStack::FExplicitUpdateExecutor(FName("TableViewerModel", InstanceCounter++), RowQueryStack);
 		
+		QueryStackExecutor.Update();
 		Refresh();
 	}
 
diff --git a/Engine/Plugins/Experimental/EditorDataStorageFeatures/Source/TedsTableViewer/Private/Widgets/STedsHierarchyViewer.cpp b/Engine/Plugins/Experimental/EditorDataStorageFeatures/Source/TedsTableViewer/Private/Widgets/STedsHierarchyViewer.cpp
--- a/Engine/Plugins/Experimental/EditorDataStorageFeatures/Source/TedsTableViewer/Private/Widgets/STedsHierarchyViewer.cpp
+++ b/Engine/Plugins/Experimental/EditorDataStorageFeatures/Source/TedsTableViewer/Private/Widgets/STedsHierarchyViewer.cpp
@@ -225,6 +225,7 @@ namespace UE::Editor::DataStorage
 	void SHierarchyViewer::SetQueryStack(TSharedPtr<QueryStack::IRowNode> InRowQueryStack)
 	{
 		Model->SetQueryStack(InRowQueryStack);
+		TopLevelRowsNode = MakeShared<QueryStack::FTopLevelRowsNode>(Model->GetDataStorageInterface(), HierarchyInterface, Model->GetRowNode());
 	}
 
 	void SHierarchyViewer::SetColumns(const TArray<TWeakObjectPtr<const UScriptStruct>>& InColumns)
