[TedsOperations] Adjusted some utility functions.

diff --git a/Engine/Plugins/Experimental/EditorDataStorageFeatures/Source/TedsDropOperations/Private/TedsDropOperationsFactory.cpp b/Engine/Plugins/Experimental/EditorDataStorageFeatures/Source/TedsDropOperations/Private/TedsDropOperationsFactory.cpp
--- a/Engine/Plugins/Experimental/EditorDataStorageFeatures/Source/TedsDropOperations/Private/TedsDropOperationsFactory.cpp
+++ b/Engine/Plugins/Experimental/EditorDataStorageFeatures/Source/TedsDropOperations/Private/TedsDropOperationsFactory.cpp
@@ -6,7 +6,6 @@
 #include "TedsDropOperationsColumns.h"
 #include "TedsOperationsColumns.h"
 #include "TedsQueryNode.h"
-#include "TedsRowQueryResultsNode.h"
 
 UDropOperationsFactory::UDropOperationsFactory()
 	: UOperationsFactory()
@@ -20,38 +19,38 @@ UDropOperationsFactory::~UDropOperationsFactory()
 
 void UDropOperationsFactory::RegisterTables(UE::Editor::DataStorage::ICoreProvider& Storage)
 {
-	static const FName TableName      = FName("DragAndDropOperationsTable");
-	static const FName InputTableName = FName("DragAndDropOperationsInputTable");
-	
+	static const FName OperationTableName                   = FName("DragAndDropOperationsTable");
+	static const FName InputTableName              = FName("DragAndDropOperationsInputTable");
+	static const FName InputTableNameNoDescription = FName("DragAndDropOperationsInputTableNoDescription");
+
 	using namespace UE::Editor::DataStorage;
-	
-	Table      = Storage.RegisterTable<Operations::FApplyColumn, DragAndDrop::FOperationTag>(TableName);
-	InputTable = Storage.RegisterTable<DragAndDrop::FTargetColumn, DragAndDrop::FSourceColumn>(InputTableName);
+
+	OperationTable = Storage.RegisterTable<Operations::FApplyColumn, DragAndDrop::FOperationTag>(OperationTableName);
+	InputTable = Storage.RegisterTable<DragAndDrop::FTargetColumn, DragAndDrop::FSourceColumn, DragAndDrop::FDescriptionColumn>(InputTableName);
+	InputTableNoDescription = Storage.RegisterTable<DragAndDrop::FTargetColumn, DragAndDrop::FSourceColumn>(InputTableNameNoDescription);
 }
 
-UE::Editor::DataStorage::FScopedRow UDropOperationsFactory::CreateInputRow(UE::Editor::DataStorage::RowHandle TargetRow,
+UE::Editor::DataStorage::RowHandle UDropOperationsFactory::CreateInputRow(UE::Editor::DataStorage::RowHandle TargetRow,
 	UE::Editor::DataStorage::RowHandle SourceRow, bool bAddDescription)
-{	
+{
 	using namespace UE::Editor::DataStorage;
 
-	if (!StoragePtr->IsRowAvailable(TargetRow) || !StoragePtr->IsRowAvailable(SourceRow))
-	{
-		return { StoragePtr, InvalidRowHandle };
-	}
-
-	if (InputTable == InvalidTableHandle)
+	TableHandle Table = bAddDescription ? InputTable : InputTableNoDescription;
+	if (Table == InvalidTableHandle)
 	{
-		return { StoragePtr, InvalidRowHandle };
+		return InvalidRowHandle;
 	}
 
-	RowHandle InputRow = StoragePtr->AddRow(InputTable);
+	RowHandle InputRow = StoragePtr->AddRow(Table);
 	StoragePtr->AddColumn(InputRow, DragAndDrop::FTargetColumn{ .Value = TargetRow });
 	StoragePtr->AddColumn(InputRow, DragAndDrop::FSourceColumn{ .Value = SourceRow });
-	if (bAddDescription)
-	{
-		StoragePtr->AddColumn<DragAndDrop::FDescriptionColumn>(InputRow);
-	}
-	return { StoragePtr, InputRow };
+	return InputRow;
+}
+
+UE::Editor::DataStorage::FScopedRow UDropOperationsFactory::CreateInputRowScoped(UE::Editor::DataStorage::RowHandle TargetRow,
+	UE::Editor::DataStorage::RowHandle SourceRow, bool bAddDescription)
+{
+	return { StoragePtr, CreateInputRow(TargetRow, SourceRow, bAddDescription) };
 }
 
 void UDropOperationsFactory::GetQueryDescription(UE::Editor::DataStorage::FQueryDescription& OutDescription) const
diff --git a/Engine/Plugins/Experimental/EditorDataStorageFeatures/Source/TedsDropOperations/Public/TedsDropOperationsColumnsHelpers.h b/Engine/Plugins/Experimental/EditorDataStorageFeatures/Source/TedsDropOperations/Public/TedsDropOperationsColumnsHelpers.h
--- a/Engine/Plugins/Experimental/EditorDataStorageFeatures/Source/TedsDropOperations/Public/TedsDropOperationsColumnsHelpers.h
+++ b/Engine/Plugins/Experimental/EditorDataStorageFeatures/Source/TedsDropOperations/Public/TedsDropOperationsColumnsHelpers.h
@@ -30,4 +30,11 @@ inline FText* GetDescription(ICoreProvider& Storage, RowHandle InputRow)
 	return Column ? &Column->Value : nullptr;
 }
 
+/** Returns a the description text as a value. */
+inline FText GetDescriptionValue(ICoreProvider& Storage, RowHandle InputRow)
+{
+	FText* Description = GetDescription(Storage, InputRow);
+	return Description ? *Description : FText();
+}
+
 }
diff --git a/Engine/Plugins/Experimental/EditorDataStorageFeatures/Source/TedsDropOperations/Public/TedsDropOperationsFactory.h b/Engine/Plugins/Experimental/EditorDataStorageFeatures/Source/TedsDropOperations/Public/TedsDropOperationsFactory.h
--- a/Engine/Plugins/Experimental/EditorDataStorageFeatures/Source/TedsDropOperations/Public/TedsDropOperationsFactory.h
+++ b/Engine/Plugins/Experimental/EditorDataStorageFeatures/Source/TedsDropOperations/Public/TedsDropOperationsFactory.h
@@ -25,14 +25,22 @@ public:
 	UE_API virtual void RegisterTables(UE::Editor::DataStorage::ICoreProvider& Storage) override;
 
 	/** Utility method to create a row with columns for the given input arguments. */
-	UE_API UE::Editor::DataStorage::FScopedRow CreateInputRow(UE::Editor::DataStorage::RowHandle TargetRow,
+	UE_API UE::Editor::DataStorage::RowHandle CreateInputRow(UE::Editor::DataStorage::RowHandle TargetRow,
 		UE::Editor::DataStorage::RowHandle SourceRow, bool bAddDescription = false);
 
+	/** Utility method to create a row with columns for the given input arguments. */
+	UE_API UE::Editor::DataStorage::FScopedRow CreateInputRowScoped(UE::Editor::DataStorage::RowHandle TargetRow,
+		UE::Editor::DataStorage::RowHandle SourceRow, bool bAddDescription = false);
+
+	UE::Editor::DataStorage::TableHandle GetInputTable() const              { return InputTable; }
+	UE::Editor::DataStorage::TableHandle GetInputTableNoDescription() const { return InputTableNoDescription; }
+
 protected:
 	UE_API virtual void GetQueryDescription(UE::Editor::DataStorage::FQueryDescription& OutDescription) const;
-	
+
 protected:
-	UE::Editor::DataStorage::TableHandle InputTable;	
+	UE::Editor::DataStorage::TableHandle InputTable;
+	UE::Editor::DataStorage::TableHandle InputTableNoDescription;
 };
 
 #undef UE_API
diff --git a/Engine/Plugins/Experimental/EditorDataStorageFeatures/Source/TedsOperations/Private/TedsOperationsFactory.cpp b/Engine/Plugins/Experimental/EditorDataStorageFeatures/Source/TedsOperations/Private/TedsOperationsFactory.cpp
--- a/Engine/Plugins/Experimental/EditorDataStorageFeatures/Source/TedsOperations/Private/TedsOperationsFactory.cpp
+++ b/Engine/Plugins/Experimental/EditorDataStorageFeatures/Source/TedsOperations/Private/TedsOperationsFactory.cpp
@@ -43,7 +43,7 @@ bool UOperationsFactory::FDefaultLess::operator()(const UE::Editor::DataStorage:
 
 UOperationsFactory::UOperationsFactory()
 	: StoragePtr(nullptr)
-	, Table(UE::Editor::DataStorage::InvalidTableHandle)
+	, OperationTable(UE::Editor::DataStorage::InvalidTableHandle)
 	, QueryNode()
 	, ResultsNode()
 	, Executor()
@@ -61,8 +61,8 @@ void UOperationsFactory::PreRegister(UE::Editor::DataStorage::ICoreProvider& Sto
 
 void UOperationsFactory::RegisterTables(UE::Editor::DataStorage::ICoreProvider& Storage)
 {
-	static const FName TableName("OperationsTable");
-	Table = Storage.RegisterTable<FOperationApplyColumn>(TableName);
+	static const FName OperationTableName("OperationsTable");
+	OperationTable = Storage.RegisterTable<FOperationApplyColumn>(OperationTableName);
 }
 
 void UOperationsFactory::RegisterQueries(UE::Editor::DataStorage::ICoreProvider& Storage)
@@ -86,12 +86,12 @@ UE::Editor::DataStorage::RowHandle UOperationsFactory::AddOperation(FName Name,
 {
 	using namespace UE::Editor::DataStorage;
 
-	if (Table == InvalidTableHandle)
+	if (OperationTable == InvalidTableHandle)
 	{
 		return InvalidRowHandle;
 	}
 
-	RowHandle Row = StoragePtr->AddRow(Table);
+	RowHandle Row = StoragePtr->AddRow(OperationTable);
 	if (Row != InvalidRowHandle)
 	{
 		StoragePtr->AddColumn(Row, Operations::FApplyColumn{ .Callback = MoveTemp(Operation) });
@@ -160,17 +160,22 @@ void UOperationsFactory::GetOperations(TArray<UE::Editor::DataStorage::RowHandle
 	}
 }
 
-bool UOperationsFactory::Test(UE::Editor::DataStorage::RowHandle InputRow, const FFilter& Filter, const FLess& Less) const
+UE::Editor::DataStorage::RowHandle UOperationsFactory::Test(UE::Editor::DataStorage::RowHandle InputRow, const FFilter& Filter, const FLess& Less) const
 {
 	using namespace UE::Editor::DataStorage;
 	
 	TArray<RowHandle> OperationRows;
 	GetOperations(OperationRows, InputRow, Filter, Less);
-	return Algo::AnyOf(OperationRows, [this, InputRow](RowHandle OperationRow)
+	for (RowHandle OperationRow : OperationRows)
 	{
 		const Operations::FTestColumn* TestColumn = StoragePtr->GetColumn<Operations::FTestColumn>(OperationRow);
-		return !TestColumn || !TestColumn->Callback || TestColumn->Callback(*StoragePtr, InputRow);	// No column equals True.
-	});
+		// No column equals validity.
+		if (!TestColumn || !TestColumn->Callback || TestColumn->Callback(*StoragePtr, InputRow))
+		{
+			return OperationRow;
+		}
+	}
+	return InvalidRowHandle;
 }
 
 TOptional<UE::Editor::DataStorage::Operations::FResult> UOperationsFactory::Apply(UE::Editor::DataStorage::RowHandle InputRow, const FFilter& Filter,
diff --git a/Engine/Plugins/Experimental/EditorDataStorageFeatures/Source/TedsOperations/Public/TedsOperationsFactory.h b/Engine/Plugins/Experimental/EditorDataStorageFeatures/Source/TedsOperations/Public/TedsOperationsFactory.h
--- a/Engine/Plugins/Experimental/EditorDataStorageFeatures/Source/TedsOperations/Public/TedsOperationsFactory.h
+++ b/Engine/Plugins/Experimental/EditorDataStorageFeatures/Source/TedsOperations/Public/TedsOperationsFactory.h
@@ -69,13 +69,15 @@ public:
 	 * @param InputRow The row that contains input data that will be passed to the operations.
 	 * @param Filter [optional] Predicate that can be used to customarily filter the operation rows. False = Skip, True = Keep.
 	 * @param Less [optional] Sorter that can be used to customarily sort the operation rows.
+	 * @returns The operation row that was successful for the given input.
 	 */
-	UE_API bool Test(UE::Editor::DataStorage::RowHandle InputRow, const FFilter& Filter = {}, const FLess& Less = FDefaultLess()) const;
+	UE_API UE::Editor::DataStorage::RowHandle Test(UE::Editor::DataStorage::RowHandle InputRow, const FFilter& Filter = {}, const FLess& Less = FDefaultLess()) const;
 
 	/** Invokes the Apply callback for all operation rows that accept (probe call) the given input row and returns on the first successful one.
 	 * @param InputRow The row that contains input data that will be passed to the operations.
 	 * @param Filter [optional] Predicate that can be used to customarily filter the operation rows. False = Skip, True = Keep.
 	 * @param Less [optional] Sorter that can be used to customarily sort the operation rows.
+	 * @returns Information about the invoked operation's changes if successful.
 	 */
 	UE_API TOptional<UE::Editor::DataStorage::Operations::FResult> Apply(UE::Editor::DataStorage::RowHandle InputRow,
 		const FFilter& Filter = {}, const FLess& Less = FDefaultLess()) const;
@@ -84,7 +86,7 @@ protected:
 	
 protected:
 	UE::Editor::DataStorage::ICoreProvider*                     StoragePtr;
-	UE::Editor::DataStorage::TableHandle                        Table;
+	UE::Editor::DataStorage::TableHandle                        OperationTable;
 	TSharedPtr<UE::Editor::DataStorage::QueryStack::IQueryNode> QueryNode;
 	TSharedPtr<UE::Editor::DataStorage::QueryStack::IRowNode>   ResultsNode;
 	TSharedPtr<UE::Editor::DataStorage::QueryStack::FExplicitUpdateExecutor> Executor;
diff --git a/Engine/Plugins/Experimental/EditorDataStorageFeatures/Source/TedsOperations/Public/TedsOperationsResult.h b/Engine/Plugins/Experimental/EditorDataStorageFeatures/Source/TedsOperations/Public/TedsOperationsResult.h
--- a/Engine/Plugins/Experimental/EditorDataStorageFeatures/Source/TedsOperations/Public/TedsOperationsResult.h
+++ b/Engine/Plugins/Experimental/EditorDataStorageFeatures/Source/TedsOperations/Public/TedsOperationsResult.h
@@ -3,24 +3,16 @@
 #pragma once
 
 #include "CoreMinimal.h"
-#include "DataStorage/Handles.h"
+#include "Elements/Framework/TypedElementRowHandleArray.h"
 
 namespace UE::Editor::DataStorage::Operations
 {
 
-enum class EResultType
+struct FResult
 {
-	Created,
-	Modified,
-	Removed,
+	FRowHandleArray Created;
+	FRowHandleArray Changed;
+	FRowHandleArray Removed;
 };
 
-struct FResultEntry
-{
-	RowHandle   Row;
-	EResultType Type;
-};
-
-using FResult = TArray<FResultEntry>;
-
 }
diff --git a/Engine/Plugins/Experimental/EditorDataStorageFeatures/Source/TedsOutliner/Private/TedsOutlinerMode.cpp b/Engine/Plugins/Experimental/EditorDataStorageFeatures/Source/TedsOutliner/Private/TedsOutlinerMode.cpp
--- a/Engine/Plugins/Experimental/EditorDataStorageFeatures/Source/TedsOutliner/Private/TedsOutlinerMode.cpp
+++ b/Engine/Plugins/Experimental/EditorDataStorageFeatures/Source/TedsOutliner/Private/TedsOutlinerMode.cpp
@@ -164,9 +164,9 @@ FSceneOutlinerDragValidationInfo FTedsOutlinerMode::ValidateDrop(const ISceneOut
 			
 			for (RowHandle SourceRow : DraggedRowHandles.GetRows())
 			{
-				FScopedRow InputRow = DropFactory->CreateInputRow(TedsDropTarget->GetRowHandle(), SourceRow, true);
+				FScopedRow InputRow = DropFactory->CreateInputRowScoped(TedsDropTarget->GetRowHandle(), SourceRow, true);
  
-				bool bResult = DropFactory->Test(InputRow.GetRow());
+				bool bResult = DropFactory->Test(InputRow.GetRow()) != InvalidRowHandle;
 
 				// Currently the description of the last operation will take precedence, we might want support for batched tests in the future
 				if (FText* DescriptionPtr = DragAndDrop::Helpers::GetDescription(*Storage, InputRow.GetRow()))
@@ -213,8 +213,8 @@ void FTedsOutlinerMode::OnDrop(ISceneOutlinerTreeItem& DropTarget, const FSceneO
 			{
 				for (RowHandle SourceRow : DraggedRowHandles.GetRows())
 				{
-					FScopedRow InputRow = DropFactory->CreateInputRow(TedsItem->GetRowHandle(), SourceRow, false);
-					DropFactory->Apply(InputRow.GetRow());
+					FScopedRow InputRow = DropFactory->CreateInputRowScoped(TedsItem->GetRowHandle(), SourceRow, false);
+					(void)DropFactory->Apply(InputRow.GetRow());
 				}
 			}
 		}
diff --git a/Engine/Plugins/Experimental/EditorDataStorageFeatures/Source/TedsTableViewer/Private/Widgets/STedsTableViewerRow.cpp b/Engine/Plugins/Experimental/EditorDataStorageFeatures/Source/TedsTableViewer/Private/Widgets/STedsTableViewerRow.cpp
--- a/Engine/Plugins/Experimental/EditorDataStorageFeatures/Source/TedsTableViewer/Private/Widgets/STedsTableViewerRow.cpp
+++ b/Engine/Plugins/Experimental/EditorDataStorageFeatures/Source/TedsTableViewer/Private/Widgets/STedsTableViewerRow.cpp
@@ -125,9 +125,9 @@ namespace UE::Editor::DataStorage
 					UE::Editor::DataStorage::DragAndDrop::GetRowsFromData(DraggedRows, *Operation);
 					for (RowHandle SourceRow : DraggedRows.GetRows())
 					{
-						FScopedRow InputRow = DropFactory->CreateInputRow(TargetRow, SourceRow, true);
+						FScopedRow InputRow = DropFactory->CreateInputRowScoped(TargetRow, SourceRow, true);
 
-						Result.bCanDrop = DropFactory->Test(InputRow.GetRow());
+						Result.bCanDrop = DropFactory->Test(InputRow.GetRow()) != InvalidRowHandle;
 
 						// Currently the description of the last operation will take precedence, we might want support for batched tests in the future
 						if (const FText* Description = DragAndDrop::Helpers::GetDescription(*Storage, InputRow.GetRow()))
@@ -192,7 +192,7 @@ namespace UE::Editor::DataStorage
 					UE::Editor::DataStorage::DragAndDrop::GetRowsFromData(DraggedRows, *Operation);
 					for (RowHandle SourceRow : DraggedRows.GetRows())
 						{
-							FScopedRow InputRow = DropFactory->CreateInputRow(TargetRow, SourceRow, false);
+							FScopedRow InputRow = DropFactory->CreateInputRowScoped(TargetRow, SourceRow, false);
 						(void)DropFactory->Apply(InputRow.GetRow());
 						}
 						return FReply::Handled();
