[TedsOperations] Proper initialization of input table vars.

diff --git a/Engine/Plugins/Experimental/EditorDataStorageFeatures/Source/TedsDropOperations/Private/TedsDropOperationsFactory.cpp b/Engine/Plugins/Experimental/EditorDataStorageFeatures/Source/TedsDropOperations/Private/TedsDropOperationsFactory.cpp
--- a/Engine/Plugins/Experimental/EditorDataStorageFeatures/Source/TedsDropOperations/Private/TedsDropOperationsFactory.cpp
+++ b/Engine/Plugins/Experimental/EditorDataStorageFeatures/Source/TedsDropOperations/Private/TedsDropOperationsFactory.cpp
@@ -9,7 +9,8 @@
 
 UDropOperationsFactory::UDropOperationsFactory()
 	: UOperationsFactory()
-	, InputTable()
+	, InputTable(UE::Editor::DataStorage::InvalidTableHandle)
+	, InputTableNoDescription(UE::Editor::DataStorage::InvalidTableHandle)
 {
 }
 
