Hierarchy Viewer:
Add a column on the widget row and cell widget rows that can be used to access the hierarchy interface the viewer is operating on

TEDS Visibility widget:
- FIx selection only operating on level editor selection set, it operates on the same set as the target row now
- Use hierarchy interfaces instead of hardcoded logic to use Outliner tree items, but keep the latter as a fallback for now

diff --git a/Engine/Plugins/Experimental/EditorDataStorageFeatures/Source/TedsOutliner/Private/Widgets/OutlinerVisibilityWidget.cpp b/Engine/Plugins/Experimental/EditorDataStorageFeatures/Source/TedsOutliner/Private/Widgets/OutlinerVisibilityWidget.cpp
--- a/Engine/Plugins/Experimental/EditorDataStorageFeatures/Source/TedsOutliner/Private/Widgets/OutlinerVisibilityWidget.cpp
+++ b/Engine/Plugins/Experimental/EditorDataStorageFeatures/Source/TedsOutliner/Private/Widgets/OutlinerVisibilityWidget.cpp
@@ -19,6 +19,7 @@
 #include "ActorTreeItem.h"
 #include "TedsOutlinerHelpers.h"
 #include "TedsOutlinerItem.h"
+#include "TedsTableViewerWidgetColumns.h"
 
 #include UE_INLINE_GENERATED_CPP_BY_NAME(OutlinerVisibilityWidget)
 
@@ -294,13 +295,26 @@ bool STedsVisibilityWidget::IsSelected() const
 	return false;
 }
 
-/** Set the item this widget is responsible for to be hidden or shown */
 void STedsVisibilityWidget::SetIsVisible(RowHandle InRow, const bool bVisible)
 {
-	TWeakPtr<ISceneOutlinerTreeItem> RowTreeItem = GetTreeItem(InRow);
-	if (UE::Editor::DataStorage::ICoreProvider* DataStorage = GetDataStorage(); RowTreeItem.IsValid())
+	using namespace UE::Editor::DataStorage;
+
+	if (ICoreProvider* DataStorage = GetDataStorage())
 	{
-		SetVisibility_Recursive(*DataStorage, RowTreeItem.Pin(), bVisible);
+		// Use the hierarchy data on this widget directly if present
+		if (FHierarchyViewerDataColumn* DataColumn = DataStorage->GetColumn<FHierarchyViewerDataColumn>(WidgetRow))
+		{
+			if (TSharedPtr<IHierarchyViewerDataInterface> DataInterface = DataColumn->HierarchyData.Pin())
+			{
+				SetVisibility_Recursive(*DataStorage, InRow, DataInterface, bVisible);
+			}
+		}
+		// The Teds Outliner doesn't use the new hierarchy data yet so we have a fallback to the tree item API
+		else if (TSharedPtr<ISceneOutlinerTreeItem> RowTreeItem = GetTreeItem(InRow))
+		{
+			SetVisibility_Recursive(*DataStorage, RowTreeItem, bVisible);
+		}
+		
 	}
 }
 
@@ -315,6 +329,23 @@ void STedsVisibilityWidget::CommitVisibility(UE::Editor::DataStorage::ICoreProvi
 	DataStorage.AddColumn<FTypedElementSyncBackToWorldTag>(Row);
 }
 
+void STedsVisibilityWidget::SetVisibility_Recursive(UE::Editor::DataStorage::ICoreProvider& DataStorage, RowHandle Row,
+	TSharedPtr<UE::Editor::DataStorage::IHierarchyViewerDataInterface> HierarchyData, bool bVisible)
+{
+	using namespace UE::Editor::DataStorage;
+
+	FRowHandleArray Children;
+	HierarchyData->WalkDepthFirst(DataStorage, Row, [&Children](const ICoreProvider& Context, RowHandle Owner, RowHandle Target)
+	{
+		Children.Add(Target);
+	});
+
+	for (RowHandle ChildRow : Children.GetRows())
+	{
+		CommitVisibility(DataStorage, ChildRow, bVisible);
+	}
+}
+
 void STedsVisibilityWidget::SetVisibility_Recursive(UE::Editor::DataStorage::ICoreProvider& DataStorage, FSceneOutlinerTreeItemPtr TreeItem, bool bVisible)
 {
 	using namespace UE::Editor::DataStorage;
@@ -385,27 +416,33 @@ UE::Editor::DataStorage::ICompatibilityProvider* STedsVisibilityWidget::GetDataS
 	return GetMutableDataStorageFeature<ICompatibilityProvider>(CompatibilityFeatureName);
 }
 
-void STedsVisibilityWidget::GetSelectedRows(TArray<RowHandle>& OutSelectedRows)
+void STedsVisibilityWidget::GetSelectedRows(TArray<RowHandle>& OutSelectedRows) const
 {
 	using namespace UE::Editor::DataStorage;
 	using namespace UE::Editor::DataStorage::Queries;
 	if (ICoreProvider* DataStorage = GetDataStorage())
 	{
-		static QueryHandle AllSelectedItemsQuery = DataStorage->RegisterQuery(Select().Where().All<FTypedElementSelectionColumn>().Compile());
-		DataStorage->RunQuery(AllSelectedItemsQuery,
-			CreateDirectQueryCallbackBinding([&OutSelectedRows, &DataStorage](const IDirectQueryContext& Context)
-			{
-				// Only add selections from the level editor - for now.
-				// UE-231184: Support Visibility Column in multiple editor contexts
-				for (RowHandle Row : Context.GetRowHandles())
+		if (const FTypedElementSelectionColumn* SelectionColumn = DataStorage->GetColumn<FTypedElementSelectionColumn>(TargetRow))
+		{
+			// We want all rows selected in the same set as the row this widget is displayed for
+			FName SelectionSet = SelectionColumn->SelectionSet;
+
+			static QueryHandle AllSelectedItemsQuery =
+				DataStorage->RegisterQuery(
+					Select()
+						.ReadOnly<FTypedElementSelectionColumn>()
+					.Where()
+					.Compile());
+			
+			DataStorage->RunQuery(AllSelectedItemsQuery,
+				CreateDirectQueryCallbackBinding([&OutSelectedRows, SelectionSet](const IDirectQueryContext& Context, RowHandle Row, const FTypedElementSelectionColumn& SelectionColumn)
 				{
-					FTypedElementSelectionColumn* Selection = DataStorage->GetColumn<FTypedElementSelectionColumn>(Row);
-					if (Selection && Selection->SelectionSet.IsNone()) // SelectionSet.IsNone() == Level Editor
+					if (SelectionColumn.SelectionSet == SelectionSet)
 					{
 						OutSelectedRows.Add(Row);
 					}
-				}
-			}));
+				}));
+		}
 	}
 }
 
diff --git a/Engine/Plugins/Experimental/EditorDataStorageFeatures/Source/TedsOutliner/Public/Widgets/OutlinerVisibilityWidget.h b/Engine/Plugins/Experimental/EditorDataStorageFeatures/Source/TedsOutliner/Public/Widgets/OutlinerVisibilityWidget.h
--- a/Engine/Plugins/Experimental/EditorDataStorageFeatures/Source/TedsOutliner/Public/Widgets/OutlinerVisibilityWidget.h
+++ b/Engine/Plugins/Experimental/EditorDataStorageFeatures/Source/TedsOutliner/Public/Widgets/OutlinerVisibilityWidget.h
@@ -16,6 +16,7 @@
 
 namespace UE::Editor::DataStorage
 {
+	class IHierarchyViewerDataInterface;
 	class ICoreProvider;
 }
 
@@ -153,7 +154,11 @@ private:
 	// Commit the visibility state into the Rows DataStorage. Will trigger a SyncBackToWorld manually but the object is never dirtied.
 	static UE_API void CommitVisibility(UE::Editor::DataStorage::ICoreProvider& DataStorage, RowHandle Row, bool bVisible);
 
-	// Recursive internal method to set visibility on a TreeItem and all of it's children.
+	// Recursive internal method to set visibility on a row and all its children.
+	static UE_API void SetVisibility_Recursive(UE::Editor::DataStorage::ICoreProvider& DataStorage, RowHandle Row,
+		TSharedPtr<UE::Editor::DataStorage::IHierarchyViewerDataInterface> HierarchyData, bool bVisible);
+
+	// Legacy method to work on Outliner tree items since they don't use the new hierarchy interface yet.
 	static UE_API void SetVisibility_Recursive(UE::Editor::DataStorage::ICoreProvider& DataStorage, FSceneOutlinerTreeItemPtr TreeItem, bool bVisible);
 
 	// Returns the tree item associated with this widget.
@@ -163,7 +168,7 @@ private:
 	static UE_API UE::Editor::DataStorage::ICoreProvider* GetDataStorageUI();
 	static UE_API UE::Editor::DataStorage::ICompatibilityProvider* GetDataStorageCompatibility();
 
-	static UE_API void GetSelectedRows(TArray<RowHandle>& OutSelectedRows);
+	UE_API void GetSelectedRows(TArray<RowHandle>& OutSelectedRows) const;
 };
 
 #undef UE_API
diff --git a/Engine/Plugins/Experimental/EditorDataStorageFeatures/Source/TedsTableViewer/Private/Widgets/STedsHierarchyViewer.cpp b/Engine/Plugins/Experimental/EditorDataStorageFeatures/Source/TedsTableViewer/Private/Widgets/STedsHierarchyViewer.cpp
--- a/Engine/Plugins/Experimental/EditorDataStorageFeatures/Source/TedsTableViewer/Private/Widgets/STedsHierarchyViewer.cpp
+++ b/Engine/Plugins/Experimental/EditorDataStorageFeatures/Source/TedsTableViewer/Private/Widgets/STedsHierarchyViewer.cpp
@@ -143,6 +143,7 @@ namespace UE::Editor::DataStorage
 				DataStorage->AddColumns<FHideRowFromUITag, FWidgetContextMenuColumn, FWidgetRowScrolledIntoView,
 					FWidgetDoubleClickedColumn, FTableViewerSelectionChangedColumn, FTableViewerRenameSelectionColumn>(WidgetRowHandle);
 				DataStorage->AddColumn(WidgetRowHandle, FTableViewerReferenceColumn{.TableViewerWeak = SharedThis(this)});
+				DataStorage->AddColumn(WidgetRowHandle, FHierarchyViewerDataColumn{.HierarchyData = HierarchyInterface});
 			}
 		}
 	}
diff --git a/Engine/Plugins/Experimental/EditorDataStorageFeatures/Source/TedsTableViewer/Private/Widgets/STedsTableViewerRow.cpp b/Engine/Plugins/Experimental/EditorDataStorageFeatures/Source/TedsTableViewer/Private/Widgets/STedsTableViewerRow.cpp
--- a/Engine/Plugins/Experimental/EditorDataStorageFeatures/Source/TedsTableViewer/Private/Widgets/STedsTableViewerRow.cpp
+++ b/Engine/Plugins/Experimental/EditorDataStorageFeatures/Source/TedsTableViewer/Private/Widgets/STedsTableViewerRow.cpp
@@ -12,6 +12,7 @@
 #include "TedsDropOperationsUtility.h"
 #include "TedsTableViewerColumn.h"
 #include "TedsTableViewerUtils.h"
+#include "TedsTableViewerWidgetColumns.h"
 #include "Widgets/STedsTableViewer.h"
 
 #define LOCTEXT_NAMESPACE "TedsTableViewerRow"
@@ -79,13 +80,10 @@ namespace UE::Editor::DataStorage
 	{
 		if(TSharedPtr<FTedsTableViewerColumn> Column = TableViewerModel->GetColumn(ColumnName))
 		{
+			// Lambda so we can access member function (capturing this is safe since this lambda is run immediately and discarded)
 			auto WidgetRowSetupDelegate = [this] (ICoreProvider& InStorage, const RowHandle& InUIRowHandle)
 			{
-				// Currently only this column is needed (not FExternalWidgetExclusiveSelectionColumn) since the double-click action will automatically
-				// deselect all and select only the desired row. 
-				InStorage.AddColumn(InUIRowHandle, FExternalWidgetSelectionColumn { .IsSelected = FIsSelected::CreateSP(this, &STedsTableViewerRow::IsSelected)});
-				const FHierarchyHandle WidgetHierarchy = InStorage.FindHierarchyByName(TableViewerUtils::GetWidgetHierarchyName());
-				InStorage.SetParentRow(WidgetHierarchy, InUIRowHandle, ParentWidgetRowHandle);
+				SetupWidgetColumns(InStorage, InUIRowHandle);
 			};
 
 			if (TSharedPtr<SWidget> RowWidget = Column->ConstructRowWidget(Item, WidgetRowSetupDelegate))
@@ -103,6 +101,12 @@ namespace UE::Editor::DataStorage
 		return SNullWidget::NullWidget;
 	}
 
+	void STedsTableViewerRow::SetupWidgetColumns(ICoreProvider& InStorage, const RowHandle& InUIRowHandle)
+	{
+		const FHierarchyHandle WidgetHierarchy = InStorage.FindHierarchyByName(TableViewerUtils::GetWidgetHierarchyName());
+		InStorage.SetParentRow(WidgetHierarchy, InUIRowHandle, ParentWidgetRowHandle);
+	}
+
 	FOptionalSize STedsTableViewerRow::GetCurrentItemHeight() const
 	{
 		return ItemHeight.IsSet() ? ItemHeight.Get() : FOptionalSize();
@@ -250,6 +254,17 @@ namespace UE::Editor::DataStorage
 
 		return ActualWidget;
 	}
+
+	void SHierarchyViewerRow::SetupWidgetColumns(ICoreProvider& InStorage, const RowHandle& InUIRowHandle)
+	{
+		// Forward hierarchy data from parent widget (hierarchy viewer) to cell widget for convinient access
+		if (const FHierarchyViewerDataColumn* Data = InStorage.GetColumn<FHierarchyViewerDataColumn>(ParentWidgetRowHandle))
+		{
+			InStorage.AddColumn(InUIRowHandle, FHierarchyViewerDataColumn{.HierarchyData = Data->HierarchyData });
+		}
+		
+		STedsTableViewerRow::SetupWidgetColumns(InStorage, InUIRowHandle);
+	}
 } // namespace UE::Editor::DataStorage
 
 #undef LOCTEXT_NAMESPACE
diff --git a/Engine/Plugins/Experimental/EditorDataStorageFeatures/Source/TedsTableViewer/Public/TedsTableViewerWidgetColumns.h b/Engine/Plugins/Experimental/EditorDataStorageFeatures/Source/TedsTableViewer/Public/TedsTableViewerWidgetColumns.h
--- a/Engine/Plugins/Experimental/EditorDataStorageFeatures/Source/TedsTableViewer/Public/TedsTableViewerWidgetColumns.h
+++ b/Engine/Plugins/Experimental/EditorDataStorageFeatures/Source/TedsTableViewer/Public/TedsTableViewerWidgetColumns.h
@@ -2,6 +2,7 @@
 
 #pragma once
 
+#include "HierarchyViewerIntefaces.h"
 #include "DataStorage/CommonTypes.h"
 #include "Framework/SlateDelegates.h"
 #include "ITedsTableViewer.h"
@@ -60,9 +61,25 @@ struct FTedsTableViewerRenameSelectionColumn : public FEditorDataStorageColumn
 	FSimpleMulticastDelegate OnRenameSelection;
 };
 
+/**
+ * Column on widget rows that contains a weak ptr to the hierarchy data used by the hierarchy viewer to display rows (if present).
+ * This is useful if your widget is meant to be displayed in a hierarchy view and needs access to data on the visual parent of the data row
+ * or do a ForEachChild.
+ * NOTE: This contains hierarchy data about the data/target row and not the widget row. Widget hierarchies can be accessed directly by using
+ * TableViewerUtils::GetWidgetHierarchyName() to get the widget row of the parent widget.
+ */
+USTRUCT(meta = (DisplayName = "Hierarchy viewer access interface"))
+struct FTedsHierarchyViewerDataColumn : public FEditorDataStorageColumn
+{
+	GENERATED_BODY()
+	
+	TWeakPtr<UE::Editor::DataStorage::IHierarchyViewerDataInterface> HierarchyData;
+};
+
 namespace UE::Editor::DataStorage
 {
 	using FTableViewerReferenceColumn = FTedsTableViewerReferenceColumn;
 	using FTableViewerSelectionChangedColumn = FTedsTableViewerSelectionChangedColumn;
 	using FTableViewerRenameSelectionColumn = FTedsTableViewerRenameSelectionColumn;
+	using FHierarchyViewerDataColumn = FTedsHierarchyViewerDataColumn;
 }
\ No newline at end of file
diff --git a/Engine/Plugins/Experimental/EditorDataStorageFeatures/Source/TedsTableViewer/Public/Widgets/STedsTableViewerRow.h b/Engine/Plugins/Experimental/EditorDataStorageFeatures/Source/TedsTableViewer/Public/Widgets/STedsTableViewerRow.h
--- a/Engine/Plugins/Experimental/EditorDataStorageFeatures/Source/TedsTableViewer/Public/Widgets/STedsTableViewerRow.h
+++ b/Engine/Plugins/Experimental/EditorDataStorageFeatures/Source/TedsTableViewer/Public/Widgets/STedsTableViewerRow.h
@@ -46,6 +46,11 @@ namespace UE::Editor::DataStorage
 		/** Overridden from SMultiColumnTableRow.  Generates a widget for this column of the tree row. */
 		virtual TSharedRef<SWidget> GenerateWidgetForColumn( const FName& ColumnName ) override;
 
+	protected:
+
+		/** Add some common columns to the widget row for each cell widget created */
+		virtual void SetupWidgetColumns(ICoreProvider& InStorage, const RowHandle& InUIRowHandle);
+
 	private:
 		/** Get the current ItemHeight as an OptionalSize for the Box HeightOverride */
 		FOptionalSize GetCurrentItemHeight() const;
@@ -72,5 +77,8 @@ namespace UE::Editor::DataStorage
 	{
 	public:
 		virtual TSharedRef<SWidget> GenerateWidgetForColumn( const FName& ColumnName ) override;
+
+	protected:
+		virtual void SetupWidgetColumns(ICoreProvider& InStorage, const RowHandle& InUIRowHandle) override;
 	};
 } // namespace UE::Editor::DataStorage
