TedsTableViewer: Bring back external widget selection column that was accidentally removed in a merge resolution

diff --git a/Engine/Plugins/Experimental/EditorDataStorageFeatures/Source/TedsTableViewer/Private/Widgets/STedsTableViewerRow.cpp b/Engine/Plugins/Experimental/EditorDataStorageFeatures/Source/TedsTableViewer/Private/Widgets/STedsTableViewerRow.cpp
--- a/Engine/Plugins/Experimental/EditorDataStorageFeatures/Source/TedsTableViewer/Private/Widgets/STedsTableViewerRow.cpp
+++ b/Engine/Plugins/Experimental/EditorDataStorageFeatures/Source/TedsTableViewer/Private/Widgets/STedsTableViewerRow.cpp
@@ -105,6 +105,10 @@ namespace UE::Editor::DataStorage
 	{
 		const FHierarchyHandle WidgetHierarchy = InStorage.FindHierarchyByName(TableViewerUtils::GetWidgetHierarchyName());
 		InStorage.SetParentRow(WidgetHierarchy, InUIRowHandle, ParentWidgetRowHandle);
+
+		// Currently only this column is needed (not FExternalWidgetExclusiveSelectionColumn) since the double-click action will automatically
+		// deselect all and select only the desired row. 
+		InStorage.AddColumn(InUIRowHandle, FExternalWidgetSelectionColumn { .IsSelected = FIsSelected::CreateSP(this, &STedsTableViewerRow::IsSelected)});
 	}
 
 	FOptionalSize STedsTableViewerRow::GetCurrentItemHeight() const
