Hierarchy Viewer: Add an argument to control the 'primary' column which is where the expander arrow appears.
This is internally supported for all viewers, but currently only implemented for the hierarchy viewer since it doesn't do anything for the other ones functionally. Support can be added later if there is a use case.

diff --git a/Engine/Plugins/Experimental/EditorDataStorageFeatures/Source/TedsTableViewer/Private/TedsTableViewerModel.cpp b/Engine/Plugins/Experimental/EditorDataStorageFeatures/Source/TedsTableViewer/Private/TedsTableViewerModel.cpp
--- a/Engine/Plugins/Experimental/EditorDataStorageFeatures/Source/TedsTableViewer/Private/TedsTableViewerModel.cpp
+++ b/Engine/Plugins/Experimental/EditorDataStorageFeatures/Source/TedsTableViewer/Private/TedsTableViewerModel.cpp
@@ -15,20 +15,18 @@ DEFINE_LOG_CATEGORY_STATIC(LogTableViewer, Log, All)
 
 namespace UE::Editor::DataStorage
 {
-	FTedsTableViewerModel::FTedsTableViewerModel(const TSharedPtr<QueryStack::IRowNode>& InRowQueryStack,
-		const TArray<TWeakObjectPtr<const UScriptStruct>>& InRequestedColumns, const TArray<TSharedPtr<FColumnMetaData>>& InColumnMetaData,
-		const IUiProvider::FPurposeID& InCellWidgetPurpose, const IUiProvider::FPurposeID& InHeaderWidgetPurpose,
-		const FIsItemVisible& InIsItemVisibleDelegate)
-		: RequestedTedsColumns(InRequestedColumns)
-		, ColumnMetaData(InColumnMetaData)
-		, CellWidgetPurpose(InCellWidgetPurpose)
-		, HeaderWidgetPurpose(InHeaderWidgetPurpose)
-		, IsItemVisible(InIsItemVisibleDelegate)
+	FTedsTableViewerModel::FTedsTableViewerModel(const FArguments& Args)
+		: RequestedTedsColumns(Args.RequestedColumns)
+		, ColumnMetaData(Args.ColumnMetaData)
+		, CellWidgetPurpose(Args.CellWidgetPurpose)
+		, HeaderWidgetPurpose(Args.HeaderWidgetPurpose)
+		, IsItemVisible(Args.IsItemVisibleDelegate)
+		, PrimaryColumn(Args.PrimaryColumn)
 	{
 		Storage = GetMutableDataStorageFeature<ICoreProvider>(StorageFeatureName);
 		StorageUi = GetMutableDataStorageFeature<IUiProvider>(UiFeatureName);
 		
-		SetQueryStack(InRowQueryStack);
+		SetQueryStack(Args.RowQueryStack);
 
 		TickerHandle = FTSTicker::GetCoreTicker().AddTicker(FTickerDelegate::CreateRaw(this, &FTedsTableViewerModel::Tick), 0);
 
@@ -70,6 +68,11 @@ namespace UE::Editor::DataStorage
 		return RowQueryStack;
 	}
 
+	FName FTedsTableViewerModel::GetPrimaryColumnName() const
+	{
+		return PrimaryUIColumnName;
+	}
+
 	void FTedsTableViewerModel::ValidateRequestedColumns()
 	{
 		RequestedTedsColumns.RemoveAll([](const TWeakObjectPtr<const UScriptStruct>& Column)
@@ -271,6 +274,16 @@ namespace UE::Editor::DataStorage
 				}
 			}
 		}
+
+		// Cache the primary UI column name since it can only change when we regenerate the columns
+		if (const TSharedRef<FTedsTableViewerColumn>* FoundPrimaryColumn = NewColumnMap.Find(PrimaryColumn))
+		{
+			PrimaryUIColumnName = (*FoundPrimaryColumn)->GetColumnName();
+		}
+		else
+		{
+			PrimaryUIColumnName = NAME_None;
+		}
 	}
 
 	bool FTedsTableViewerModel::IsRowVisible(RowHandle InRowHandle) const
diff --git a/Engine/Plugins/Experimental/EditorDataStorageFeatures/Source/TedsTableViewer/Private/Widgets/Composite/STedsCompositeHierarchyViewer.cpp b/Engine/Plugins/Experimental/EditorDataStorageFeatures/Source/TedsTableViewer/Private/Widgets/Composite/STedsCompositeHierarchyViewer.cpp
--- a/Engine/Plugins/Experimental/EditorDataStorageFeatures/Source/TedsTableViewer/Private/Widgets/Composite/STedsCompositeHierarchyViewer.cpp
+++ b/Engine/Plugins/Experimental/EditorDataStorageFeatures/Source/TedsTableViewer/Private/Widgets/Composite/STedsCompositeHierarchyViewer.cpp
@@ -111,7 +111,8 @@ namespace UE::Editor::DataStorage
 			.HeaderWidgetPurpose(HierarchyViewerArgs._HeaderWidgetPurpose)
 			.ColumnMetaData(HierarchyViewerArgs._ColumnMetaData)
 			.ListSelectionMode(HierarchyViewerArgs._ListSelectionMode)
-			.OnSelectionChanged(HierarchyViewerArgs._OnSelectionChanged);
+			.OnSelectionChanged(HierarchyViewerArgs._OnSelectionChanged)
+			.PrimaryColumn(HierarchyViewerArgs._PrimaryColumn);
 		
 		ContentWidget->AddSlot()
 		[
diff --git a/Engine/Plugins/Experimental/EditorDataStorageFeatures/Source/TedsTableViewer/Private/Widgets/STedsHierarchyViewer.cpp b/Engine/Plugins/Experimental/EditorDataStorageFeatures/Source/TedsTableViewer/Private/Widgets/STedsHierarchyViewer.cpp
--- a/Engine/Plugins/Experimental/EditorDataStorageFeatures/Source/TedsTableViewer/Private/Widgets/STedsHierarchyViewer.cpp
+++ b/Engine/Plugins/Experimental/EditorDataStorageFeatures/Source/TedsTableViewer/Private/Widgets/STedsHierarchyViewer.cpp
@@ -49,9 +49,16 @@ namespace UE::Editor::DataStorage
 			HeaderWidgetPurpose = IUiProvider::FPurposeInfo("General", "Header", NAME_None).GeneratePurposeID();
 		}
 
-		Model = MakeShared<FTedsTableViewerModel>(InArgs._AllNodeProvider, InArgs._Columns, InArgs._ColumnMetaData, CellWidgetPurpose, HeaderWidgetPurpose,
-			FTedsTableViewerModel::FIsItemVisible::CreateSP(this, &SHierarchyViewer::IsItemVisible));
-		
+		Model = MakeShared<FTedsTableViewerModel>(FTedsTableViewerModel::FArguments{
+			.RowQueryStack = InArgs._AllNodeProvider,
+			.RequestedColumns = InArgs._Columns,
+			.ColumnMetaData = InArgs._ColumnMetaData,
+			.CellWidgetPurpose = CellWidgetPurpose,
+			.HeaderWidgetPurpose = HeaderWidgetPurpose,
+			.IsItemVisibleDelegate = FTedsTableViewerModel::FIsItemVisible::CreateSP(this, &SHierarchyViewer::IsItemVisible),
+			.PrimaryColumn = InArgs._PrimaryColumn
+		});
+
 		TopLevelRowsNode = MakeShared<QueryStack::FTopLevelRowsNode>(Model->GetDataStorageInterface(), HierarchyInterface, Model->GetRowNode());
 		
 		static int InstanceCounter = 0;
diff --git a/Engine/Plugins/Experimental/EditorDataStorageFeatures/Source/TedsTableViewer/Private/Widgets/STedsTableViewer.cpp b/Engine/Plugins/Experimental/EditorDataStorageFeatures/Source/TedsTableViewer/Private/Widgets/STedsTableViewer.cpp
--- a/Engine/Plugins/Experimental/EditorDataStorageFeatures/Source/TedsTableViewer/Private/Widgets/STedsTableViewer.cpp
+++ b/Engine/Plugins/Experimental/EditorDataStorageFeatures/Source/TedsTableViewer/Private/Widgets/STedsTableViewer.cpp
@@ -39,8 +39,14 @@ namespace UE::Editor::DataStorage
 			HeaderWidgetPurpose = IUiProvider::FPurposeInfo("General", "Header", NAME_None).GeneratePurposeID();
 		}
 		
-		Model = MakeShared<FTedsTableViewerModel>(InArgs._QueryStack, InArgs._Columns, InArgs._ColumnMetaData,
-			CellWidgetPurpose, HeaderWidgetPurpose, FTedsTableViewerModel::FIsItemVisible::CreateSP(this, &STedsTableViewer::IsItemVisible));
+		Model = MakeShared<FTedsTableViewerModel>(FTedsTableViewerModel::FArguments{
+			.RowQueryStack = InArgs._QueryStack,
+			.RequestedColumns = InArgs._Columns,
+			.ColumnMetaData = InArgs._ColumnMetaData,
+			.CellWidgetPurpose = CellWidgetPurpose,
+			.HeaderWidgetPurpose = HeaderWidgetPurpose,
+			.IsItemVisibleDelegate = FTedsTableViewerModel::FIsItemVisible::CreateSP(this, &STedsTableViewer::IsItemVisible)
+		});
 		
 		HeaderRowWidget = SNew( SHeaderRow )
 							.CanSelectGeneratedColumn(InArgs._CanSelectGeneratedColumn);
diff --git a/Engine/Plugins/Experimental/EditorDataStorageFeatures/Source/TedsTableViewer/Private/Widgets/STedsTableViewerRow.cpp b/Engine/Plugins/Experimental/EditorDataStorageFeatures/Source/TedsTableViewer/Private/Widgets/STedsTableViewerRow.cpp
--- a/Engine/Plugins/Experimental/EditorDataStorageFeatures/Source/TedsTableViewer/Private/Widgets/STedsTableViewerRow.cpp
+++ b/Engine/Plugins/Experimental/EditorDataStorageFeatures/Source/TedsTableViewer/Private/Widgets/STedsTableViewerRow.cpp
@@ -235,8 +235,14 @@ namespace UE::Editor::DataStorage
 	{
 		TSharedRef<SWidget> ActualWidget = STedsTableViewerRow::GenerateWidgetForColumn(ColumnName);
 
-		// We show the expander arrow on the first column only
-		if (TableViewerModel->GetColumnIndex(ColumnName) == 0)
+		FName PrimaryColumName = TableViewerModel->GetPrimaryColumnName();
+		
+		bool bShowExpanderArrow =
+			PrimaryColumName.IsNone()
+			? TableViewerModel->GetColumnIndex(ColumnName) == 0 // If there was no primary column specified, the expander shows up on the leftmost column
+			: PrimaryColumName == ColumnName; // Otherwise it shows up on the primary column only
+			
+		if (bShowExpanderArrow)
 		{
 			return SNew(SBox)
 				.MinDesiredHeight(20.0f)
diff --git a/Engine/Plugins/Experimental/EditorDataStorageFeatures/Source/TedsTableViewer/Private/Widgets/STedsTileViewer.cpp b/Engine/Plugins/Experimental/EditorDataStorageFeatures/Source/TedsTableViewer/Private/Widgets/STedsTileViewer.cpp
--- a/Engine/Plugins/Experimental/EditorDataStorageFeatures/Source/TedsTableViewer/Private/Widgets/STedsTileViewer.cpp
+++ b/Engine/Plugins/Experimental/EditorDataStorageFeatures/Source/TedsTableViewer/Private/Widgets/STedsTileViewer.cpp
@@ -43,8 +43,14 @@ namespace UE::Editor::DataStorage
 		TArray<TWeakObjectPtr<const UScriptStruct>> DummyColumns;
 		TArray<TSharedPtr<FColumnMetaData>> DummyColumnMetaData;
 		
-		Model = MakeShared<FTedsTableViewerModel>(InArgs._QueryStack, DummyColumns, DummyColumnMetaData, DummyPurpose, DummyPurpose,
-			FTedsTableViewerModel::FIsItemVisible::CreateSP(this, &STedsTileViewer::IsItemVisible));
+		Model = MakeShared<FTedsTableViewerModel>(FTedsTableViewerModel::FArguments{
+			.RowQueryStack = InArgs._QueryStack,
+			.RequestedColumns = DummyColumns,
+			.ColumnMetaData = DummyColumnMetaData,
+			.CellWidgetPurpose = DummyPurpose,
+			.HeaderWidgetPurpose = DummyPurpose,
+			.IsItemVisibleDelegate = FTedsTableViewerModel::FIsItemVisible::CreateSP(this, &STedsTileViewer::IsItemVisible)
+		});
 
 		IUiProvider* StorageUi = Model->GetDataStorageUiProvider();
 		
diff --git a/Engine/Plugins/Experimental/EditorDataStorageFeatures/Source/TedsTableViewer/Public/TedsTableViewerModel.h b/Engine/Plugins/Experimental/EditorDataStorageFeatures/Source/TedsTableViewer/Public/TedsTableViewerModel.h
--- a/Engine/Plugins/Experimental/EditorDataStorageFeatures/Source/TedsTableViewer/Public/TedsTableViewerModel.h
+++ b/Engine/Plugins/Experimental/EditorDataStorageFeatures/Source/TedsTableViewer/Public/TedsTableViewerModel.h
@@ -38,13 +38,18 @@ namespace UE::Editor::DataStorage
 		DECLARE_DELEGATE_RetVal_OneParam(bool, FIsItemVisible, TableViewerItemPtr);
 		DECLARE_MULTICAST_DELEGATE(FOnModelChanged);
 
-		TEDSTABLEVIEWER_API FTedsTableViewerModel(
-			const TSharedPtr<QueryStack::IRowNode>& RowQueryStack,
-			const TArray<TWeakObjectPtr<const UScriptStruct>>& RequestedColumns,
-			const TArray<TSharedPtr<FColumnMetaData>>& ColumnMetaData,
-			const IUiProvider::FPurposeID& CellWidgetPurpose, 
-			const IUiProvider::FPurposeID& HeaderWidgetPurpose, 
-			const FIsItemVisible& InIsItemVisibleDelegate);
+		struct FArguments
+		{
+			TSharedPtr<QueryStack::IRowNode> RowQueryStack;
+			TArray<TWeakObjectPtr<const UScriptStruct>> RequestedColumns;
+			TArray<TSharedPtr<FColumnMetaData>> ColumnMetaData;
+			IUiProvider::FPurposeID CellWidgetPurpose; 
+			IUiProvider::FPurposeID HeaderWidgetPurpose;
+			FIsItemVisible IsItemVisibleDelegate;
+			TWeakObjectPtr<const UScriptStruct> PrimaryColumn;
+		};
+
+		TEDSTABLEVIEWER_API FTedsTableViewerModel(const FArguments& Args);
 
 		~FTedsTableViewerModel();
 
@@ -87,6 +92,9 @@ namespace UE::Editor::DataStorage
 
 		// Get the bottom-most row node used in the query stack for this model
 		TEDSTABLEVIEWER_API TSharedPtr<QueryStack::IRowNode> GetRowNode() const;
+
+		// The FName of the primary UI column (matches the column name in the underlying slate widgets)
+		TEDSTABLEVIEWER_API FName GetPrimaryColumnName() const;
 		
 	protected:
 
@@ -160,5 +168,11 @@ namespace UE::Editor::DataStorage
 		// Teds Constructs
 		ICoreProvider* Storage = nullptr;
 		IUiProvider* StorageUi = nullptr;
+
+		// The TEDS column that was requested as primary for this viewer
+		TWeakObjectPtr<const UScriptStruct> PrimaryColumn;
+
+		// TEDS -> UI columns is not a 1:1 mapping, so we cache the name of the UI column so widgets can query it easily
+		FName PrimaryUIColumnName;
 	};
 } // namespace UE::Editor::DataStorage
diff --git a/Engine/Plugins/Experimental/EditorDataStorageFeatures/Source/TedsTableViewer/Public/Widgets/STedsHierarchyViewer.h b/Engine/Plugins/Experimental/EditorDataStorageFeatures/Source/TedsTableViewer/Public/Widgets/STedsHierarchyViewer.h
--- a/Engine/Plugins/Experimental/EditorDataStorageFeatures/Source/TedsTableViewer/Public/Widgets/STedsHierarchyViewer.h
+++ b/Engine/Plugins/Experimental/EditorDataStorageFeatures/Source/TedsTableViewer/Public/Widgets/STedsHierarchyViewer.h
@@ -86,6 +86,9 @@ namespace UE::Editor::DataStorage
 		// Empty = simply show the column headers instead of a message
 		SLATE_ATTRIBUTE(FText, EmptyRowsMessage)
 
+		// The primary column in the hierarchy viewer is the one where the expansion row shows up
+		SLATE_ARGUMENT(TWeakObjectPtr<const UScriptStruct>, PrimaryColumn)
+
 		SLATE_END_ARGS()
 
 	public:
