[TEDS FilterBar] Allow for defining a custom 'common filters' section where those filters will display at the top of the create menu

diff --git a/Engine/Plugins/Experimental/EditorDataStorageFeatures/Source/TedsTableViewer/Private/Widgets/Composite/STedsCompositeHierarchyViewer.cpp b/Engine/Plugins/Experimental/EditorDataStorageFeatures/Source/TedsTableViewer/Private/Widgets/Composite/STedsCompositeHierarchyViewer.cpp
--- a/Engine/Plugins/Experimental/EditorDataStorageFeatures/Source/TedsTableViewer/Private/Widgets/Composite/STedsCompositeHierarchyViewer.cpp
+++ b/Engine/Plugins/Experimental/EditorDataStorageFeatures/Source/TedsTableViewer/Private/Widgets/Composite/STedsCompositeHierarchyViewer.cpp
@@ -40,6 +40,7 @@ namespace UE::Editor::DataStorage
 				FilterBar = SNew(STedsFilterBar)
 					.InFilterableRowNode(QueryStackNode)
 					.OutFilteredNode(&FilterNode)
+					.CommonSectionFilters(InArgs._CommonSectionFilters)
 					.Filters(InArgs._Filters)
 					.ClassFilters(InArgs._ClassFilters)
 					.UseSectionsForCategories(InArgs._UseSectionsForCategories)
diff --git a/Engine/Plugins/Experimental/EditorDataStorageFeatures/Source/TedsTableViewer/Private/Widgets/STedsFilterBar.cpp b/Engine/Plugins/Experimental/EditorDataStorageFeatures/Source/TedsTableViewer/Private/Widgets/STedsFilterBar.cpp
--- a/Engine/Plugins/Experimental/EditorDataStorageFeatures/Source/TedsTableViewer/Private/Widgets/STedsFilterBar.cpp
+++ b/Engine/Plugins/Experimental/EditorDataStorageFeatures/Source/TedsTableViewer/Private/Widgets/STedsFilterBar.cpp
@@ -17,6 +17,7 @@ namespace UE::Editor::DataStorage
 		using namespace DataStorage::QueryStack;
 
 		OnPostFiltersChanged = InArgs._OnPostFiltersChanged;
+		CommonSectionFilters = InArgs._CommonSectionFilters;
 
 		typename SBasicFilterBar<FTedsRowHandle&>::FArguments Args;
 		Args._OnFilterChanged = InArgs._OnFilterChanged;
@@ -70,8 +71,8 @@ namespace UE::Editor::DataStorage
 
 	void STedsFilterBar::PopulateCommonFilterSections(UToolMenu* InMenu)
 	{
-		FToolMenuSection& Section = InMenu->AddSection("FilterBarResetFilters");
-		Section.AddMenuEntry(
+		FToolMenuSection& ActionSection = InMenu->AddSection("FilterBarResetFilters");
+		ActionSection.AddMenuEntry(
 			"ResetFilters",
 			LOCTEXT("FilterListResetFilters", "Reset Filters"),
 			LOCTEXT("FilterListResetToolTip", "Resets current pinned filters"),
@@ -82,7 +83,32 @@ namespace UE::Editor::DataStorage
 			)
 		);
 
-		Section.AddSeparator("TedsCommonFiltersSeparator");
+		if (CommonSectionFilters.IsEmpty())
+		{
+			ActionSection.AddSeparator("TedsActionsFiltersSeparator");
+		}
+		else
+		{
+			FToolMenuSection& CommonFiltersSection = InMenu->AddSection("FilterBarCommonFilters", LOCTEXT("TedsFilterBarCommonFiltersSection", "Common Filters"));
+			for (const TSharedRef<FFilterBase<FTedsRowHandle&>>& FrontendFilter : AllFrontendFilters)
+			{
+				if(CommonSectionFilters.Contains(FName(*FrontendFilter->GetName())))
+				{
+					CommonFiltersSection.AddMenuEntry(
+						NAME_None,
+						FrontendFilter->GetDisplayName(),
+						FrontendFilter->GetToolTipText(),
+						FSlateIcon(FAppStyle::GetAppStyleSetName(), FrontendFilter->GetIconName()),
+						FUIAction(
+						FExecuteAction::CreateSP(this, &STedsFilterBar::FrontendFilterClicked, FrontendFilter),
+						FCanExecuteAction(),
+						FIsActionChecked::CreateSP(this, &STedsFilterBar::IsFrontendFilterInUse, FrontendFilter)),
+						EUserInterfaceActionType::ToggleButton
+						);
+				}
+			}
+			CommonFiltersSection.AddSeparator("TedsCommonFiltersSeparator");
+		}
 	}
 
 	void STedsFilterBar::PopulateCustomFilters(UToolMenu* InMenu)
diff --git a/Engine/Plugins/Experimental/EditorDataStorageFeatures/Source/TedsTableViewer/Public/Widgets/Composite/STedsCompositeHierarchyViewer.h b/Engine/Plugins/Experimental/EditorDataStorageFeatures/Source/TedsTableViewer/Public/Widgets/Composite/STedsCompositeHierarchyViewer.h
--- a/Engine/Plugins/Experimental/EditorDataStorageFeatures/Source/TedsTableViewer/Public/Widgets/Composite/STedsCompositeHierarchyViewer.h
+++ b/Engine/Plugins/Experimental/EditorDataStorageFeatures/Source/TedsTableViewer/Public/Widgets/Composite/STedsCompositeHierarchyViewer.h
@@ -50,6 +50,9 @@ namespace UE::Editor::DataStorage
 
 		// If we want to enable the filter bar on the table
 		SLATE_ARGUMENT(bool, EnableFiltering)
+		
+		/** Set of filter names to also add to the 'common filters' section (section at the top of the menu) */
+        SLATE_ARGUMENT(TSet<FName>, CommonSectionFilters)
 
 		// Array of filters to add to the filter bar *AND Filters
 		SLATE_ARGUMENT(TArray<FTedsFilterData>, Filters)
diff --git a/Engine/Plugins/Experimental/EditorDataStorageFeatures/Source/TedsTableViewer/Public/Widgets/STedsFilterBar.h b/Engine/Plugins/Experimental/EditorDataStorageFeatures/Source/TedsTableViewer/Public/Widgets/STedsFilterBar.h
--- a/Engine/Plugins/Experimental/EditorDataStorageFeatures/Source/TedsTableViewer/Public/Widgets/STedsFilterBar.h
+++ b/Engine/Plugins/Experimental/EditorDataStorageFeatures/Source/TedsTableViewer/Public/Widgets/STedsFilterBar.h
@@ -54,6 +54,9 @@ namespace UE::Editor::DataStorage
 			/** Array of class filters to add to the filter bar *OR Filters */
 			SLATE_ARGUMENT(TArray<UClass*>, ClassFilters)
 			
+			/** Set of filter names to also add to the 'common filters' section (section at the top of the menu) */
+            SLATE_ARGUMENT(TSet<FName>, CommonSectionFilters)
+			
 			/** Delegate for when filters have changed */
 			SLATE_EVENT(FOnFilterChanged, OnFilterChanged)
 
@@ -104,6 +107,10 @@ namespace UE::Editor::DataStorage
 
 		ICoreProvider* Storage = nullptr;
 
+		// Set of Filter FNames used to populate the 'common filters' section that filters can have a duplicate entry
+		// placed at the top of the add filter menu
+		TSet<FName> CommonSectionFilters;
+
 		// The row node passed in as the data to filter on
 		TSharedPtr<QueryStack::IRowNode> InFilterableRowNode;
 
