Skip IDO types from TEDS type columns

diff --git a/Engine/Plugins/Experimental/EditorDataStorageFeatures/Source/TedsTypeInfo/Private/TedsTypeInfoFactory.cpp b/Engine/Plugins/Experimental/EditorDataStorageFeatures/Source/TedsTypeInfo/Private/TedsTypeInfoFactory.cpp
--- a/Engine/Plugins/Experimental/EditorDataStorageFeatures/Source/TedsTypeInfo/Private/TedsTypeInfoFactory.cpp
+++ b/Engine/Plugins/Experimental/EditorDataStorageFeatures/Source/TedsTypeInfo/Private/TedsTypeInfoFactory.cpp
@@ -16,6 +16,7 @@
 #include "Elements/Framework/TypedElementQueryBuilder.h"
 
 #include "HAL/IConsoleManager.h"
+#include "UObject/InstanceDataObjectUtils.h"
 #include "UObject/PropertyBagRepository.h"
 
 #include "UObject/UObjectIterator.h"
@@ -146,6 +147,12 @@ void UTypeInfoFactory::PopulateTypeInfo()
 
 bool UTypeInfoFactory::TryAddTypeInfoRow(const UStruct* TypeInfo)
 {
+	// Do not include InstanceDataObject types in TEDS
+	if (UE::IsClassOfInstanceDataObjectClass(TypeInfo))
+	{
+		return false;
+	}
+
 	using namespace UE::Editor::DataStorage;
 
 	ICoreProvider* DataStorage = GetMutableDataStorageFeature<ICoreProvider>(StorageFeatureName);
