[Prefab Editor] Create the Unsaved/IsPackageDirty Column and show IsSaveable and IsDirty widgets in the Prefab Editor
- Extract the Package sync from the CVar (since it should not be included in syncing revision control)
- Add Tags in the UnsavedAssetTracker

diff --git a/Engine/Plugins/Experimental/EditorDataStorageFeatures/Source/TedsOutliner/Private/Widgets/OutlinerUnsavedWidget.cpp b/Engine/Plugins/Experimental/EditorDataStorageFeatures/Source/TedsOutliner/Private/Widgets/OutlinerUnsavedWidget.cpp
new file mode 100644
--- /dev/null
+++ b/Engine/Plugins/Experimental/EditorDataStorageFeatures/Source/TedsOutliner/Private/Widgets/OutlinerUnsavedWidget.cpp
@@ -0,0 +1,94 @@
+// Copyright Epic Games, Inc. All Rights Reserved.
+
+#include "Widgets/OutlinerUnsavedWidget.h"
+
+#include "Columns/SlateHeaderColumns.h"
+#include "Elements/Columns/TypedElementPackageColumns.h"
+#include "Elements/Framework/TypedElementAttributeBinding.h"
+#include "Styling/AppStyle.h"
+#include "Widgets/Images/SImage.h"
+#include "Widgets/Layout/SBox.h"
+#include "Widgets/Text/STextBlock.h"
+
+#include UE_INLINE_GENERATED_CPP_BY_NAME(OutlinerUnsavedWidget)
+
+#define LOCTEXT_NAMESPACE "OutlinerUnsavedWidget"
+
+void UOutlinerUnsavedWidgetFactory::RegisterWidgetConstructors(UE::Editor::DataStorage::ICoreProvider& DataStorage,
+	UE::Editor::DataStorage::IUiProvider& DataStorageUi) const
+{
+	using namespace UE::Editor::DataStorage::Queries;
+
+	DataStorageUi.RegisterWidgetFactory<FOutlinerUnsavedWidgetConstructor>(
+	DataStorageUi.FindPurpose(UE::Editor::DataStorage::IUiProvider::FPurposeInfo("General", "Cell", NAME_None).GeneratePurposeID()),
+		TColumn<FTedsPrimaryPackageObjectTag>());
+	DataStorageUi.RegisterWidgetFactory<FOutlinerUnsavedHeaderConstructor>(
+		DataStorageUi.FindPurpose(UE::Editor::DataStorage::IUiProvider::FPurposeInfo("General", "Header", NAME_None).GeneratePurposeID()),
+		TColumn<FTedsPrimaryPackageObjectTag>());
+}
+
+FOutlinerUnsavedWidgetConstructor::FOutlinerUnsavedWidgetConstructor()
+	: FSimpleWidgetConstructor(StaticStruct())
+{
+}
+
+FOutlinerUnsavedHeaderConstructor::FOutlinerUnsavedHeaderConstructor()
+	: FSimpleWidgetConstructor(StaticStruct())
+{
+}
+
+TSharedPtr<SWidget> FOutlinerUnsavedHeaderConstructor::CreateWidget(UE::Editor::DataStorage::ICoreProvider* DataStorage,
+	UE::Editor::DataStorage::IUiProvider* DataStorageUi,
+	UE::Editor::DataStorage::RowHandle TargetRow, UE::Editor::DataStorage::RowHandle WidgetRow, const UE::Editor::DataStorage::FMetaDataView& Arguments)
+{
+	DataStorage->AddColumn(WidgetRow, FHeaderWidgetSizeColumn
+	{
+		.ColumnSizeMode = EColumnSizeMode::Fixed,
+		.Width = 24.0f
+	});
+
+	return SNew(SImage)
+		.DesiredSizeOverride(FVector2D(12.f, 12.f))
+		.Image(FAppStyle::Get().GetBrush("Icons.DirtyBadge"))
+		.ToolTipText(LOCTEXT("SceneOutlinerUnsavedHeaderTooltip", "Unsaved"));
+}
+
+TSharedPtr<SWidget> FOutlinerUnsavedWidgetConstructor::CreateWidget(
+	UE::Editor::DataStorage::ICoreProvider* DataStorage,
+	UE::Editor::DataStorage::IUiProvider* DataStorageUi,
+	UE::Editor::DataStorage::RowHandle TargetRow,
+	UE::Editor::DataStorage::RowHandle WidgetRow,
+	const UE::Editor::DataStorage::FMetaDataView& Arguments)
+{
+	using namespace UE::Editor::DataStorage;
+	if (!DataStorage->IsRowAvailable(TargetRow))
+	{
+		return SNew(STextBlock)
+			.Text(LOCTEXT("MissingRowReferenceColumn", "Unable to retrieve row reference."));
+	}
+
+	FAttributeBinder Binder(TargetRow, DataStorage);
+
+	return SNew(SBox)
+			.HAlign(HAlign_Center)
+			.VAlign(VAlign_Center)
+			[
+				SNew(SImage)
+					.ColorAndOpacity(FSlateColor::UseForeground())
+					.Image(Binder.BindColumnsPresence(Queries::TColumn<FTedsPackageDirtyTag>(), [](bool bIsUnsaved)
+						{
+							// Since the widget maps to the FTedsPrimaryPackageObjectTag, it means that the row is mapped to a package. Meaning that
+							// the presence of FTedsPackageDirtyTag just lets us know to use the IsDirty icon vs the IsSavable icon
+							if (bIsUnsaved)
+							{
+								return FAppStyle::GetBrush("Icons.DirtyBadge");
+							}
+							else
+							{
+								return FAppStyle::GetBrush("Icons.SaveableBadge");
+							}
+						}))
+			];
+}
+
+#undef LOCTEXT_NAMESPACE
diff --git a/Engine/Plugins/Experimental/EditorDataStorageFeatures/Source/TedsOutliner/Public/Widgets/OutlinerUnsavedWidget.h b/Engine/Plugins/Experimental/EditorDataStorageFeatures/Source/TedsOutliner/Public/Widgets/OutlinerUnsavedWidget.h
new file mode 100644
--- /dev/null
+++ b/Engine/Plugins/Experimental/EditorDataStorageFeatures/Source/TedsOutliner/Public/Widgets/OutlinerUnsavedWidget.h
@@ -0,0 +1,55 @@
+// Copyright Epic Games, Inc. All Rights Reserved.
+
+#pragma once
+
+#include "Elements/Interfaces/TypedElementDataStorageFactory.h"
+#include "Elements/Interfaces/TypedElementDataStorageInterface.h"
+#include "Elements/Interfaces/TypedElementDataStorageUiInterface.h"
+
+#include "OutlinerUnsavedWidget.generated.h"
+
+UCLASS()
+class UOutlinerUnsavedWidgetFactory : public UEditorDataStorageFactory
+{
+	GENERATED_BODY()
+
+public:
+	~UOutlinerUnsavedWidgetFactory() override = default;
+
+	void RegisterWidgetConstructors(UE::Editor::DataStorage::ICoreProvider& DataStorage,
+		UE::Editor::DataStorage::IUiProvider& DataStorageUi) const override;
+};
+
+USTRUCT()
+struct FOutlinerUnsavedWidgetConstructor : public FSimpleWidgetConstructor
+{
+	GENERATED_BODY()
+
+public:
+	TEDSOUTLINER_API FOutlinerUnsavedWidgetConstructor();
+	~FOutlinerUnsavedWidgetConstructor() override = default;
+
+	TEDSOUTLINER_API virtual TSharedPtr<SWidget> CreateWidget(
+		UE::Editor::DataStorage::ICoreProvider* DataStorage,
+		UE::Editor::DataStorage::IUiProvider* DataStorageUi,
+		UE::Editor::DataStorage::RowHandle TargetRow,
+		UE::Editor::DataStorage::RowHandle WidgetRow,
+		const UE::Editor::DataStorage::FMetaDataView& Arguments) override;
+};
+
+USTRUCT()
+struct FOutlinerUnsavedHeaderConstructor : public FSimpleWidgetConstructor
+{
+	GENERATED_BODY()
+
+public:
+	TEDSOUTLINER_API FOutlinerUnsavedHeaderConstructor();
+	~FOutlinerUnsavedHeaderConstructor() override = default;
+
+	TEDSOUTLINER_API virtual TSharedPtr<SWidget> CreateWidget(
+		UE::Editor::DataStorage::ICoreProvider* DataStorage,
+		UE::Editor::DataStorage::IUiProvider* DataStorageUi,
+		UE::Editor::DataStorage::RowHandle TargetRow,
+		UE::Editor::DataStorage::RowHandle WidgetRow, 
+		const UE::Editor::DataStorage::FMetaDataView& Arguments) override;
+};
diff --git a/Engine/Plugins/Experimental/EditorDataStorageFeatures/Source/TedsRevisionControl/Private/Processors/RevisionControlProcessors.cpp b/Engine/Plugins/Experimental/EditorDataStorageFeatures/Source/TedsRevisionControl/Private/Processors/RevisionControlProcessors.cpp
--- a/Engine/Plugins/Experimental/EditorDataStorageFeatures/Source/TedsRevisionControl/Private/Processors/RevisionControlProcessors.cpp
+++ b/Engine/Plugins/Experimental/EditorDataStorageFeatures/Source/TedsRevisionControl/Private/Processors/RevisionControlProcessors.cpp
@@ -21,7 +21,12 @@
 
 namespace UE::Editor::RevisionControl::Private
 {
-	extern FAutoConsoleVariableRef CVarAutoPopulateState;
+	static bool bAutoPopulateRevisionControlState = false;
+	FAutoConsoleVariableRef CVarAutoPopulateState(
+		TEXT("TEDS.RevisionControl.AutoPopulateState"),
+		bAutoPopulateRevisionControlState,
+		TEXT("Automatically query revision control provider and fill information into TEDS")
+	);
 
 	// Update the overlay color for all rows with InColumn
 	void UpdateSCCOverlayStates(const UScriptStruct* InColumn)
diff --git a/Engine/Plugins/Experimental/EditorDataStorageFeatures/Source/TedsRevisionControl/Private/Queries/ObjectPackagePathToColumnQueries.cpp b/Engine/Plugins/Experimental/EditorDataStorageFeatures/Source/TedsRevisionControl/Private/Queries/ObjectPackagePathToColumnQueries.cpp
--- a/Engine/Plugins/Experimental/EditorDataStorageFeatures/Source/TedsRevisionControl/Private/Queries/ObjectPackagePathToColumnQueries.cpp
+++ b/Engine/Plugins/Experimental/EditorDataStorageFeatures/Source/TedsRevisionControl/Private/Queries/ObjectPackagePathToColumnQueries.cpp
@@ -2,10 +2,12 @@
 
 #include "Queries/ObjectPackagePathToColumnQueries.h"
 
+#include "DataStorage/Features.h"
 #include "Elements/Columns/TypedElementCompatibilityColumns.h"
 #include "Elements/Columns/TypedElementPackageColumns.h"
 #include "Elements/Columns/TypedElementRevisionControlColumns.h"
 #include "Elements/Framework/TypedElementQueryBuilder.h"
+#include "Engine/Blueprint.h"
 #include "HAL/IConsoleManager.h"
 #include "Misc/PackageName.h"
 #include "Misc/Paths.h"
@@ -17,14 +19,14 @@ namespace UE::Editor::RevisionControl::Private
 {
 	using namespace UE::Editor::DataStorage;
 
-	static bool bAutoPopulateRevisionControlState = false;
-	FAutoConsoleVariableRef CVarAutoPopulateState(
-		TEXT("TEDS.RevisionControl.AutoPopulateState"),
-		bAutoPopulateRevisionControlState,
-		TEXT("Automatically query revision control provider and fill information into TEDS")
-	);
+	static bool bAutoSyncPackageInformation = true;
+    FAutoConsoleVariableRef CVarSyncPackageInformation(
+    	TEXT("TEDS.RevisionControl.SyncPackageInformation"),
+    	bAutoSyncPackageInformation,
+    	TEXT("Automatically query package information and fill information into TEDS")
+    );
 
-	static void ResolvePackageReference(IQueryContext& Context, const UPackage* Package, RowHandle Row, RowHandle PackageRow)
+	static void ResolvePackageReference(IQueryContext& Context, const UPackage* Package, const UObject* Object, RowHandle Row, RowHandle PackageRow)
 	{
 		FTypedElementPackageReference PackageReference;
 		PackageReference.Row = PackageRow;
@@ -39,17 +41,58 @@ namespace UE::Editor::RevisionControl::Private
 		Context.AddColumn(PackageRow, MoveTemp(PathColumn));
 		Context.AddColumn(PackageRow, MoveTemp(LoadedPathColumn));
 
-		FTypedElementPackageReference PackageBackReference;
-		PackageBackReference.Row = Row;
-		Context.AddColumn(PackageRow, MoveTemp(PackageBackReference));
+		if (const UObject* PackageInstance = Package->FindAssetInPackage())
+		{
+			if (Object == PackageInstance)
+			{
+				Context.AddColumns<FTedsPrimaryPackageObjectTag>(Row);
+			}
+			// If the Row Object is a CDO, we have to compare the CDO of the primary package object instead of the object itself
+			else if (Context.HasColumn<FTypedElementClassDefaultObjectTag>())
+			{
+				const UObject* PackageObjectCDO = nullptr;
+				// Check if it was BP generated and get the Default Object of that generated class if so
+				if(const UBlueprint* PackageObjectBP = Cast<UBlueprint>(Object))
+				{
+					if (const TSubclassOf<UObject> PackageObjectBPGeneratedClass = PackageObjectBP->GeneratedClass)
+					{
+						PackageObjectCDO = PackageObjectBPGeneratedClass->GetDefaultObject();
+					}
+				}
+				else
+				{
+					PackageObjectCDO = Object->GetClass()->GetDefaultObject();
+				}
+				
+				if (Object == PackageObjectCDO)
+				{
+					Context.AddColumns<FTedsPrimaryPackageObjectTag>(Row);
+				}
+			}
+		}
+
+		Context.SetParentRow(Row, PackageRow);
 	};
-} // namespace UE::Editor::DataStorage
+} // namespace UE::Editor::RevisionControl::Private
+
+
+void UTypedElementUObjectPackagePathFactory::RegisterHierarchies(UE::Editor::DataStorage::ICoreProvider& DataStorage)
+{
+	using namespace UE::Editor::DataStorage;
+
+	const FHierarchyRegistrationParams Params
+	{
+		.Name = PackageHierarchyName,
+	};
+
+	DataStorage.RegisterHierarchy(Params);
+}
 
 void UTypedElementUObjectPackagePathFactory::RegisterQueries(UE::Editor::DataStorage::ICoreProvider& DataStorage)
 {
 	using namespace UE::Editor::DataStorage::Queries;
 	
-	UE::Editor::RevisionControl::Private::CVarAutoPopulateState->AsVariable()->OnChangedDelegate().AddLambda(
+	UE::Editor::RevisionControl::Private::CVarSyncPackageInformation->AsVariable()->OnChangedDelegate().AddLambda(
 		[this, &DataStorage](IConsoleVariable* AutoPopulate)
 		{
 			if (AutoPopulate->GetBool())
@@ -82,11 +125,12 @@ void UTypedElementUObjectPackagePathFactory::RegisterQueries(UE::Editor::DataSto
 				const UPackage* Package = ObjectInstance->GetPackage();
 				Context.RemoveColumns(Row, { FTypedElementPackageUnresolvedReference::StaticStruct() });
 
-				UE::Editor::RevisionControl::Private::ResolvePackageReference(Context, Package, Row, PackageRow);
+				UE::Editor::RevisionControl::Private::ResolvePackageReference(Context, Package, ObjectInstance, Row, PackageRow);
 			})
+		.AccessesHierarchy(UE::Editor::DataStorage::PackageHierarchyName)
 		.Compile());
 
-	if (UE::Editor::RevisionControl::Private::CVarAutoPopulateState->GetBool())
+	if (UE::Editor::RevisionControl::Private::CVarSyncPackageInformation->GetBool())
 	{
 		RegisterTryAddPackageRef(DataStorage);
 	}
@@ -94,7 +138,40 @@ void UTypedElementUObjectPackagePathFactory::RegisterQueries(UE::Editor::DataSto
 
 void UTypedElementUObjectPackagePathFactory::RegisterTryAddPackageRef(UE::Editor::DataStorage::ICoreProvider& DataStorage)
 {
+	using namespace UE::Editor::DataStorage;
 	using namespace UE::Editor::DataStorage::Queries;
+
+	struct FRegisterPackageRowCommand
+	{
+		void operator()()
+		{
+			if (CoreProvider)
+			{
+				// Check in case this row was mapped before this command was pushed 
+				RowHandle PackageRow = CoreProvider->LookupMappedRow(UE::Editor::RevisionControl::MappingDomain, FMapKey(PackagePathname));
+				if (!CoreProvider->IsRowAvailable(PackageRow))
+				{
+					const RowHandle RevisionControlTable = CoreProvider->FindTable(FName("Editor_RevisionControlTable"));
+					if (RevisionControlTable != InvalidRowHandle)
+					{
+						RowHandle Row = CoreProvider->AddRow(RevisionControlTable);
+						CoreProvider->MapRow(UE::Editor::RevisionControl::MappingDomain, FMapKey(PackagePathname), Row);
+						if (bIsMapPackage)
+						{
+							// Also map to the .uasset path
+							CoreProvider->MapRow(
+								UE::Editor::RevisionControl::MappingDomain,
+								FMapKey(FPaths::ChangeExtension(PackagePathname, FPackageName::GetAssetPackageExtension())),
+								Row);
+						}
+					}
+				}
+			}
+		}
+		FString PackagePathname;
+		ICoreProvider* CoreProvider = nullptr;
+		bool bIsMapPackage = false;
+	};
 	
 	TryAddPackageRef = DataStorage.RegisterQuery(
 		Select(
@@ -119,9 +196,15 @@ void UTypedElementUObjectPackagePathFactory::RegisterTryAddPackageRef(UE::Editor
 						return;
 					}
 
-					const FString PackageExtension = Target->ContainsMap()
-						? FPackageName::GetMapPackageExtension()
-						: FPackageName::GetAssetPackageExtension();
+					// Certain Package utilities don't get the .umap and will return the .uasset instead so we want to map this row to both of those
+					FString PackageExtension = FPackageName::GetAssetPackageExtension();
+
+					bool bIsMapPackage = false;
+					if (Target->ContainsMap())
+					{
+						bIsMapPackage = true;
+						PackageExtension = FPackageName::GetMapPackageExtension();
+					}
 
 					FString Path;
 					Target->GetPathName(nullptr, Path);
@@ -134,16 +217,25 @@ void UTypedElementUObjectPackagePathFactory::RegisterTryAddPackageRef(UE::Editor
 						if (Context.IsRowAvailable(PackageRow))
 						{
 							const UPackage* Package = ObjectInstance->GetPackage();
-							UE::Editor::RevisionControl::Private::ResolvePackageReference(Context, Package, Row, PackageRow);
+							UE::Editor::RevisionControl::Private::ResolvePackageReference(Context, Package, ObjectInstance, Row, PackageRow);
 						}
 						else
 						{
 							FTypedElementPackageUnresolvedReference UnresolvedPackageReference;
-							UnresolvedPackageReference.PathOnDisk = MoveTemp(FullPackageFilename);
+							UnresolvedPackageReference.PathOnDisk = FullPackageFilename;
 							Context.AddColumn(Row, MoveTemp(UnresolvedPackageReference));
+
+							FRegisterPackageRowCommand Command
+							{
+								.PackagePathname = FullPackageFilename,
+								.CoreProvider = GetMutableDataStorageFeature<ICoreProvider>(StorageFeatureName),
+								.bIsMapPackage = bIsMapPackage
+							};
+							Context.PushCommand(Command);
 						}
 					}
 				}
 			})
+		.AccessesHierarchy(PackageHierarchyName)
 		.Compile());
 }
diff --git a/Engine/Plugins/Experimental/EditorDataStorageFeatures/Source/TedsRevisionControl/Private/Queries/ObjectPackagePathToColumnQueries.h b/Engine/Plugins/Experimental/EditorDataStorageFeatures/Source/TedsRevisionControl/Private/Queries/ObjectPackagePathToColumnQueries.h
--- a/Engine/Plugins/Experimental/EditorDataStorageFeatures/Source/TedsRevisionControl/Private/Queries/ObjectPackagePathToColumnQueries.h
+++ b/Engine/Plugins/Experimental/EditorDataStorageFeatures/Source/TedsRevisionControl/Private/Queries/ObjectPackagePathToColumnQueries.h
@@ -20,6 +20,7 @@ class UTypedElementUObjectPackagePathFactory : public UEditorDataStorageFactory
 public:
 	~UTypedElementUObjectPackagePathFactory() override = default;
 
+	void RegisterHierarchies(UE::Editor::DataStorage::ICoreProvider& DataStorage) override;
 	void RegisterQueries(UE::Editor::DataStorage::ICoreProvider& DataStorage) override;
 
 private:
diff --git a/Engine/Source/Developer/UnsavedAssetsTracker/Source/Private/UnsavedAssetsTracker.cpp b/Engine/Source/Developer/UnsavedAssetsTracker/Source/Private/UnsavedAssetsTracker.cpp
--- a/Engine/Source/Developer/UnsavedAssetsTracker/Source/Private/UnsavedAssetsTracker.cpp
+++ b/Engine/Source/Developer/UnsavedAssetsTracker/Source/Private/UnsavedAssetsTracker.cpp
@@ -25,6 +25,11 @@
 #include "IAssetTools.h"
 #include "UObject/UObjectGlobals.h"
 
+#include "DataStorage/Features.h"
+#include "Elements/Columns/TypedElementPackageColumns.h"
+#include "Elements/Columns/TypedElementRevisionControlColumns.h"
+#include "Elements/Interfaces/TypedElementDataStorageInterface.h"
+
 #define LOCTEXT_NAMESPACE "UnsavedAssetsTracker"
 
 DEFINE_LOG_CATEGORY_STATIC(LogUnsavedAssetsTracker, Log, All);
@@ -434,6 +439,8 @@ void FUnsavedAssetsTracker::StartTrackingDirtyPackage(const UPackage* Package)
 		AbsPackageFilePathname,
 		FSourceControlFileStatusMonitor::FOnSourceControlFileStatus::CreateSP(this, &FUnsavedAssetsTracker::OnSourceControlFileStatusUpdate));
 
+	UpdateTEDSPackage(PackagePathname, true);
+
 	FUnsavedAssetsTrackerModule::Get().OnUnsavedAssetAdded.Broadcast(AbsPackageFilePathname);
 
 	UE_LOG(LogUnsavedAssetsTracker, Verbose, TEXT("Added file to the unsaved asset list: %s (%s)"), *GetHumanFriendlyAssetName(Package), *PackagePathname);
@@ -451,6 +458,8 @@ void FUnsavedAssetsTracker::StopTrackingDirtyPackage(const FString& PackagePathn
 		WarningPackagePathnames.Remove(PackagePathname);
 		UnsavedAbsFilePathames.Remove(AbsPackageFilePathname);
 
+		UpdateTEDSPackage(PackagePathname, false);
+
 		ISourceControlModule::Get().GetSourceControlFileStatusMonitor().StopMonitoringFile(reinterpret_cast<uintptr_t>(this), AbsPackageFilePathname);
 		FUnsavedAssetsTrackerModule::Get().OnUnsavedAssetRemoved.Broadcast(AbsPackageFilePathname);
 
@@ -463,6 +472,39 @@ void FUnsavedAssetsTracker::StopTrackingDirtyPackage(const FString& PackagePathn
 	}
 }
 
+void FUnsavedAssetsTracker::UpdateTEDSPackage(const FString& PackagePathname, const bool bIsDirty) const
+{
+	using namespace UE::Editor::DataStorage;
+
+	if (ICoreProvider* Storage = GetMutableDataStorageFeature<ICoreProvider>(StorageFeatureName))
+	{
+		FString FullPackagePath = FPackageName::LongPackageNameToFilename(PackagePathname, FPackageName::GetAssetPackageExtension());
+		FullPackagePath = FPaths::ConvertRelativePathToFull(FullPackagePath);
+		const RowHandle PackageRow = Storage->LookupMappedRow(UE::Editor::RevisionControl::MappingDomain, FMapKey(FullPackagePath));
+		if (Storage->IsRowAvailable(PackageRow))
+		{
+			TArray<RowHandle> ChildrenRowHandles;
+			auto GetChildRows = [&ChildrenRowHandles](const ICoreProvider& Context, RowHandle, RowHandle TargetRow)
+			{
+				ChildrenRowHandles.Add(TargetRow);
+			};
+			Storage->WalkDepthFirst(Storage->FindHierarchyByName(PackageHierarchyName), PackageRow, GetChildRows);
+
+			for (RowHandle ChildRow : ChildrenRowHandles)
+			{
+				if (bIsDirty)
+				{
+					Storage->AddColumns<FTedsPackageDirtyTag>(ChildRow);
+				}
+				else
+				{
+					Storage->RemoveColumn<FTedsPackageDirtyTag>(ChildRow);
+				}
+			}
+		}
+	}
+}
+
 void FUnsavedAssetsTracker::OnSourceControlFileStatusUpdate(const FString& AbsFilePathname, const ISourceControlState* State)
 {
 	auto DiscardWarning = [this](const FString& InPackagePathname, FStatus& InStatus)
diff --git a/Engine/Source/Developer/UnsavedAssetsTracker/Source/Private/UnsavedAssetsTracker.h b/Engine/Source/Developer/UnsavedAssetsTracker/Source/Private/UnsavedAssetsTracker.h
--- a/Engine/Source/Developer/UnsavedAssetsTracker/Source/Private/UnsavedAssetsTracker.h
+++ b/Engine/Source/Developer/UnsavedAssetsTracker/Source/Private/UnsavedAssetsTracker.h
@@ -4,6 +4,7 @@
 
 #include "CoreMinimal.h"
 #include "Containers/Ticker.h"
+#include "DataStorage/Handles.h"
 #include "HAL/CriticalSection.h"
 #include "Templates/SharedPointer.h"
 #include "UObject/ObjectSaveContext.h"
@@ -135,6 +136,9 @@ private:
 
 	/** Whether the specified package should be tracked. */
 	bool ShouldTrackDirtyPackage(const UPackage* Package);
+	
+	/** Adds the Dirty tag to the Package and its contained object rows */
+	void UpdateTEDSPackage(const FString& PackagePathname, const bool bIsDirty) const;
 
 private:
 	/** Supports multithreaded Dirty notifications. */
diff --git a/Engine/Source/Developer/UnsavedAssetsTracker/Source/UnsavedAssetsTracker.Build.cs b/Engine/Source/Developer/UnsavedAssetsTracker/Source/UnsavedAssetsTracker.Build.cs
--- a/Engine/Source/Developer/UnsavedAssetsTracker/Source/UnsavedAssetsTracker.Build.cs
+++ b/Engine/Source/Developer/UnsavedAssetsTracker/Source/UnsavedAssetsTracker.Build.cs
@@ -14,6 +14,7 @@ public class UnsavedAssetsTracker : ModuleRules
 				"EditorSubsystem",
 				"SourceControl",
 				"UncontrolledChangelists",
+				"TypedElementFramework",
 			}
 		);
 
diff --git a/Engine/Source/Runtime/TypedElementFramework/Public/Elements/Columns/TypedElementPackageColumns.h b/Engine/Source/Runtime/TypedElementFramework/Public/Elements/Columns/TypedElementPackageColumns.h
--- a/Engine/Source/Runtime/TypedElementFramework/Public/Elements/Columns/TypedElementPackageColumns.h
+++ b/Engine/Source/Runtime/TypedElementFramework/Public/Elements/Columns/TypedElementPackageColumns.h
@@ -14,6 +14,11 @@
 
 #include "TypedElementPackageColumns.generated.h"
 
+namespace UE::Editor::DataStorage
+{
+	inline static const FName PackageHierarchyName = TEXT("PackageHierarchy");
+}
+
 /**
  * A package reference column that has not yet been resolved to reference a package.
  */
@@ -72,3 +77,21 @@ struct FTypedElementPackageLoadedPathColumn final : public FEditorDataStorageCol
 
 	FPackagePath LoadedPath;
 };
+
+/**
+ * Tag to indicate that the data on the row is in a package that was marked unsaved/marked dirty
+ */
+USTRUCT(meta = (DisplayName = "In Dirty Package"))
+struct FTedsPackageDirtyTag final : public FEditorDataStorageTag
+{
+	GENERATED_BODY()
+};
+
+/**
+ * Tag to indicate that the data on the row is the primary object in a package
+ */
+USTRUCT(meta = (DisplayName = "Primary Package Object"))
+struct FTedsPrimaryPackageObjectTag final : public FEditorDataStorageTag
+{
+	GENERATED_BODY()
+};
