[TEDs] Fix TypeInfoWidget crash on invalid weak pointer

diff --git a/Engine/Plugins/Experimental/EditorDataStorageFeatures/Source/TedsOutliner/Private/Widgets/TypeInfoWidget.cpp b/Engine/Plugins/Experimental/EditorDataStorageFeatures/Source/TedsOutliner/Private/Widgets/TypeInfoWidget.cpp
--- a/Engine/Plugins/Experimental/EditorDataStorageFeatures/Source/TedsOutliner/Private/Widgets/TypeInfoWidget.cpp
+++ b/Engine/Plugins/Experimental/EditorDataStorageFeatures/Source/TedsOutliner/Private/Widgets/TypeInfoWidget.cpp
@@ -33,7 +33,9 @@ namespace UE::Editor::DataStorage
 
 	FPrefixInfo FTypeInfoWidgetSorter::CalculatePrefix(const FTypedElementClassTypeInfoColumn& Column, uint32 ByteIndex) const
 	{
-		return CreateSortPrefix(ByteIndex, TSortStringView(FSortCaseInsensitive{}, Column.TypeInfo->GetName()));
+		TObjectPtr<const UClass> TypeInfoPtr = Column.TypeInfo.Get();
+		FStringView TypeInfoName = TypeInfoPtr ? TypeInfoPtr->GetName() : TEXT("NULL");
+		return CreateSortPrefix(ByteIndex, TSortStringView(FSortCaseInsensitive{}, TypeInfoName));
 	}
 }
 
