[Backout] - CL47601913
#rnx
Original CL Desc
-----------------------------------------------------------------
[Prefab Editor] Create the Unsaved/IsPackageDirty Column and show IsSaveable and IsDirty widgets in the Prefab Editor
- Extract the Package sync from the CVar (since it should not be included in syncing revision control)
- Add Tags in the UnsavedAssetTracker

diff --git a/Engine/Plugins/Experimental/EditorDataStorageFeatures/Source/TedsOutliner/Private/Widgets/OutlinerUnsavedWidget.cpp b/Engine/Plugins/Experimental/EditorDataStorageFeatures/Source/TedsOutliner/Private/Widgets/OutlinerUnsavedWidget.cpp
deleted file mode 100644
diff --git a/Engine/Plugins/Experimental/EditorDataStorageFeatures/Source/TedsOutliner/Public/Widgets/OutlinerUnsavedWidget.h b/Engine/Plugins/Experimental/EditorDataStorageFeatures/Source/TedsOutliner/Public/Widgets/OutlinerUnsavedWidget.h
deleted file mode 100644
diff --git a/Engine/Plugins/Experimental/EditorDataStorageFeatures/Source/TedsRevisionControl/Private/Processors/RevisionControlProcessors.cpp b/Engine/Plugins/Experimental/EditorDataStorageFeatures/Source/TedsRevisionControl/Private/Processors/RevisionControlProcessors.cpp
--- a/Engine/Plugins/Experimental/EditorDataStorageFeatures/Source/TedsRevisionControl/Private/Processors/RevisionControlProcessors.cpp
+++ b/Engine/Plugins/Experimental/EditorDataStorageFeatures/Source/TedsRevisionControl/Private/Processors/RevisionControlProcessors.cpp
@@ -21,12 +21,7 @@
 
 namespace UE::Editor::RevisionControl::Private
 {
-	static bool bAutoPopulateRevisionControlState = false;
-	FAutoConsoleVariableRef CVarAutoPopulateState(
-		TEXT("TEDS.RevisionControl.AutoPopulateState"),
-		bAutoPopulateRevisionControlState,
-		TEXT("Automatically query revision control provider and fill information into TEDS")
-	);
+	extern FAutoConsoleVariableRef CVarAutoPopulateState;
 
 	// Update the overlay color for all rows with InColumn
 	void UpdateSCCOverlayStates(const UScriptStruct* InColumn)
diff --git a/Engine/Plugins/Experimental/EditorDataStorageFeatures/Source/TedsRevisionControl/Private/Queries/ObjectPackagePathToColumnQueries.cpp b/Engine/Plugins/Experimental/EditorDataStorageFeatures/Source/TedsRevisionControl/Private/Queries/ObjectPackagePathToColumnQueries.cpp
--- a/Engine/Plugins/Experimental/EditorDataStorageFeatures/Source/TedsRevisionControl/Private/Queries/ObjectPackagePathToColumnQueries.cpp
+++ b/Engine/Plugins/Experimental/EditorDataStorageFeatures/Source/TedsRevisionControl/Private/Queries/ObjectPackagePathToColumnQueries.cpp
@@ -2,12 +2,10 @@
 
 #include "Queries/ObjectPackagePathToColumnQueries.h"
 
-#include "DataStorage/Features.h"
 #include "Elements/Columns/TypedElementCompatibilityColumns.h"
 #include "Elements/Columns/TypedElementPackageColumns.h"
 #include "Elements/Columns/TypedElementRevisionControlColumns.h"
 #include "Elements/Framework/TypedElementQueryBuilder.h"
-#include "Engine/Blueprint.h"
 #include "HAL/IConsoleManager.h"
 #include "Misc/PackageName.h"
 #include "Misc/Paths.h"
@@ -19,14 +17,14 @@ namespace UE::Editor::RevisionControl::Private
 {
 	using namespace UE::Editor::DataStorage;
 
-	static bool bAutoSyncPackageInformation = true;
-    FAutoConsoleVariableRef CVarSyncPackageInformation(
-    	TEXT("TEDS.RevisionControl.SyncPackageInformation"),
-    	bAutoSyncPackageInformation,
-    	TEXT("Automatically query package information and fill information into TEDS")
-    );
+	static bool bAutoPopulateRevisionControlState = false;
+	FAutoConsoleVariableRef CVarAutoPopulateState(
+		TEXT("TEDS.RevisionControl.AutoPopulateState"),
+		bAutoPopulateRevisionControlState,
+		TEXT("Automatically query revision control provider and fill information into TEDS")
+	);
 
-	static void ResolvePackageReference(IQueryContext& Context, const UPackage* Package, const UObject* Object, RowHandle Row, RowHandle PackageRow)
+	static void ResolvePackageReference(IQueryContext& Context, const UPackage* Package, RowHandle Row, RowHandle PackageRow)
 	{
 		FTypedElementPackageReference PackageReference;
 		PackageReference.Row = PackageRow;
@@ -41,58 +39,17 @@ namespace UE::Editor::RevisionControl::Private
 		Context.AddColumn(PackageRow, MoveTemp(PathColumn));
 		Context.AddColumn(PackageRow, MoveTemp(LoadedPathColumn));
 
-		if (const UObject* PackageInstance = Package->FindAssetInPackage())
-		{
-			if (Object == PackageInstance)
-			{
-				Context.AddColumns<FTedsPrimaryPackageObjectTag>(Row);
-			}
-			// If the Row Object is a CDO, we have to compare the CDO of the primary package object instead of the object itself
-			else if (Context.HasColumn<FTypedElementClassDefaultObjectTag>())
-			{
-				const UObject* PackageObjectCDO = nullptr;
-				// Check if it was BP generated and get the Default Object of that generated class if so
-				if(const UBlueprint* PackageObjectBP = Cast<UBlueprint>(Object))
-				{
-					if (const TSubclassOf<UObject> PackageObjectBPGeneratedClass = PackageObjectBP->GeneratedClass)
-					{
-						PackageObjectCDO = PackageObjectBPGeneratedClass->GetDefaultObject();
-					}
-				}
-				else
-				{
-					PackageObjectCDO = Object->GetClass()->GetDefaultObject();
-				}
-				
-				if (Object == PackageObjectCDO)
-				{
-					Context.AddColumns<FTedsPrimaryPackageObjectTag>(Row);
-				}
-			}
-		}
-
-		Context.SetParentRow(Row, PackageRow);
+		FTypedElementPackageReference PackageBackReference;
+		PackageBackReference.Row = Row;
+		Context.AddColumn(PackageRow, MoveTemp(PackageBackReference));
 	};
-} // namespace UE::Editor::RevisionControl::Private
-
-
-void UTypedElementUObjectPackagePathFactory::RegisterHierarchies(UE::Editor::DataStorage::ICoreProvider& DataStorage)
-{
-	using namespace UE::Editor::DataStorage;
-
-	const FHierarchyRegistrationParams Params
-	{
-		.Name = PackageHierarchyName,
-	};
-
-	DataStorage.RegisterHierarchy(Params);
-}
+} // namespace UE::Editor::DataStorage
 
 void UTypedElementUObjectPackagePathFactory::RegisterQueries(UE::Editor::DataStorage::ICoreProvider& DataStorage)
 {
 	using namespace UE::Editor::DataStorage::Queries;
 	
-	UE::Editor::RevisionControl::Private::CVarSyncPackageInformation->AsVariable()->OnChangedDelegate().AddLambda(
+	UE::Editor::RevisionControl::Private::CVarAutoPopulateState->AsVariable()->OnChangedDelegate().AddLambda(
 		[this, &DataStorage](IConsoleVariable* AutoPopulate)
 		{
 			if (AutoPopulate->GetBool())
@@ -125,12 +82,11 @@ void UTypedElementUObjectPackagePathFactory::RegisterQueries(UE::Editor::DataSto
 				const UPackage* Package = ObjectInstance->GetPackage();
 				Context.RemoveColumns(Row, { FTypedElementPackageUnresolvedReference::StaticStruct() });
 
-				UE::Editor::RevisionControl::Private::ResolvePackageReference(Context, Package, ObjectInstance, Row, PackageRow);
+				UE::Editor::RevisionControl::Private::ResolvePackageReference(Context, Package, Row, PackageRow);
 			})
-		.AccessesHierarchy(UE::Editor::DataStorage::PackageHierarchyName)
 		.Compile());
 
-	if (UE::Editor::RevisionControl::Private::CVarSyncPackageInformation->GetBool())
+	if (UE::Editor::RevisionControl::Private::CVarAutoPopulateState->GetBool())
 	{
 		RegisterTryAddPackageRef(DataStorage);
 	}
@@ -138,40 +94,7 @@ void UTypedElementUObjectPackagePathFactory::RegisterQueries(UE::Editor::DataSto
 
 void UTypedElementUObjectPackagePathFactory::RegisterTryAddPackageRef(UE::Editor::DataStorage::ICoreProvider& DataStorage)
 {
-	using namespace UE::Editor::DataStorage;
 	using namespace UE::Editor::DataStorage::Queries;
-
-	struct FRegisterPackageRowCommand
-	{
-		void operator()()
-		{
-			if (CoreProvider)
-			{
-				// Check in case this row was mapped before this command was pushed 
-				RowHandle PackageRow = CoreProvider->LookupMappedRow(UE::Editor::RevisionControl::MappingDomain, FMapKey(PackagePathname));
-				if (!CoreProvider->IsRowAvailable(PackageRow))
-				{
-					const RowHandle RevisionControlTable = CoreProvider->FindTable(FName("Editor_RevisionControlTable"));
-					if (RevisionControlTable != InvalidRowHandle)
-					{
-						RowHandle Row = CoreProvider->AddRow(RevisionControlTable);
-						CoreProvider->MapRow(UE::Editor::RevisionControl::MappingDomain, FMapKey(PackagePathname), Row);
-						if (bIsMapPackage)
-						{
-							// Also map to the .uasset path
-							CoreProvider->MapRow(
-								UE::Editor::RevisionControl::MappingDomain,
-								FMapKey(FPaths::ChangeExtension(PackagePathname, FPackageName::GetAssetPackageExtension())),
-								Row);
-						}
-					}
-				}
-			}
-		}
-		FString PackagePathname;
-		ICoreProvider* CoreProvider = nullptr;
-		bool bIsMapPackage = false;
-	};
 	
 	TryAddPackageRef = DataStorage.RegisterQuery(
 		Select(
@@ -196,15 +119,9 @@ void UTypedElementUObjectPackagePathFactory::RegisterTryAddPackageRef(UE::Editor
 						return;
 					}
 
-					// Certain Package utilities don't get the .umap and will return the .uasset instead so we want to map this row to both of those
-					FString PackageExtension = FPackageName::GetAssetPackageExtension();
-
-					bool bIsMapPackage = false;
-					if (Target->ContainsMap())
-					{
-						bIsMapPackage = true;
-						PackageExtension = FPackageName::GetMapPackageExtension();
-					}
+					const FString PackageExtension = Target->ContainsMap()
+						? FPackageName::GetMapPackageExtension()
+						: FPackageName::GetAssetPackageExtension();
 
 					FString Path;
 					Target->GetPathName(nullptr, Path);
@@ -217,25 +134,16 @@ void UTypedElementUObjectPackagePathFactory::RegisterTryAddPackageRef(UE::Editor
 						if (Context.IsRowAvailable(PackageRow))
 						{
 							const UPackage* Package = ObjectInstance->GetPackage();
-							UE::Editor::RevisionControl::Private::ResolvePackageReference(Context, Package, ObjectInstance, Row, PackageRow);
+							UE::Editor::RevisionControl::Private::ResolvePackageReference(Context, Package, Row, PackageRow);
 						}
 						else
 						{
 							FTypedElementPackageUnresolvedReference UnresolvedPackageReference;
-							UnresolvedPackageReference.PathOnDisk = FullPackageFilename;
+							UnresolvedPackageReference.PathOnDisk = MoveTemp(FullPackageFilename);
 							Context.AddColumn(Row, MoveTemp(UnresolvedPackageReference));
-
-							FRegisterPackageRowCommand Command
-							{
-								.PackagePathname = FullPackageFilename,
-								.CoreProvider = GetMutableDataStorageFeature<ICoreProvider>(StorageFeatureName),
-								.bIsMapPackage = bIsMapPackage
-							};
-							Context.PushCommand(Command);
 						}
 					}
 				}
 			})
-		.AccessesHierarchy(PackageHierarchyName)
 		.Compile());
 }
diff --git a/Engine/Plugins/Experimental/EditorDataStorageFeatures/Source/TedsRevisionControl/Private/Queries/ObjectPackagePathToColumnQueries.h b/Engine/Plugins/Experimental/EditorDataStorageFeatures/Source/TedsRevisionControl/Private/Queries/ObjectPackagePathToColumnQueries.h
--- a/Engine/Plugins/Experimental/EditorDataStorageFeatures/Source/TedsRevisionControl/Private/Queries/ObjectPackagePathToColumnQueries.h
+++ b/Engine/Plugins/Experimental/EditorDataStorageFeatures/Source/TedsRevisionControl/Private/Queries/ObjectPackagePathToColumnQueries.h
@@ -20,7 +20,6 @@ class UTypedElementUObjectPackagePathFactory : public UEditorDataStorageFactory
 public:
 	~UTypedElementUObjectPackagePathFactory() override = default;
 
-	void RegisterHierarchies(UE::Editor::DataStorage::ICoreProvider& DataStorage) override;
 	void RegisterQueries(UE::Editor::DataStorage::ICoreProvider& DataStorage) override;
 
 private:
diff --git a/Engine/Source/Developer/UnsavedAssetsTracker/Source/Private/UnsavedAssetsTracker.cpp b/Engine/Source/Developer/UnsavedAssetsTracker/Source/Private/UnsavedAssetsTracker.cpp
--- a/Engine/Source/Developer/UnsavedAssetsTracker/Source/Private/UnsavedAssetsTracker.cpp
+++ b/Engine/Source/Developer/UnsavedAssetsTracker/Source/Private/UnsavedAssetsTracker.cpp
@@ -25,11 +25,6 @@
 #include "IAssetTools.h"
 #include "UObject/UObjectGlobals.h"
 
-#include "DataStorage/Features.h"
-#include "Elements/Columns/TypedElementPackageColumns.h"
-#include "Elements/Columns/TypedElementRevisionControlColumns.h"
-#include "Elements/Interfaces/TypedElementDataStorageInterface.h"
-
 #define LOCTEXT_NAMESPACE "UnsavedAssetsTracker"
 
 DEFINE_LOG_CATEGORY_STATIC(LogUnsavedAssetsTracker, Log, All);
@@ -439,8 +434,6 @@ void FUnsavedAssetsTracker::StartTrackingDirtyPackage(const UPackage* Package)
 		AbsPackageFilePathname,
 		FSourceControlFileStatusMonitor::FOnSourceControlFileStatus::CreateSP(this, &FUnsavedAssetsTracker::OnSourceControlFileStatusUpdate));
 
-	UpdateTEDSPackage(PackagePathname, true);
-
 	FUnsavedAssetsTrackerModule::Get().OnUnsavedAssetAdded.Broadcast(AbsPackageFilePathname);
 
 	UE_LOG(LogUnsavedAssetsTracker, Verbose, TEXT("Added file to the unsaved asset list: %s (%s)"), *GetHumanFriendlyAssetName(Package), *PackagePathname);
@@ -458,8 +451,6 @@ void FUnsavedAssetsTracker::StopTrackingDirtyPackage(const FString& PackagePathn
 		WarningPackagePathnames.Remove(PackagePathname);
 		UnsavedAbsFilePathames.Remove(AbsPackageFilePathname);
 
-		UpdateTEDSPackage(PackagePathname, false);
-
 		ISourceControlModule::Get().GetSourceControlFileStatusMonitor().StopMonitoringFile(reinterpret_cast<uintptr_t>(this), AbsPackageFilePathname);
 		FUnsavedAssetsTrackerModule::Get().OnUnsavedAssetRemoved.Broadcast(AbsPackageFilePathname);
 
@@ -472,39 +463,6 @@ void FUnsavedAssetsTracker::StopTrackingDirtyPackage(const FString& PackagePathn
 	}
 }
 
-void FUnsavedAssetsTracker::UpdateTEDSPackage(const FString& PackagePathname, const bool bIsDirty) const
-{
-	using namespace UE::Editor::DataStorage;
-
-	if (ICoreProvider* Storage = GetMutableDataStorageFeature<ICoreProvider>(StorageFeatureName))
-	{
-		FString FullPackagePath = FPackageName::LongPackageNameToFilename(PackagePathname, FPackageName::GetAssetPackageExtension());
-		FullPackagePath = FPaths::ConvertRelativePathToFull(FullPackagePath);
-		const RowHandle PackageRow = Storage->LookupMappedRow(UE::Editor::RevisionControl::MappingDomain, FMapKey(FullPackagePath));
-		if (Storage->IsRowAvailable(PackageRow))
-		{
-			TArray<RowHandle> ChildrenRowHandles;
-			auto GetChildRows = [&ChildrenRowHandles](const ICoreProvider& Context, RowHandle, RowHandle TargetRow)
-			{
-				ChildrenRowHandles.Add(TargetRow);
-			};
-			Storage->WalkDepthFirst(Storage->FindHierarchyByName(PackageHierarchyName), PackageRow, GetChildRows);
-
-			for (RowHandle ChildRow : ChildrenRowHandles)
-			{
-				if (bIsDirty)
-				{
-					Storage->AddColumns<FTedsPackageDirtyTag>(ChildRow);
-				}
-				else
-				{
-					Storage->RemoveColumn<FTedsPackageDirtyTag>(ChildRow);
-				}
-			}
-		}
-	}
-}
-
 void FUnsavedAssetsTracker::OnSourceControlFileStatusUpdate(const FString& AbsFilePathname, const ISourceControlState* State)
 {
 	auto DiscardWarning = [this](const FString& InPackagePathname, FStatus& InStatus)
diff --git a/Engine/Source/Developer/UnsavedAssetsTracker/Source/Private/UnsavedAssetsTracker.h b/Engine/Source/Developer/UnsavedAssetsTracker/Source/Private/UnsavedAssetsTracker.h
--- a/Engine/Source/Developer/UnsavedAssetsTracker/Source/Private/UnsavedAssetsTracker.h
+++ b/Engine/Source/Developer/UnsavedAssetsTracker/Source/Private/UnsavedAssetsTracker.h
@@ -4,7 +4,6 @@
 
 #include "CoreMinimal.h"
 #include "Containers/Ticker.h"
-#include "DataStorage/Handles.h"
 #include "HAL/CriticalSection.h"
 #include "Templates/SharedPointer.h"
 #include "UObject/ObjectSaveContext.h"
@@ -136,9 +135,6 @@ private:
 
 	/** Whether the specified package should be tracked. */
 	bool ShouldTrackDirtyPackage(const UPackage* Package);
-	
-	/** Adds the Dirty tag to the Package and its contained object rows */
-	void UpdateTEDSPackage(const FString& PackagePathname, const bool bIsDirty) const;
 
 private:
 	/** Supports multithreaded Dirty notifications. */
diff --git a/Engine/Source/Developer/UnsavedAssetsTracker/Source/UnsavedAssetsTracker.Build.cs b/Engine/Source/Developer/UnsavedAssetsTracker/Source/UnsavedAssetsTracker.Build.cs
--- a/Engine/Source/Developer/UnsavedAssetsTracker/Source/UnsavedAssetsTracker.Build.cs
+++ b/Engine/Source/Developer/UnsavedAssetsTracker/Source/UnsavedAssetsTracker.Build.cs
@@ -14,7 +14,6 @@ public class UnsavedAssetsTracker : ModuleRules
 				"EditorSubsystem",
 				"SourceControl",
 				"UncontrolledChangelists",
-				"TypedElementFramework",
 			}
 		);
 
diff --git a/Engine/Source/Runtime/TypedElementFramework/Public/Elements/Columns/TypedElementPackageColumns.h b/Engine/Source/Runtime/TypedElementFramework/Public/Elements/Columns/TypedElementPackageColumns.h
--- a/Engine/Source/Runtime/TypedElementFramework/Public/Elements/Columns/TypedElementPackageColumns.h
+++ b/Engine/Source/Runtime/TypedElementFramework/Public/Elements/Columns/TypedElementPackageColumns.h
@@ -14,11 +14,6 @@
 
 #include "TypedElementPackageColumns.generated.h"
 
-namespace UE::Editor::DataStorage
-{
-	inline static const FName PackageHierarchyName = TEXT("PackageHierarchy");
-}
-
 /**
  * A package reference column that has not yet been resolved to reference a package.
  */
@@ -77,21 +72,3 @@ struct FTypedElementPackageLoadedPathColumn final : public FEditorDataStorageCol
 
 	FPackagePath LoadedPath;
 };
-
-/**
- * Tag to indicate that the data on the row is in a package that was marked unsaved/marked dirty
- */
-USTRUCT(meta = (DisplayName = "In Dirty Package"))
-struct FTedsPackageDirtyTag final : public FEditorDataStorageTag
-{
-	GENERATED_BODY()
-};
-
-/**
- * Tag to indicate that the data on the row is the primary object in a package
- */
-USTRUCT(meta = (DisplayName = "Primary Package Object"))
-struct FTedsPrimaryPackageObjectTag final : public FEditorDataStorageTag
-{
-	GENERATED_BODY()
-};
