[EntityFramework] Introduced WidgetDropHandlers to streamline drag&drop behavior using new teds operations workflow.

diff --git a/Engine/Plugins/Experimental/EditorDataStorageFeatures/Source/TedsDropOperations/Public/TedsDropOperations.h b/Engine/Plugins/Experimental/EditorDataStorageFeatures/Source/TedsDropOperations/Public/TedsDropOperations.h
--- a/Engine/Plugins/Experimental/EditorDataStorageFeatures/Source/TedsDropOperations/Public/TedsDropOperations.h
+++ b/Engine/Plugins/Experimental/EditorDataStorageFeatures/Source/TedsDropOperations/Public/TedsDropOperations.h
@@ -5,3 +5,4 @@
 #include "TedsDropOperationsColumns.h"
 #include "TedsDropOperationsColumnsHelpers.h"
 #include "TedsDropOperationsFactory.h"
+#include "TedsDropOperationsUtility.h"
diff --git a/Engine/Plugins/Experimental/EditorDataStorageFeatures/Source/TedsTableViewer/Private/Widgets/WidgetDropHandler.cpp b/Engine/Plugins/Experimental/EditorDataStorageFeatures/Source/TedsTableViewer/Private/Widgets/WidgetDropHandler.cpp
new file mode 100644
--- /dev/null
+++ b/Engine/Plugins/Experimental/EditorDataStorageFeatures/Source/TedsTableViewer/Private/Widgets/WidgetDropHandler.cpp
@@ -0,0 +1,329 @@
+// Copyright Epic Games, Inc. All Rights Reserved.
+
+#include "Widgets/WidgetDropHandler.h"
+
+#include "Containers/Ticker.h"
+#include "DataStorage/Features.h"
+#include "DragAndDrop/DecoratedDragDropOp.h"
+#include "Elements/Columns/TypedElementHandleColumn.h"
+#include "Elements/Framework/TypedElementSelectionSet.h"
+#include "Elements/Interfaces/TypedElementDataStorageInterface.h"
+#include "Input/DragAndDrop.h"
+#include "ScopedTransaction.h"
+#include "TedsDropOperations.h"
+#include "TedsOperationsColumns.h"
+
+#define LOCTEXT_NAMESPACE "WidgetDropHandler"
+
+namespace UE::Editor
+{
+
+namespace WidgetDropHandler_Private
+{
+	static void SelectRows(DataStorage::ICoreProvider& Storage, UTypedElementSelectionSet& SelectionSet, DataStorage::FRowHandleArray RowsToSelect)
+	{
+		using namespace UE::Editor::DataStorage;
+
+		if (RowsToSelect.IsEmpty())
+		{
+			return;
+		}
+		
+		// Select next tick, so that all columns have been set and updated.
+		FTSTicker::GetCoreTicker().AddTicker(FTickerDelegate::CreateLambda(
+			[Storage = &Storage, SelectionSet = TWeakObjectPtr(&SelectionSet), RowsToSelect = MoveTemp(RowsToSelect)](float)
+			{
+				if (TStrongObjectPtr<UTypedElementSelectionSet> SelectionSetPin = SelectionSet.Pin())
+				{
+					TArray<FTypedElementHandle> Elements;
+					Elements.Reserve(RowsToSelect.Num());
+					for (RowHandle Row : RowsToSelect.GetRows())
+					{
+						Compatibility::FTypedElementColumn* Column = Storage->GetColumn<Compatibility::FTypedElementColumn>(Row);
+						if (!Column || !Column->Handle)
+						{
+							continue;
+						}
+						
+						Elements.Add(Column->Handle);
+					}
+					SelectionSetPin->SetSelection(Elements, {});
+				}
+				return false;
+			}));
+	}
+}
+
+FWidgetDropHandler::FWidgetDropHandler(TWeakObjectPtr<UTypedElementSelectionSet> SelectionSet)
+	: SelectionSet(MoveTemp(SelectionSet))
+	, CurrentOperation(nullptr)
+	, bIsValidDrop(false)
+{
+}
+
+FWidgetDropHandler::~FWidgetDropHandler()
+{
+}
+
+bool FWidgetDropHandler::OnDragEnter(const FGeometry& Geometry, const FDragDropEvent& InputEvent)
+{
+	using namespace UE::Editor::DataStorage;
+	
+	if (TSharedPtr<FDragDropOperation> Operation = InputEvent.GetOperation())
+	{
+		if (ICoreProvider* Storage = GetMutableDataStorageFeature<ICoreProvider>(StorageFeatureName))
+		{
+			if (UDropOperationsFactory* DropSystem = Storage->FindFactory<UDropOperationsFactory>())
+			{
+				FDropResult Result = Start(*Storage, *DropSystem, Geometry, InputEvent, *Operation);
+				UpdateOperationDecorator(*Operation, Result);
+				return true;
+			}
+		}
+	}
+	return false;
+}
+
+FReply FWidgetDropHandler::OnDragOver(const FGeometry& Geometry, const FDragDropEvent& InputEvent)
+{
+	using namespace UE::Editor::DataStorage;
+
+	TSharedPtr<FDragDropOperation> Operation = InputEvent.GetOperation();
+	if (Operation && CurrentOperation == Operation.Get())
+	{
+		if (ICoreProvider* Storage = GetMutableDataStorageFeature<ICoreProvider>(StorageFeatureName))
+		{
+			if (UDropOperationsFactory* DropSystem = Storage->FindFactory<UDropOperationsFactory>())
+			{
+				if (TOptional<FDropResult> Result = Update(*Storage, *DropSystem, Geometry, InputEvent))
+				{
+					// Only update for changes.
+					UpdateOperationDecorator(*Operation, Result);
+				}
+			}
+		}
+		return FReply::Handled();
+	}
+	return FReply::Unhandled();
+}
+
+bool FWidgetDropHandler::OnDragLeave(const FDragDropEvent& InputEvent)
+{
+	using namespace UE::Editor::DataStorage;
+
+	TSharedPtr<FDragDropOperation> Operation = InputEvent.GetOperation();
+	if (Operation && CurrentOperation == Operation.Get())
+	{
+		if (ICoreProvider* Storage = GetMutableDataStorageFeature<ICoreProvider>(StorageFeatureName))
+		{
+			Stop(*Storage);
+			UpdateOperationDecorator(*Operation, {});
+		}
+		return true;
+	}
+	return false;
+}
+
+FReply FWidgetDropHandler::OnDrop(const FGeometry& Geometry, const FDragDropEvent& InputEvent)
+{
+	using namespace UE::Editor::DataStorage;
+
+	TSharedPtr<FDragDropOperation> Operation = InputEvent.GetOperation();
+	if (Operation && CurrentOperation == Operation.Get())
+	{
+		ICoreProvider* Storage = GetMutableDataStorageFeature<ICoreProvider>(StorageFeatureName);
+		if (Storage && bIsValidDrop)
+		{
+			Drop(*Storage);
+			UpdateOperationDecorator(*Operation, {});
+		}
+		return FReply::Handled();
+	}
+	return FReply::Unhandled();
+}
+
+FDragDropOperation* FWidgetDropHandler::GetCurrentOperation() const
+{
+	return CurrentOperation;
+}
+
+void FWidgetDropHandler::PrepareInput(DataStorage::ICoreProvider& Storage, UDropOperationsFactory& DropSystem, const FDragDropOperation& Operation)
+{
+	using namespace UE::Editor::DataStorage;
+	
+	FRowHandleArray SourceRows;
+	DragAndDrop::GetRowsFromData(SourceRows, Operation);
+	SourceRows.Remove(InvalidRowHandle);
+	
+	PrepareInput(Storage, DropSystem, SourceRows.GetRows());
+}
+
+void FWidgetDropHandler::PrepareInput(DataStorage::ICoreProvider& Storage, UDropOperationsFactory& DropSystem,
+	DataStorage::FRowHandleArrayView SourceRows)
+{
+	using namespace UE::Editor::DataStorage;
+	checkf(DropEntries.IsEmpty(), TEXT("Expected previous operation to be cleared."));
+	
+	TableHandle Table = DropSystem.GetInputTable();
+	if (Table == InvalidTableHandle)
+	{
+		return;
+	}
+
+	const RowHandle* SourceRowPtr = SourceRows.begin();
+	DropEntries.Reserve(SourceRows.Num());
+	Storage.BatchAddRow(Table, SourceRows.Num(), [this, &Storage, &SourceRowPtr](RowHandle InputRow)
+		{
+			if (InputRow != InvalidRowHandle)
+			{
+				Storage.AddColumn(InputRow, DragAndDrop::FSourceColumn{ .Value = *SourceRowPtr++ });
+				DropEntries.Add(FDropEntry{ .InputRow = InputRow });
+			}
+		});
+}
+
+FWidgetDropHandler::FDropResult FWidgetDropHandler::ResetOperations(DataStorage::ICoreProvider& Storage, UDropOperationsFactory& DropSystem)
+{
+	using namespace UE::Editor::DataStorage;
+
+	FDropResult Result;
+	for (FDropEntry& Entry : DropEntries)
+	{
+		Entry.OperationRow = DropSystem.Test(Entry.InputRow);
+		Result = FDropResult{
+			.bCanDrop = Entry.OperationRow != InvalidRowHandle,
+			.Tooltip = DragAndDrop::Helpers::GetDescriptionValue(Storage, Entry.InputRow),
+		};
+
+		if (!Result.bCanDrop)
+		{
+			break;
+		}
+	}
+
+	bIsValidDrop = Result.bCanDrop;
+	if (!Result.bCanDrop && Result.Tooltip.IsEmpty())
+	{
+		Result.Tooltip = LOCTEXT("Drop_Unsupported", "Drop operation is unsupported.");
+	}
+	return Result;
+}
+
+FWidgetDropHandler::EUpdateInput FWidgetDropHandler::UpdateInput(DataStorage::ICoreProvider& Storage, const FGeometry& Geometry, const FDragDropEvent& InputEvent)
+{
+	return EUpdateInput::NoChanges;
+}
+
+FWidgetDropHandler::FDropResult FWidgetDropHandler::Start(DataStorage::ICoreProvider& Storage, UDropOperationsFactory& DropSystem,
+	const FGeometry& Geometry, const FDragDropEvent& InputEvent, FDragDropOperation& Operation)
+{
+	using namespace UE::Editor::DataStorage;
+
+	if (CurrentOperation)
+	{
+		Stop(Storage);
+	}
+	
+	CurrentOperation = &Operation;
+	
+	PrepareInput(Storage, DropSystem, Operation);
+	(void)UpdateInput(Storage, Geometry, InputEvent);
+	FDropResult Result = ResetOperations(Storage, DropSystem);
+	if (!bIsValidDrop)
+	{
+		Stop(Storage);
+	}
+	return Result;
+}
+
+TOptional<FWidgetDropHandler::FDropResult> FWidgetDropHandler::Update(DataStorage::ICoreProvider& Storage, UDropOperationsFactory& DropSystem,
+	const FGeometry& Geometry, const FDragDropEvent& InputEvent)
+{
+	using namespace UE::Editor::DataStorage;
+	check(CurrentOperation != nullptr);
+
+	EUpdateInput Result = UpdateInput(Storage, Geometry, InputEvent);
+	if (Result == EUpdateInput::ResetOperations)
+	{
+		return ResetOperations(Storage, DropSystem);
+	}	
+	return {};
+}
+	
+void FWidgetDropHandler::Stop(DataStorage::ICoreProvider& Storage)
+{
+	using namespace UE::Editor::DataStorage;
+	check(CurrentOperation != nullptr);
+
+	CurrentOperation = nullptr;
+	bIsValidDrop = false;
+
+	FRowHandleArray InputRows;
+	InputRows.Reserve(DropEntries.Num());
+	for (FDropEntry& Entry : DropEntries)
+	{
+		InputRows.Add(Entry.InputRow);
+	}	
+	Storage.BatchRemoveRows(InputRows.GetRows());
+	DropEntries.Empty();
+}
+
+void FWidgetDropHandler::Drop(DataStorage::ICoreProvider& Storage)
+{
+	using namespace UE::Editor::DataStorage;
+	check(bIsValidDrop);
+
+	FRowHandleArray RowsToSelect;
+	{
+		FScopedTransaction Transaction(LOCTEXT("Drop_Transaction", "Drop on viewport."));
+		for (const FDropEntry& Entry : DropEntries)
+		{
+			const Operations::FApplyColumn* Column = Storage.GetColumn<Operations::FApplyColumn>(Entry.OperationRow);
+			if (Column && Column->Callback)
+			{
+				if (TOptional<Operations::FResult> Result = Column->Callback(Storage, Entry.InputRow))
+				{
+					RowsToSelect.Append(Result->Created.GetRows());
+				}
+			}
+		}
+	}
+	
+	Stop(Storage);
+
+	if (UTypedElementSelectionSet* SelectionSetPtr = SelectionSet.Get())
+	{
+		WidgetDropHandler_Private::SelectRows(Storage, *SelectionSetPtr, MoveTemp(RowsToSelect));
+	}
+}
+
+void FWidgetDropHandler::UpdateOperationDecorator(FDragDropOperation& Operation, TOptional<FDropResult> Result) const
+{
+	Operation.SetDecoratorVisibility(true);
+	if (Operation.IsOfType<FDecoratedDragDropOp>())
+	{
+		static constexpr TCHAR BrushName_Ok[] = TEXT("Graph.ConnectorFeedback.OK");
+		static constexpr TCHAR BrushName_Error[] = TEXT("Graph.ConnectorFeedback.Error");
+		FDecoratedDragDropOp& DecoratedOp = static_cast<FDecoratedDragDropOp&>(Operation);
+		if (Result)
+		{
+			DecoratedOp.SetToolTip(Result->Tooltip, FAppStyle::GetBrush(Result->bCanDrop ? BrushName_Ok : BrushName_Error));
+		}
+		else
+		{
+			DecoratedOp.ResetToDefaultToolTip();
+		}
+	}
+	else if (Result)
+	{
+		Operation.SetCursorOverride(Result->bCanDrop ? EMouseCursor::Default : EMouseCursor::SlashedCircle);
+	}
+	else
+	{
+		Operation.SetCursorOverride({});
+	}
+}
+	
+}
+
+#undef LOCTEXT_NAMESPACE
diff --git a/Engine/Plugins/Experimental/EditorDataStorageFeatures/Source/TedsTableViewer/Public/Widgets/WidgetDropHandler.h b/Engine/Plugins/Experimental/EditorDataStorageFeatures/Source/TedsTableViewer/Public/Widgets/WidgetDropHandler.h
new file mode 100644
--- /dev/null
+++ b/Engine/Plugins/Experimental/EditorDataStorageFeatures/Source/TedsTableViewer/Public/Widgets/WidgetDropHandler.h
@@ -0,0 +1,83 @@
+// Copyright Epic Games, Inc. All Rights Reserved.
+
+#pragma once
+
+#include "DataStorage/Handles.h"
+#include "DragAndDrop/IWidgetDropHandler.h"
+#include "Elements/Framework/TypedElementRowHandleArrayView.h"
+#include "Misc/Optional.h"
+
+class UTypedElementSelectionSet;
+class UDropOperationsFactory;
+namespace UE::Editor::DataStorage
+{
+class ICoreProvider;
+}
+
+#define UE_API TEDSTABLEVIEWER_API
+
+namespace UE::Editor
+{
+
+class FWidgetDropHandler : public IWidgetDropHandler
+{
+public:
+	struct FDropResult
+	{
+		bool  bCanDrop = false;
+		FText Tooltip  = FText();
+	};
+
+protected:
+	struct FDropEntry
+	{
+		DataStorage::RowHandle InputRow;
+		DataStorage::RowHandle OperationRow;
+	};
+
+	enum class EUpdateInput
+	{
+		NoChanges,
+		ResetOperations,
+	};
+
+public:
+	UE_API FWidgetDropHandler(TWeakObjectPtr<UTypedElementSelectionSet> SelectionSet);
+	UE_API virtual ~FWidgetDropHandler() override;
+
+	UE_API virtual bool                OnDragEnter(const FGeometry& Geometry, const FDragDropEvent& InputEvent) override;
+	UE_API virtual FReply              OnDragOver(const FGeometry& Geometry, const FDragDropEvent& InputEvent) override;
+	UE_API virtual bool                OnDragLeave(const FDragDropEvent& InputEvent) override;
+	UE_API virtual FReply              OnDrop(const FGeometry& Geometry, const FDragDropEvent& InputEvent) override;
+	UE_API virtual FDragDropOperation* GetCurrentOperation() const override;
+
+protected:
+	/** Creates input rows for the given drag&drop input event. */
+	UE_API virtual void PrepareInput(DataStorage::ICoreProvider& Storage, UDropOperationsFactory& DropSystem, const FDragDropOperation& Operation);
+	/** Creates input rows for the given drag&drop input event. */
+	UE_API virtual void PrepareInput(DataStorage::ICoreProvider& Storage, UDropOperationsFactory& DropSystem, DataStorage::FRowHandleArrayView SourceRows);
+	/** Finds suitable operations for the given input rows and sets 'bIsValidDrop'. */
+	UE_API virtual FDropResult ResetOperations(DataStorage::ICoreProvider& Storage, UDropOperationsFactory& DropSystem);
+	/** Updates all input rows for the updated input parameters. May also set status vars like 'bIsValidDrop'. */
+	UE_API virtual EUpdateInput UpdateInput(DataStorage::ICoreProvider& Storage, const FGeometry& Geometry, const FDragDropEvent& InputEvent);
+
+	UE_API virtual FDropResult Start(DataStorage::ICoreProvider& Storage, UDropOperationsFactory& DropSystem, const FGeometry& Geometry,
+		const FDragDropEvent& InputEvent, FDragDropOperation& Operation);
+	UE_API virtual TOptional<FDropResult> Update(DataStorage::ICoreProvider& Storage, UDropOperationsFactory& DropSystem, const FGeometry& Geometry,
+		const FDragDropEvent& InputEvent);
+	UE_API virtual void Stop(DataStorage::ICoreProvider& Storage);
+	UE_API virtual void Drop(DataStorage::ICoreProvider& Storage);
+
+	UE_API virtual void UpdateOperationDecorator(FDragDropOperation& Operation, TOptional<FDropResult> Result) const;
+
+protected:
+	TWeakObjectPtr<UTypedElementSelectionSet> SelectionSet;
+
+	FDragDropOperation* CurrentOperation;
+	TArray<FDropEntry>  DropEntries;
+	bool                bIsValidDrop;
+};
+
+}
+
+#undef UE_API
\ No newline at end of file
diff --git a/Engine/Plugins/Experimental/EditorDataStorageFeatures/Source/TedsTableViewer/TedsTableViewer.Build.cs b/Engine/Plugins/Experimental/EditorDataStorageFeatures/Source/TedsTableViewer/TedsTableViewer.Build.cs
--- a/Engine/Plugins/Experimental/EditorDataStorageFeatures/Source/TedsTableViewer/TedsTableViewer.Build.cs
+++ b/Engine/Plugins/Experimental/EditorDataStorageFeatures/Source/TedsTableViewer/TedsTableViewer.Build.cs
@@ -18,7 +18,8 @@ public class TedsTableViewer : ModuleRules
 				new string[]
 				{
 					"TypedElementFramework",
-					"TedsQueryStack"
+					"TedsQueryStack",
+					"ToolWidgets"
 				});
 
 			PrivateDependencyModuleNames.AddRange(
@@ -33,9 +34,10 @@ public class TedsTableViewer : ModuleRules
 					"SlateCore",
 					"WorkspaceMenuStructure",
 					"ToolMenus",
-					"ToolWidgets",
 					"TedsOperations",
-					"TedsDropOperations"
+					"TedsDropOperations",
+					"UnrealEd",
+					"TypedElementRuntime"
 				});
 		}
 	}
diff --git a/Engine/Source/Developer/ToolWidgets/Public/DragAndDrop/IWidgetDropHandler.h b/Engine/Source/Developer/ToolWidgets/Public/DragAndDrop/IWidgetDropHandler.h
new file mode 100644
--- /dev/null
+++ b/Engine/Source/Developer/ToolWidgets/Public/DragAndDrop/IWidgetDropHandler.h
@@ -0,0 +1,30 @@
+// Copyright Epic Games, Inc. All Rights Reserved.
+
+#pragma once
+
+#include "Input/Reply.h"
+
+struct FGeometry;
+class  FDragDropEvent;
+class  FDragDropOperation;
+
+namespace UE::Editor
+{
+
+/** Interface that can be used to encapsulate drag&drop behavior. */
+UE_EXPERIMENTAL(5.7, "Experimental introduction to add common behavior to specific widgets. Prone to change.")
+class IWidgetDropHandler
+{
+public:
+	virtual ~IWidgetDropHandler() = default;
+
+	virtual bool   OnDragEnter(const FGeometry& Geometry, const FDragDropEvent& InputEvent) = 0;
+	virtual FReply OnDragOver(const FGeometry& Geometry, const FDragDropEvent& InputEvent) = 0;
+	virtual bool   OnDragLeave(const FDragDropEvent& InputEvent) = 0;
+	virtual FReply OnDrop(const FGeometry& Geometry, const FDragDropEvent& InputEvent) = 0;
+
+	virtual FDragDropOperation* GetCurrentOperation() const = 0;
+	bool IsInOperation() const { return GetCurrentOperation() != nullptr; }
+};
+
+}
