OutlinerTextWidget: Fix incorrect logic causing no widget to be created if the editable metadata was not present

diff --git a/Engine/Plugins/Experimental/EditorDataStorageFeatures/Source/TedsOutliner/Private/Widgets/OutlinerTextWidget.cpp b/Engine/Plugins/Experimental/EditorDataStorageFeatures/Source/TedsOutliner/Private/Widgets/OutlinerTextWidget.cpp
--- a/Engine/Plugins/Experimental/EditorDataStorageFeatures/Source/TedsOutliner/Private/Widgets/OutlinerTextWidget.cpp
+++ b/Engine/Plugins/Experimental/EditorDataStorageFeatures/Source/TedsOutliner/Private/Widgets/OutlinerTextWidget.cpp
@@ -163,22 +163,27 @@ TSharedPtr<SWidget> FOutlinerTextWidgetConstructor::CreateWidget(UE::Editor::Dat
 
 	TSharedPtr<SWidget> Result = SNullWidget::NullWidget;
 
+	bool bEditable = false;
+	
 	FMetaDataEntryView MetaDataEntryView = Arguments.FindForColumn<FTypedElementLabelColumn>(IsEditableName);
 
 	if (MetaDataEntryView.IsSet())
 	{
-		const bool* IsEditable = MetaDataEntryView.TryGetExact<bool>();
-
-		if (IsEditable && *IsEditable)
+		if (const bool* EditableMetaDataPtr = MetaDataEntryView.TryGetExact<bool>())
 		{
-			Result = CreateEditableWidget(DataStorage, DataStorageUi, TargetRow, WidgetRow, Arguments);
-		}
-		else
-		{
-			Result = CreateNonEditableWidget(DataStorage, DataStorageUi, TargetRow, WidgetRow, Arguments);
+			bEditable = *EditableMetaDataPtr;
 		}
 	}
 	
+	if (bEditable)
+	{
+		Result = CreateEditableWidget(DataStorage, DataStorageUi, TargetRow, WidgetRow, Arguments);
+	}
+	else
+	{
+		Result = CreateNonEditableWidget(DataStorage, DataStorageUi, TargetRow, WidgetRow, Arguments);
+	}
+	
 	
 	return Result;
 }
