Unshelved from pending changelist '47735984':

[TEDS Hierarchy Viewer/Model] Fix the executor holding onto the old query stack and don't have the query stack update on initialize for the current implementation of ExpandAll

diff --git a/Engine/Plugins/Experimental/EditorDataStorageFeatures/Source/TedsTableViewer/Private/TedsTableViewerModel.cpp b/Engine/Plugins/Experimental/EditorDataStorageFeatures/Source/TedsTableViewer/Private/TedsTableViewerModel.cpp
--- a/Engine/Plugins/Experimental/EditorDataStorageFeatures/Source/TedsTableViewer/Private/TedsTableViewerModel.cpp
+++ b/Engine/Plugins/Experimental/EditorDataStorageFeatures/Source/TedsTableViewer/Private/TedsTableViewerModel.cpp
@@ -26,7 +26,7 @@ namespace UE::Editor::DataStorage
 		Storage = GetMutableDataStorageFeature<ICoreProvider>(StorageFeatureName);
 		StorageUi = GetMutableDataStorageFeature<IUiProvider>(UiFeatureName);
 		
-		SetQueryStack(Args.RowQueryStack);
+		SetQueryStack(Args.RowQueryStack, false);
 
 		TickerHandle = FTSTicker::GetCoreTicker().AddTicker(FTickerDelegate::CreateRaw(this, &FTedsTableViewerModel::Tick), 0);
 
@@ -160,7 +160,7 @@ namespace UE::Editor::DataStorage
 		return OnModelChanged;
 	}
 
-	void FTedsTableViewerModel::SetQueryStack(const TSharedPtr<QueryStack::IRowNode>& InRowQueryStack)
+	void FTedsTableViewerModel::SetQueryStack(const TSharedPtr<QueryStack::IRowNode>& InRowQueryStack, bool bUpdateExecutor /*= true*/)
 	{
 		checkf(InRowQueryStack, TEXT("The TEDS Table Viewer requires a valid query stack to provide it with rows to process."));
 		// If the sorting nodes have been made and are just being reset to a new query stack, persist the sorting state
@@ -172,8 +172,12 @@ namespace UE::Editor::DataStorage
 
 		static int InstanceCounter = 0;
 		QueryStackExecutor = QueryStack::FExplicitUpdateExecutor(FName("TableViewerModel", InstanceCounter++), RowQueryStack);
+
+		if (bUpdateExecutor)
+		{
+			QueryStackExecutor.Update();
+		}
 		
-		QueryStackExecutor.Update();
 		Refresh();
 	}
 
diff --git a/Engine/Plugins/Experimental/EditorDataStorageFeatures/Source/TedsTableViewer/Private/Widgets/STedsHierarchyViewer.cpp b/Engine/Plugins/Experimental/EditorDataStorageFeatures/Source/TedsTableViewer/Private/Widgets/STedsHierarchyViewer.cpp
--- a/Engine/Plugins/Experimental/EditorDataStorageFeatures/Source/TedsTableViewer/Private/Widgets/STedsHierarchyViewer.cpp
+++ b/Engine/Plugins/Experimental/EditorDataStorageFeatures/Source/TedsTableViewer/Private/Widgets/STedsHierarchyViewer.cpp
@@ -234,6 +234,8 @@ namespace UE::Editor::DataStorage
 	{
 		Model->SetQueryStack(InRowQueryStack);
 		TopLevelRowsNode = MakeShared<QueryStack::FTopLevelRowsNode>(Model->GetDataStorageInterface(), HierarchyInterface, Model->GetRowNode());
+		static int InstanceCounter = 0;
+		QueryStackExecutor = QueryStack::FExplicitUpdateExecutor(FName("HierarchyViewer", InstanceCounter++), TopLevelRowsNode);
 	}
 
 	void SHierarchyViewer::SetColumns(const TArray<TWeakObjectPtr<const UScriptStruct>>& InColumns)
diff --git a/Engine/Plugins/Experimental/EditorDataStorageFeatures/Source/TedsTableViewer/Public/TedsTableViewerModel.h b/Engine/Plugins/Experimental/EditorDataStorageFeatures/Source/TedsTableViewer/Public/TedsTableViewerModel.h
--- a/Engine/Plugins/Experimental/EditorDataStorageFeatures/Source/TedsTableViewer/Public/TedsTableViewerModel.h
+++ b/Engine/Plugins/Experimental/EditorDataStorageFeatures/Source/TedsTableViewer/Public/TedsTableViewerModel.h
@@ -75,7 +75,7 @@ namespace UE::Editor::DataStorage
 		TEDSTABLEVIEWER_API FOnModelChanged& GetOnModelChanged();
 
 		// Clear the current QueryStack being displayed, set it to the given node, and recreate the sorting nodes 
-		TEDSTABLEVIEWER_API void SetQueryStack(const TSharedPtr<QueryStack::IRowNode>& InRowQueryStack);
+		TEDSTABLEVIEWER_API void SetQueryStack(const TSharedPtr<QueryStack::IRowNode>& InRowQueryStack, bool bUpdateExecutor = true);
 		
 		// Clear the current list of columns being displayed and set it to the given list
 		TEDSTABLEVIEWER_API void SetColumns(const TArray<TWeakObjectPtr<const UScriptStruct>>& InColumns);
