Teds Table Viewer:
Use new widget drop handler
Add support for treating drop onto empty area as drop onto root

diff --git a/Engine/Plugins/Experimental/EditorDataStorageFeatures/Source/TedsTableViewer/Private/Widgets/STedsHierarchyViewer.cpp b/Engine/Plugins/Experimental/EditorDataStorageFeatures/Source/TedsTableViewer/Private/Widgets/STedsHierarchyViewer.cpp
--- a/Engine/Plugins/Experimental/EditorDataStorageFeatures/Source/TedsTableViewer/Private/Widgets/STedsHierarchyViewer.cpp
+++ b/Engine/Plugins/Experimental/EditorDataStorageFeatures/Source/TedsTableViewer/Private/Widgets/STedsHierarchyViewer.cpp
@@ -2,6 +2,11 @@
 
 #include "Widgets/STedsHierarchyViewer.h"
 
+#include "TableViewerWidgetDropHandler.h"
+#include "TedsDragDropOp.h"
+#include "TedsDropOperationsColumnsHelpers.h"
+#include "TedsDropOperationsFactory.h"
+#include "TedsDropOperationsUtility.h"
 #include "TedsTableViewerColumn.h"
 #include "TedsTableViewerWidgetColumns.h"
 #include "Columns/SlateDelegateColumns.h"
@@ -94,6 +99,8 @@ namespace UE::Editor::DataStorage
 
 		// Whenever the model changes, refresh the list to update the UI
 		Model->GetOnModelChanged().AddRaw(this, &SHierarchyViewer::OnModelChanged);
+
+		DropHandler = MakePimpl<Widgets::FTableViewerWidgetDropHandler>(SharedThis(this));
 	}
 
 	void SHierarchyViewer::OnModelChanged()
@@ -115,6 +122,18 @@ namespace UE::Editor::DataStorage
 			}
 		}
 
+		if(ICoreProvider* DataStorage = Model->GetDataStorageInterface())
+		{
+			if (!Rows.IsEmpty())
+			{
+				DataStorage->AddColumn(GetWidgetRowHandle(), FHierarchyViewerRootColumn{.Root = Rows[0]});
+			}
+			else
+			{
+				DataStorage->RemoveColumn<FHierarchyViewerRootColumn>(GetWidgetRowHandle());
+			}
+		}
+
 		if (bExpansionStateChangeRequested.IsSet())
 		{
 			if (TSharedPtr<QueryStack::IRowNode> AllRowsNode = Model->GetRowNode())
@@ -137,11 +156,37 @@ namespace UE::Editor::DataStorage
 		}
 	}
 
+	void SHierarchyViewer::OnDragEnter(const FGeometry& MyGeometry, const FDragDropEvent& DragDropEvent)
+	{
+		if (DropHandler)
+		{
+			DropHandler->OnDragEnter(MyGeometry, DragDropEvent);
+		}
+	}
+
+	FReply SHierarchyViewer::OnDragOver(const FGeometry& MyGeometry, const FDragDropEvent& DragDropEvent)
+	{
+		return DropHandler ? DropHandler->OnDragOver(MyGeometry, DragDropEvent) : FReply::Unhandled();
+	}
+
+	void SHierarchyViewer::OnDragLeave(const FDragDropEvent& DragDropEvent)
+	{
+		if (DropHandler)
+		{
+			DropHandler->OnDragLeave(DragDropEvent);
+		}
+	}
+
+	FReply SHierarchyViewer::OnDrop(const FGeometry& MyGeometry, const FDragDropEvent& DragDropEvent)
+	{
+		return DropHandler ? DropHandler->OnDrop(MyGeometry, DragDropEvent) : FReply::Unhandled();
+	}
+
 	void SHierarchyViewer::AddWidgetColumns()
 	{
 		if(ICoreProvider* DataStorage = Model->GetDataStorageInterface())
 		{
-			const RowHandle WidgetRowHandle = TedsWidget->GetRowHandle();
+			const RowHandle WidgetRowHandle = GetWidgetRowHandle();
 		
 			if(DataStorage->IsRowAvailable(WidgetRowHandle))
 			{
@@ -157,14 +202,19 @@ namespace UE::Editor::DataStorage
 
 	TableViewerItemPtr SHierarchyViewer::GetParentRow(TableViewerItemPtr InItem) const
 	{
-		// To work around unreliable query stack update timings, in case the widget queries parents before the top level rows node has updated
-		// but after the hierarchies themselves have updated since HierarchyInterface directly queries TEDS. 
-		if (TopLevelRows.Contains(InItem) || !HierarchyInterface)
+		if (ICoreProvider* Storage = Model->GetDataStorageInterface())
 		{
-			return FTedsRowHandle(InvalidRowHandle);
-		}
+			// To work around unreliable query stack update timings, in case the widget queries parents before the top level rows node has updated
+			// but after the hierarchies themselves have updated since HierarchyInterface directly queries TEDS. 
+			if (TopLevelRows.Contains(InItem) || !HierarchyInterface)
+			{
+				return FTedsRowHandle(InvalidRowHandle);
+			}
 		
-		return FTedsRowHandle(HierarchyInterface->GetParent(*Model->GetDataStorageInterface(), InItem));
+			return FTedsRowHandle(HierarchyInterface->GetParent(*Storage, InItem));
+		}
+
+		return FTedsRowHandle(InvalidRowHandle);
 	}
 
 	void SHierarchyViewer::CreateInternalWidget()
diff --git a/Engine/Plugins/Experimental/EditorDataStorageFeatures/Source/TedsTableViewer/Private/Widgets/STedsTableViewer.cpp b/Engine/Plugins/Experimental/EditorDataStorageFeatures/Source/TedsTableViewer/Private/Widgets/STedsTableViewer.cpp
--- a/Engine/Plugins/Experimental/EditorDataStorageFeatures/Source/TedsTableViewer/Private/Widgets/STedsTableViewer.cpp
+++ b/Engine/Plugins/Experimental/EditorDataStorageFeatures/Source/TedsTableViewer/Private/Widgets/STedsTableViewer.cpp
@@ -2,6 +2,7 @@
 
 #include "Widgets/STedsTableViewer.h"
 
+#include "TableViewerWidgetDropHandler.h"
 #include "TedsTableViewerColumn.h"
 #include "TedsTableViewerWidgetColumns.h"
 #include "Columns/SlateDelegateColumns.h"
@@ -26,9 +27,15 @@ namespace UE::Editor::DataStorage
 
 		IUiProvider::FPurposeID CellWidgetPurpose = InArgs._CellWidgetPurpose;
 
+		IUiProvider* StorageUi = GetMutableDataStorageFeature<IUiProvider>(UiFeatureName);
+
+		if (!ensureMsgf(StorageUi, TEXT("Cannot use STedsTableViewer without TEDS UI being initialized")))
+		{
+			return;
+		}
+
 		if (!CellWidgetPurpose.IsSet())
 		{
-			IUiProvider* StorageUi = GetMutableDataStorageFeature<IUiProvider>(UiFeatureName);
 			CellWidgetPurpose = StorageUi->GetGeneralWidgetPurposeID();
 		}
 
@@ -50,18 +57,13 @@ namespace UE::Editor::DataStorage
 		
 		HeaderRowWidget = SNew( SHeaderRow )
 							.CanSelectGeneratedColumn(InArgs._CanSelectGeneratedColumn);
-		
-		IUiProvider* StorageUi = Model->GetDataStorageUiProvider();
 
-		if (ensure(StorageUi))
-		{
-			TedsWidget = StorageUi->CreateContainerTedsWidget(InvalidRowHandle);
+		TedsWidget = StorageUi->CreateContainerTedsWidget(InvalidRowHandle);
 			
-			ChildSlot
-			[
-				TedsWidget->AsWidget()
-			];
-		}
+		ChildSlot
+		[
+			TedsWidget->AsWidget()
+		];
 		
 		AddWidgetColumns();
 
@@ -92,6 +94,34 @@ namespace UE::Editor::DataStorage
 			ListView->RequestListRefresh();
 			CreateInternalWidget();
 		});
+
+		DropHandler = MakePimpl<Widgets::FTableViewerWidgetDropHandler>(SharedThis(this));
+	}
+
+	void STedsTableViewer::OnDragEnter(const FGeometry& MyGeometry, const FDragDropEvent& DragDropEvent)
+	{
+		if (DropHandler)
+		{
+			DropHandler->OnDragEnter(MyGeometry, DragDropEvent);
+		}
+	}
+
+	FReply STedsTableViewer::OnDragOver(const FGeometry& MyGeometry, const FDragDropEvent& DragDropEvent)
+	{
+		return DropHandler ? DropHandler->OnDragOver(MyGeometry, DragDropEvent) : FReply::Unhandled();
+	}
+
+	void STedsTableViewer::OnDragLeave(const FDragDropEvent& DragDropEvent)
+	{
+		if (DropHandler)
+		{
+			DropHandler->OnDragLeave(DragDropEvent);
+		}
+	}
+
+	FReply STedsTableViewer::OnDrop(const FGeometry& MyGeometry, const FDragDropEvent& DragDropEvent)
+	{
+		return DropHandler ? DropHandler->OnDrop(MyGeometry, DragDropEvent) : FReply::Unhandled();
 	}
 
 	void STedsTableViewer::AddWidgetColumns()
diff --git a/Engine/Plugins/Experimental/EditorDataStorageFeatures/Source/TedsTableViewer/Private/Widgets/STedsTableViewerRow.cpp b/Engine/Plugins/Experimental/EditorDataStorageFeatures/Source/TedsTableViewer/Private/Widgets/STedsTableViewerRow.cpp
--- a/Engine/Plugins/Experimental/EditorDataStorageFeatures/Source/TedsTableViewer/Private/Widgets/STedsTableViewerRow.cpp
+++ b/Engine/Plugins/Experimental/EditorDataStorageFeatures/Source/TedsTableViewer/Private/Widgets/STedsTableViewerRow.cpp
@@ -19,44 +19,6 @@
 
 namespace UE::Editor::DataStorage
 {
-	namespace TedsTableViewerRow_Private
-	{
-		struct FDropResult
-		{
-			bool  bCanDrop = false;
-			FText Hint     = FText();
-		};
-
-		static void UpdateOperationDecorator(FDragDropOperation& Operation, TOptional<FDropResult> Result)
-		{
-			Operation.SetDecoratorVisibility(true);
-			if (Operation.IsOfType<FDecoratedDragDropOp>())
-			{
-				static constexpr TCHAR BrushName_Ok[] = TEXT("Graph.ConnectorFeedback.OK");
-				static constexpr TCHAR BrushName_Error[] = TEXT("Graph.ConnectorFeedback.Error");
-
-				FDecoratedDragDropOp& DecoratedOp = static_cast<FDecoratedDragDropOp&>(Operation);
-				if (Result)
-				{
-					DecoratedOp.SetToolTip(Result->Hint, FAppStyle::GetBrush(Result->bCanDrop ? BrushName_Ok : BrushName_Error));
-				}
-				else
-				{
-					DecoratedOp.ResetToDefaultToolTip();
-				}
-
-			}
-			else if (Result)
-			{
-				Operation.SetCursorOverride(Result->bCanDrop ? EMouseCursor::Default : EMouseCursor::SlashedCircle);
-			}
-			else
-			{
-				Operation.SetCursorOverride({});
-			}
-		}
-	}
-
 	void STedsTableViewerRow::Construct(const FArguments& InArgs, TSharedRef<STableViewBase> OwnerTableView, TSharedRef<FTedsTableViewerModel> InTableViewerModel, TWeakPtr<const ITableViewer> InTableViewer)
 	{
 		Item = InArgs._Item;
@@ -69,7 +31,6 @@ namespace UE::Editor::DataStorage
 		.Padding(InArgs._Padding)
 		.OnDragDetected(this, &STedsTableViewerRow::HandleOnDragDetected)
 		.OnDragEnter(this, &STedsTableViewerRow::HandleOnDragEnter)
-		.OnAcceptDrop(this, &STedsTableViewerRow::HandleOnAcceptDrop)
 		.OnDragLeave(this, &STedsTableViewerRow::HandleOnDragLeave)
 		.Style(&FAppStyle::Get().GetWidgetStyle<FTableRowStyle>("SceneOutliner.TableViewRow"));
 
@@ -116,119 +77,54 @@ namespace UE::Editor::DataStorage
 		return ItemHeight.IsSet() ? ItemHeight.Get() : FOptionalSize();
 	}
 
-	void STedsTableViewerRow::HandleOnDragEnter(const FDragDropEvent& DragDropEvent)
+	FReply STedsTableViewerRow::HandleOnDragDetected(const FGeometry& MyGeometry, const FPointerEvent& MouseEvent)
 	{
-		ensureMsgf(ValidDropOperation == nullptr, TEXT("Expected last drop operation to be cleared."));
-
-		if (TSharedPtr<FDragDropOperation> Operation = DragDropEvent.GetOperation())
-		{
-			TedsTableViewerRow_Private::FDropResult Result;
-		if (ICoreProvider* Storage = TableViewerModel->GetDataStorageInterface())
+		if (MouseEvent.IsMouseButtonDown(EKeys::LeftMouseButton))
 		{
-			if (UDropOperationsFactory* DropFactory = Storage->FindFactory<UDropOperationsFactory>())
-			{
-					RowHandle TargetRow = Item;
-
-					FRowHandleArray DraggedRows;
-					UE::Editor::DataStorage::DragAndDrop::GetRowsFromData(DraggedRows, *Operation);
-					for (RowHandle SourceRow : DraggedRows.GetRows())
-					{
-						FScopedRow InputRow = DropFactory->CreateInputRowScoped(TargetRow, SourceRow, true);
-
-						Result.bCanDrop = DropFactory->Test(InputRow.GetRow()) != InvalidRowHandle;
-
-						// Currently the description of the last operation will take precedence, we might want support for batched tests in the future
-						if (const FText* Description = DragAndDrop::Helpers::GetDescription(*Storage, InputRow.GetRow()))
-						{
-							Result.Hint = *Description;
-						}
-
-						// No need to test further if any operation failed
-						if (!Result.bCanDrop)
-						{
-							break;
-						}
-					}
-						}
-					}
-
-			if (Result.bCanDrop)
+			if (TSharedPtr<const ITableViewer> Viewer = WeakTableViewer.Pin())
 			{
-				ValidDropOperation = Operation.Get();
+				// Add a column noting that a drop event is currently in progress
+				if (ICoreProvider* Storage = TableViewerModel->GetDataStorageInterface())
+				{
+					Storage->AddColumn(ParentWidgetRowHandle, FTableViewerDropInfoColumn{.DropTarget = InvalidRowHandle});
 				}
-			else if (Result.Hint.IsEmpty())
-			{
-				Result.Hint = LOCTEXT("InvalidDropOp", "Invalid drop operation.");
-			}
-
-			TedsTableViewerRow_Private::UpdateOperationDecorator(*Operation, MoveTemp(Result));
-		}
-	}
 
-	void STedsTableViewerRow::HandleOnDragLeave(const FDragDropEvent& DragDropEvent)
-	{
-		TSharedPtr<FDragDropOperation> Operation = DragDropEvent.GetOperation();
-		ensureMsgf(!ValidDropOperation || ValidDropOperation == Operation.Get(), TEXT("Expected to leave correct operation."));
-
-		if (Operation)
-			{
-			if (ValidDropOperation == Operation.Get())
+				FRowHandleArray Rows;
+				Viewer->ForEachSelectedRow([&Rows](RowHandle Row)
 				{
-				ValidDropOperation = nullptr;
+					Rows.Add(Row);
+				});
+			
+				TSharedRef<Widgets::FTedsDragDropOp> Operation = Widgets::FTedsDragDropOp::New(MoveTemp(Rows));
+				return FReply::Handled().BeginDragDrop(Operation);
 			}
-
-			TedsTableViewerRow_Private::UpdateOperationDecorator(*Operation, {});
+			
 		}
+
+		return FReply::Unhandled();
 	}
 
-	FReply STedsTableViewerRow::HandleOnAcceptDrop(const FDragDropEvent& DragDropEvent, EItemDropZone DropZone, TableViewerItemPtr InItem)
+	void STedsTableViewerRow::HandleOnDragEnter(const FDragDropEvent& DragDropEvent)
 	{
-		TSharedPtr<FDragDropOperation> Operation = DragDropEvent.GetOperation();
-		ensureMsgf(!ValidDropOperation || ValidDropOperation == Operation.Get(), TEXT("Expected to drop correct operation."));
-
-		if (ValidDropOperation && ValidDropOperation == Operation.Get())
-		{
-			ValidDropOperation = nullptr;
-
+		// If the drag enters this widget, the row handle corresponding to it is the drop target
+		// The actual drop logic is handled in the table viewer widget itself
 		if (ICoreProvider* Storage = TableViewerModel->GetDataStorageInterface())
 		{
-			if (UDropOperationsFactory* DropFactory = Storage->FindFactory<UDropOperationsFactory>())
-			{
-						RowHandle TargetRow = Item;
-
-					FRowHandleArray DraggedRows;
-					UE::Editor::DataStorage::DragAndDrop::GetRowsFromData(DraggedRows, *Operation);
-					for (RowHandle SourceRow : DraggedRows.GetRows())
-						{
-							FScopedRow InputRow = DropFactory->CreateInputRowScoped(TargetRow, SourceRow, false);
-						(void)DropFactory->Apply(InputRow.GetRow());
-						}
-						return FReply::Handled();
-					}
-				}
-			}
-		return FReply::Unhandled();
+			Storage->AddColumn(ParentWidgetRowHandle, FTableViewerDropInfoColumn{.DropTarget = Item});
+		}
 	}
 
-	FReply STedsTableViewerRow::HandleOnDragDetected(const FGeometry& MyGeometry, const FPointerEvent& MouseEvent)
+	void STedsTableViewerRow::HandleOnDragLeave(const FDragDropEvent& DragDropEvent)
 	{
-		// TEDS UI TODO: some way to check whether this table viewer allows drag/drop? Column on widget row?
-		if (MouseEvent.IsMouseButtonDown(EKeys::LeftMouseButton))
+		// If the drag exits this widget, remove the corresponding row handle as the drop target
+		// The actual drop logic is handled in the table viewer widget itself
+		if (ICoreProvider* Storage = TableViewerModel->GetDataStorageInterface())
 		{
-			if (TSharedPtr<const ITableViewer> TableViewer = WeakTableViewer.Pin())
+			if (FTableViewerDropInfoColumn* DropInfoColumn = Storage->GetColumn<FTableViewerDropInfoColumn>(ParentWidgetRowHandle))
 			{
-				FRowHandleArray Rows;
-				TableViewer->ForEachSelectedRow([&Rows](RowHandle Row)
-				{
-					Rows.Add(Row);
-				});
-				
-				TSharedRef<Widgets::FTedsDragDropOp> Operation = Widgets::FTedsDragDropOp::New(MoveTemp(Rows));
-				return FReply::Handled().BeginDragDrop(Operation);
+				DropInfoColumn->DropTarget = InvalidRowHandle;
 			}
 		}
-
-		return FReply::Unhandled();
 	}
 
 	TSharedRef<SWidget> SHierarchyViewerRow::GenerateWidgetForColumn(const FName& ColumnName)
diff --git a/Engine/Plugins/Experimental/EditorDataStorageFeatures/Source/TedsTableViewer/Private/Widgets/TableViewerWidgetDropHandler.cpp b/Engine/Plugins/Experimental/EditorDataStorageFeatures/Source/TedsTableViewer/Private/Widgets/TableViewerWidgetDropHandler.cpp
new file mode 100644
--- /dev/null
+++ b/Engine/Plugins/Experimental/EditorDataStorageFeatures/Source/TedsTableViewer/Private/Widgets/TableViewerWidgetDropHandler.cpp
@@ -0,0 +1,71 @@
+// Copyright Epic Games, Inc. All Rights Reserved.
+
+#include "TableViewerWidgetDropHandler.h"
+
+#include "ITedsTableViewer.h"
+#include "TedsDropOperationsColumns.h"
+#include "TedsTableViewerWidgetColumns.h"
+
+namespace UE::Editor::DataStorage::Widgets
+{
+	FTableViewerWidgetDropHandler::FTableViewerWidgetDropHandler(TWeakPtr<ITableViewer> InTableViewer)
+		: FWidgetDropHandler(nullptr) // The table viewer does not operate on selection sets
+		, TableViewer(MoveTemp(InTableViewer))
+	{
+	}
+
+	FWidgetDropHandler::EUpdateInput FTableViewerWidgetDropHandler::UpdateInput(ICoreProvider& Storage, const FGeometry& Geometry,
+		const FDragDropEvent& InputEvent)
+	{
+		RowHandle TargetRow = GetDropTarget(Storage);
+
+		bool bTargetRowChanged = LastTargetRow != TargetRow;
+
+		if (bTargetRowChanged)
+		{
+			for (FDropEntry& Entry : DropEntries)
+			{
+				Storage.AddColumn(Entry.InputRow, DragAndDrop::FTargetColumn{ .Value = TargetRow });
+			}
+			LastTargetRow = TargetRow;
+			return EUpdateInput::ResetOperations;
+		}
+		return EUpdateInput::NoChanges;
+	}
+
+	RowHandle FTableViewerWidgetDropHandler::GetDropTarget(ICoreProvider& Storage) const
+	{
+		if (TSharedPtr<ITableViewer> TableViewerPtr = TableViewer.Pin())
+		{
+			RowHandle TableViewerRow = TableViewerPtr->GetWidgetRowHandle();
+
+			// If we are dragging onto a specific row, that's the target
+			if (FTableViewerDropInfoColumn* DropInfoColumn = Storage.GetColumn<FTableViewerDropInfoColumn>(TableViewerRow))
+			{
+				if (Storage.IsRowAvailable(DropInfoColumn->DropTarget))
+				{
+					return DropInfoColumn->DropTarget;
+				}
+			}
+			// Otherwise if we have a valid "root" row, that is our drop target
+			if (FHierarchyViewerRootColumn* RootColumn = Storage.GetColumn<FHierarchyViewerRootColumn>(TableViewerRow))
+			{
+				return RootColumn->Root;
+			}
+		}
+
+		return InvalidRowHandle;
+	}
+
+	void FTableViewerWidgetDropHandler::Stop(ICoreProvider& Storage)
+	{
+		LastTargetRow = InvalidRowHandle;
+		
+		if (TSharedPtr<ITableViewer> TableViewerPtr = TableViewer.Pin())
+		{
+			Storage.RemoveColumn<FTableViewerDropInfoColumn>(TableViewerPtr->GetWidgetRowHandle());
+		}
+
+		FWidgetDropHandler::Stop(Storage);
+	}
+}
diff --git a/Engine/Plugins/Experimental/EditorDataStorageFeatures/Source/TedsTableViewer/Private/Widgets/TableViewerWidgetDropHandler.h b/Engine/Plugins/Experimental/EditorDataStorageFeatures/Source/TedsTableViewer/Private/Widgets/TableViewerWidgetDropHandler.h
new file mode 100644
--- /dev/null
+++ b/Engine/Plugins/Experimental/EditorDataStorageFeatures/Source/TedsTableViewer/Private/Widgets/TableViewerWidgetDropHandler.h
@@ -0,0 +1,29 @@
+// Copyright Epic Games, Inc. All Rights Reserved.
+
+#pragma once
+
+#include "Widgets/WidgetDropHandler.h"
+
+namespace UE::Editor::DataStorage
+{
+	class ITableViewer;
+}
+
+namespace UE::Editor::DataStorage::Widgets
+{
+	class FTableViewerWidgetDropHandler : public FWidgetDropHandler
+	{
+	public:
+
+		FTableViewerWidgetDropHandler(TWeakPtr<ITableViewer> InTableViewer);
+		virtual ~FTableViewerWidgetDropHandler() override = default;
+
+	protected:
+		virtual EUpdateInput UpdateInput(ICoreProvider& Storage, const FGeometry& Geometry, const FDragDropEvent& InputEvent) override;
+		RowHandle GetDropTarget(ICoreProvider& Storage) const;
+		virtual void Stop(ICoreProvider& Storage) override;
+	private:
+		TWeakPtr<ITableViewer> TableViewer;
+		RowHandle LastTargetRow = InvalidRowHandle;
+	};
+}
diff --git a/Engine/Plugins/Experimental/EditorDataStorageFeatures/Source/TedsTableViewer/Private/Widgets/WidgetDropHandler.cpp b/Engine/Plugins/Experimental/EditorDataStorageFeatures/Source/TedsTableViewer/Private/Widgets/WidgetDropHandler.cpp
--- a/Engine/Plugins/Experimental/EditorDataStorageFeatures/Source/TedsTableViewer/Private/Widgets/WidgetDropHandler.cpp
+++ b/Engine/Plugins/Experimental/EditorDataStorageFeatures/Source/TedsTableViewer/Private/Widgets/WidgetDropHandler.cpp
@@ -147,7 +147,7 @@ FDragDropOperation* FWidgetDropHandler::GetCurrentOperation() const
 	return CurrentOperation;
 }
 
-void FWidgetDropHandler::PrepareInput(DataStorage::ICoreProvider& Storage, UDropOperationsFactory& DropSystem, const FDragDropOperation& Operation)
+bool FWidgetDropHandler::PrepareInput(DataStorage::ICoreProvider& Storage, UDropOperationsFactory& DropSystem, const FDragDropOperation& Operation)
 {
 	using namespace UE::Editor::DataStorage;
 	
@@ -155,10 +155,10 @@ void FWidgetDropHandler::PrepareInput(DataStorage::ICoreProvider& Storage, UDrop
 	DragAndDrop::GetRowsFromData(SourceRows, Operation);
 	SourceRows.Remove(InvalidRowHandle);
 	
-	PrepareInput(Storage, DropSystem, SourceRows.GetRows());
+	return PrepareInput(Storage, DropSystem, SourceRows.GetRows());
 }
 
-void FWidgetDropHandler::PrepareInput(DataStorage::ICoreProvider& Storage, UDropOperationsFactory& DropSystem,
+bool FWidgetDropHandler::PrepareInput(DataStorage::ICoreProvider& Storage, UDropOperationsFactory& DropSystem,
 	DataStorage::FRowHandleArrayView SourceRows)
 {
 	using namespace UE::Editor::DataStorage;
@@ -167,7 +167,7 @@ void FWidgetDropHandler::PrepareInput(DataStorage::ICoreProvider& Storage, UDrop
 	TableHandle Table = DropSystem.GetInputTable();
 	if (Table == InvalidTableHandle)
 	{
-		return;
+		return false;
 	}
 
 	const RowHandle* SourceRowPtr = SourceRows.begin();
@@ -180,6 +180,8 @@ void FWidgetDropHandler::PrepareInput(DataStorage::ICoreProvider& Storage, UDrop
 				DropEntries.Add(FDropEntry{ .InputRow = InputRow });
 			}
 		});
+
+	return true;
 }
 
 FWidgetDropHandler::FDropResult FWidgetDropHandler::ResetOperations(DataStorage::ICoreProvider& Storage, UDropOperationsFactory& DropSystem)
@@ -226,13 +228,15 @@ FWidgetDropHandler::FDropResult FWidgetDropHandler::Start(DataStorage::ICoreProv
 	
 	CurrentOperation = &Operation;
 	
-	PrepareInput(Storage, DropSystem, Operation);
-	(void)UpdateInput(Storage, Geometry, InputEvent);
-	FDropResult Result = ResetOperations(Storage, DropSystem);
-	if (!bIsValidDrop)
+	if (!PrepareInput(Storage, DropSystem, Operation))
 	{
 		Stop(Storage);
+		return FDropResult{bIsValidDrop = false};
 	}
+	(void)UpdateInput(Storage, Geometry, InputEvent);
+	FDropResult Result = ResetOperations(Storage, DropSystem);
+
+	
 	return Result;
 }
 
@@ -275,7 +279,7 @@ void FWidgetDropHandler::Drop(DataStorage::ICoreProvider& Storage)
 
 	FRowHandleArray RowsToSelect;
 	{
-		FScopedTransaction Transaction(LOCTEXT("Drop_Transaction", "Drop on viewport."));
+		FScopedTransaction Transaction(LOCTEXT("Drop_Transaction", "Drop"));
 		for (const FDropEntry& Entry : DropEntries)
 		{
 			const Operations::FApplyColumn* Column = Storage.GetColumn<Operations::FApplyColumn>(Entry.OperationRow);
diff --git a/Engine/Plugins/Experimental/EditorDataStorageFeatures/Source/TedsTableViewer/Public/TedsTableViewerWidgetColumns.h b/Engine/Plugins/Experimental/EditorDataStorageFeatures/Source/TedsTableViewer/Public/TedsTableViewerWidgetColumns.h
--- a/Engine/Plugins/Experimental/EditorDataStorageFeatures/Source/TedsTableViewer/Public/TedsTableViewerWidgetColumns.h
+++ b/Engine/Plugins/Experimental/EditorDataStorageFeatures/Source/TedsTableViewer/Public/TedsTableViewerWidgetColumns.h
@@ -61,6 +61,17 @@ struct FTedsTableViewerRenameSelectionColumn : public FEditorDataStorageColumn
 	FSimpleMulticastDelegate OnRenameSelection;
 };
 
+/**
+ * Column on the table viewer row that contains information about an ongoing drag and drop event
+ */
+USTRUCT(meta = (DisplayName = "Table Viewer drop info"))
+struct FTedsTableViewerDropInfoColumn : public FEditorDataStorageColumn
+{
+	GENERATED_BODY()
+
+	UE::Editor::DataStorage::RowHandle DropTarget;
+};
+
 /**
  * Column on widget rows that contains a weak ptr to the hierarchy data used by the hierarchy viewer to display rows (if present).
  * This is useful if your widget is meant to be displayed in a hierarchy view and needs access to data on the visual parent of the data row
@@ -76,10 +87,24 @@ struct FTedsHierarchyViewerDataColumn : public FEditorDataStorageColumn
 	TWeakPtr<UE::Editor::DataStorage::IHierarchyViewerDataInterface> HierarchyData;
 };
 
+/**
+ * Column on the hierarchy viewer widget row that contains the root of the tree. If there are multiple top level items, the
+ * root is the first one (taking sorting into account)
+ */
+USTRUCT(meta = (DisplayName = "Hierarchy Viewer root item"))
+struct FTedsHierarchyViewerRootColumn : public FEditorDataStorageColumn
+{
+	GENERATED_BODY()
+
+	UE::Editor::DataStorage::RowHandle Root;
+};
+
 namespace UE::Editor::DataStorage
 {
 	using FTableViewerReferenceColumn = FTedsTableViewerReferenceColumn;
+	using FTableViewerDropInfoColumn = FTedsTableViewerDropInfoColumn;
 	using FTableViewerSelectionChangedColumn = FTedsTableViewerSelectionChangedColumn;
 	using FTableViewerRenameSelectionColumn = FTedsTableViewerRenameSelectionColumn;
 	using FHierarchyViewerDataColumn = FTedsHierarchyViewerDataColumn;
+	using FHierarchyViewerRootColumn = FTedsHierarchyViewerRootColumn;
 }
\ No newline at end of file
diff --git a/Engine/Plugins/Experimental/EditorDataStorageFeatures/Source/TedsTableViewer/Public/Widgets/STedsHierarchyViewer.h b/Engine/Plugins/Experimental/EditorDataStorageFeatures/Source/TedsTableViewer/Public/Widgets/STedsHierarchyViewer.h
--- a/Engine/Plugins/Experimental/EditorDataStorageFeatures/Source/TedsTableViewer/Public/Widgets/STedsHierarchyViewer.h
+++ b/Engine/Plugins/Experimental/EditorDataStorageFeatures/Source/TedsTableViewer/Public/Widgets/STedsHierarchyViewer.h
@@ -22,6 +22,10 @@ namespace UE::Editor::DataStorage
 	class ITedsWidget;
 	class STedsTreeView;
 
+	namespace Widgets
+	{
+		class FTableViewerWidgetDropHandler;
+	}
 }
 
 namespace UE::Editor::DataStorage
@@ -155,6 +159,12 @@ namespace UE::Editor::DataStorage
 
 		void SetExpansionState_Internal(FRowHandleArrayView Rows, bool bExpanded) const;
 
+		// Drag and Drop handlers
+		TEDSTABLEVIEWER_API virtual void OnDragEnter(const FGeometry& MyGeometry, const FDragDropEvent& DragDropEvent) override;
+		TEDSTABLEVIEWER_API virtual FReply OnDragOver(const FGeometry& MyGeometry, const FDragDropEvent& DragDropEvent) override;
+		TEDSTABLEVIEWER_API virtual void OnDragLeave(const FDragDropEvent& DragDropEvent) override;
+		TEDSTABLEVIEWER_API virtual FReply OnDrop( const FGeometry& MyGeometry, const FDragDropEvent& DragDropEvent ) override;
+
 	private:
 
 		// The actual ListView widget that displays the rows
@@ -190,5 +200,7 @@ namespace UE::Editor::DataStorage
 		// If we need to update the expansion state of all rows the next time the model changes
 		mutable TOptional<bool> bExpansionStateChangeRequested;
 
+		// Drag and Drop Handler
+		TPimplPtr<Widgets::FTableViewerWidgetDropHandler> DropHandler;
 	};
 }
\ No newline at end of file
diff --git a/Engine/Plugins/Experimental/EditorDataStorageFeatures/Source/TedsTableViewer/Public/Widgets/STedsTableViewer.h b/Engine/Plugins/Experimental/EditorDataStorageFeatures/Source/TedsTableViewer/Public/Widgets/STedsTableViewer.h
--- a/Engine/Plugins/Experimental/EditorDataStorageFeatures/Source/TedsTableViewer/Public/Widgets/STedsTableViewer.h
+++ b/Engine/Plugins/Experimental/EditorDataStorageFeatures/Source/TedsTableViewer/Public/Widgets/STedsTableViewer.h
@@ -16,6 +16,10 @@ class SHeaderRow;
 namespace UE::Editor::DataStorage
 {
 	class ITedsWidget;
+	namespace Widgets
+	{
+		class FTableViewerWidgetDropHandler;
+	}
 }
 
 namespace UE::Editor::DataStorage
@@ -117,6 +121,13 @@ namespace UE::Editor::DataStorage
 		TEDSTABLEVIEWER_API virtual bool IsSelected(RowHandle InRow) const override;
 
 		TEDSTABLEVIEWER_API virtual bool IsSelectedExclusively(RowHandle InRow) const override;
+
+		// Drag and Drop handlers
+		TEDSTABLEVIEWER_API virtual void OnDragEnter(const FGeometry& MyGeometry, const FDragDropEvent& DragDropEvent) override;
+		TEDSTABLEVIEWER_API virtual FReply OnDragOver(const FGeometry& MyGeometry, const FDragDropEvent& DragDropEvent) override;
+		TEDSTABLEVIEWER_API virtual void OnDragLeave(const FDragDropEvent& DragDropEvent) override;
+		TEDSTABLEVIEWER_API virtual FReply OnDrop( const FGeometry& MyGeometry, const FDragDropEvent& DragDropEvent ) override;
+
 	protected:
 		
 		TSharedRef<ITableRow> MakeTableRowWidget(TableViewerItemPtr InItem, const TSharedRef<STableViewBase>& OwnerTable) const;
@@ -155,5 +166,8 @@ namespace UE::Editor::DataStorage
 
 		// The message to show in place of the table viewer when there are no rows provided by the current query stack
 		TAttribute<FText> EmptyRowsMessage;
+
+		// Drag and Drop Handler
+		TPimplPtr<Widgets::FTableViewerWidgetDropHandler> DropHandler;
 	};
 }
\ No newline at end of file
diff --git a/Engine/Plugins/Experimental/EditorDataStorageFeatures/Source/TedsTableViewer/Public/Widgets/STedsTableViewerRow.h b/Engine/Plugins/Experimental/EditorDataStorageFeatures/Source/TedsTableViewer/Public/Widgets/STedsTableViewerRow.h
--- a/Engine/Plugins/Experimental/EditorDataStorageFeatures/Source/TedsTableViewer/Public/Widgets/STedsTableViewerRow.h
+++ b/Engine/Plugins/Experimental/EditorDataStorageFeatures/Source/TedsTableViewer/Public/Widgets/STedsTableViewerRow.h
@@ -59,8 +59,6 @@ namespace UE::Editor::DataStorage
 		void HandleOnDragEnter(const FDragDropEvent& DragDropEvent);
 		void HandleOnDragLeave(const FDragDropEvent& DragDropEvent);
 		FReply HandleOnDragDetected(const FGeometry& MyGeometry, const FPointerEvent& MouseEvent);
-		FReply HandleOnAcceptDrop(const FDragDropEvent& DragDropEvent, EItemDropZone DropZone, TableViewerItemPtr Item);
-
 
 	protected:
 		TAttribute<float> ItemHeight;
@@ -68,8 +66,6 @@ namespace UE::Editor::DataStorage
 		TSharedPtr<FTedsTableViewerModel> TableViewerModel;
 		TWeakPtr<const ITableViewer> WeakTableViewer;
 		TableViewerItemPtr Item;
-
-		FDragDropOperation* ValidDropOperation = nullptr; // Not for use, just to denote if a valid drop operation is hovering.
 	};
 
 	/** Widget that represents a row in the hierarchy viewers.  Generates widgets for each column and adds an expander arrow to the very first column. */
diff --git a/Engine/Plugins/Experimental/EditorDataStorageFeatures/Source/TedsTableViewer/Public/Widgets/WidgetDropHandler.h b/Engine/Plugins/Experimental/EditorDataStorageFeatures/Source/TedsTableViewer/Public/Widgets/WidgetDropHandler.h
--- a/Engine/Plugins/Experimental/EditorDataStorageFeatures/Source/TedsTableViewer/Public/Widgets/WidgetDropHandler.h
+++ b/Engine/Plugins/Experimental/EditorDataStorageFeatures/Source/TedsTableViewer/Public/Widgets/WidgetDropHandler.h
@@ -53,9 +53,9 @@ public:
 
 protected:
 	/** Creates input rows for the given drag&drop input event. */
-	UE_API virtual void PrepareInput(DataStorage::ICoreProvider& Storage, UDropOperationsFactory& DropSystem, const FDragDropOperation& Operation);
+	UE_API virtual bool PrepareInput(DataStorage::ICoreProvider& Storage, UDropOperationsFactory& DropSystem, const FDragDropOperation& Operation);
 	/** Creates input rows for the given drag&drop input event. */
-	UE_API virtual void PrepareInput(DataStorage::ICoreProvider& Storage, UDropOperationsFactory& DropSystem, DataStorage::FRowHandleArrayView SourceRows);
+	UE_API virtual bool PrepareInput(DataStorage::ICoreProvider& Storage, UDropOperationsFactory& DropSystem, DataStorage::FRowHandleArrayView SourceRows);
 	/** Finds suitable operations for the given input rows and sets 'bIsValidDrop'. */
 	UE_API virtual FDropResult ResetOperations(DataStorage::ICoreProvider& Storage, UDropOperationsFactory& DropSystem);
 	/** Updates all input rows for the updated input parameters. May also set status vars like 'bIsValidDrop'. */
