[TEDS Executor] Avoid FName recreation when resetting the Executor in the Model and Hierarchy Viewer

diff --git a/Engine/Plugins/Experimental/EditorDataStorageFeatures/Source/TedsTableViewer/Private/TedsTableViewerModel.cpp b/Engine/Plugins/Experimental/EditorDataStorageFeatures/Source/TedsTableViewer/Private/TedsTableViewerModel.cpp
--- a/Engine/Plugins/Experimental/EditorDataStorageFeatures/Source/TedsTableViewer/Private/TedsTableViewerModel.cpp
+++ b/Engine/Plugins/Experimental/EditorDataStorageFeatures/Source/TedsTableViewer/Private/TedsTableViewerModel.cpp
@@ -25,7 +25,8 @@ namespace UE::Editor::DataStorage
 	{
 		Storage = GetMutableDataStorageFeature<ICoreProvider>(StorageFeatureName);
 		StorageUi = GetMutableDataStorageFeature<IUiProvider>(UiFeatureName);
-		
+
+		ExecutorInstanceName = FName("TableViewerModel", InstanceCounter++);
 		SetQueryStack(Args.RowQueryStack, false);
 
 		TickerHandle = FTSTicker::GetCoreTicker().AddTicker(FTickerDelegate::CreateRaw(this, &FTedsTableViewerModel::Tick), 0);
@@ -169,9 +170,8 @@ namespace UE::Editor::DataStorage
 		RowPrimarySortingNode = MakeShared<QueryStack::FRowSortNode>(*Storage, RowSecondaryInversionNode, RowPrimarySortingNode ? RowPrimarySortingNode->GetColumnSorter() : nullptr);
 		RowPrimaryInversionNode = MakeShared<QueryStack::FRowOrderInversionNode>(RowPrimarySortingNode, RowPrimaryInversionNode ? RowPrimaryInversionNode->IsEnabled() : false);
 		RowQueryStack = RowPrimaryInversionNode;
-
-		static int InstanceCounter = 0;
-		QueryStackExecutor = QueryStack::FExplicitUpdateExecutor(FName("TableViewerModel", InstanceCounter++), RowQueryStack);
+		
+		QueryStackExecutor = QueryStack::FExplicitUpdateExecutor(ExecutorInstanceName, RowQueryStack);
 
 		if (bUpdateExecutor)
 		{
diff --git a/Engine/Plugins/Experimental/EditorDataStorageFeatures/Source/TedsTableViewer/Private/Widgets/STedsHierarchyViewer.cpp b/Engine/Plugins/Experimental/EditorDataStorageFeatures/Source/TedsTableViewer/Private/Widgets/STedsHierarchyViewer.cpp
--- a/Engine/Plugins/Experimental/EditorDataStorageFeatures/Source/TedsTableViewer/Private/Widgets/STedsHierarchyViewer.cpp
+++ b/Engine/Plugins/Experimental/EditorDataStorageFeatures/Source/TedsTableViewer/Private/Widgets/STedsHierarchyViewer.cpp
@@ -65,9 +65,9 @@ namespace UE::Editor::DataStorage
 		});
 
 		TopLevelRowsNode = MakeShared<QueryStack::FTopLevelRowsNode>(Model->GetDataStorageInterface(), HierarchyInterface, Model->GetRowNode());
-		
-		static int InstanceCounter = 0;
-		QueryStackExecutor = QueryStack::FExplicitUpdateExecutor(FName("HierarchyViewer", InstanceCounter++), TopLevelRowsNode);
+
+		ExecutorInstanceName = FName("TableViewerModel", InstanceCounter++);
+		QueryStackExecutor = QueryStack::FExplicitUpdateExecutor(ExecutorInstanceName, TopLevelRowsNode);
 		
 		HeaderRowWidget = SNew( SHeaderRow )
 							.CanSelectGeneratedColumn(InArgs._CanSelectGeneratedColumn);
@@ -284,8 +284,7 @@ namespace UE::Editor::DataStorage
 	{
 		Model->SetQueryStack(InRowQueryStack);
 		TopLevelRowsNode = MakeShared<QueryStack::FTopLevelRowsNode>(Model->GetDataStorageInterface(), HierarchyInterface, Model->GetRowNode());
-		static int InstanceCounter = 0;
-		QueryStackExecutor = QueryStack::FExplicitUpdateExecutor(FName("HierarchyViewer", InstanceCounter++), TopLevelRowsNode);
+		QueryStackExecutor = QueryStack::FExplicitUpdateExecutor(ExecutorInstanceName, TopLevelRowsNode);
 	}
 
 	void SHierarchyViewer::SetColumns(const TArray<TWeakObjectPtr<const UScriptStruct>>& InColumns)
diff --git a/Engine/Plugins/Experimental/EditorDataStorageFeatures/Source/TedsTableViewer/Public/TedsTableViewerModel.h b/Engine/Plugins/Experimental/EditorDataStorageFeatures/Source/TedsTableViewer/Public/TedsTableViewerModel.h
--- a/Engine/Plugins/Experimental/EditorDataStorageFeatures/Source/TedsTableViewer/Public/TedsTableViewerModel.h
+++ b/Engine/Plugins/Experimental/EditorDataStorageFeatures/Source/TedsTableViewer/Public/TedsTableViewerModel.h
@@ -174,5 +174,9 @@ namespace UE::Editor::DataStorage
 
 		// TEDS -> UI columns is not a 1:1 mapping, so we cache the name of the UI column so widgets can query it easily
 		FName PrimaryUIColumnName;
+
+		// Store the Instance count and the created name for use by the executor (so it can avoid recreation)
+		inline static int InstanceCounter = 0;
+		FName ExecutorInstanceName;
 	};
 } // namespace UE::Editor::DataStorage
diff --git a/Engine/Plugins/Experimental/EditorDataStorageFeatures/Source/TedsTableViewer/Public/Widgets/STedsHierarchyViewer.h b/Engine/Plugins/Experimental/EditorDataStorageFeatures/Source/TedsTableViewer/Public/Widgets/STedsHierarchyViewer.h
--- a/Engine/Plugins/Experimental/EditorDataStorageFeatures/Source/TedsTableViewer/Public/Widgets/STedsHierarchyViewer.h
+++ b/Engine/Plugins/Experimental/EditorDataStorageFeatures/Source/TedsTableViewer/Public/Widgets/STedsHierarchyViewer.h
@@ -191,6 +191,10 @@ namespace UE::Editor::DataStorage
 		// The executor used to update the hierarchy stack.
 		QueryStack::FExplicitUpdateExecutor QueryStackExecutor;
 		
+		// Store the Instance count and the created name for use by the executor (so it can avoid recreation)
+		inline static int InstanceCounter = 0;
+		FName ExecutorInstanceName;
+		
 		// The cached list of top level rows
 		TArray<TableViewerItemPtr> TopLevelRows;
 
