[TEDs] Property Info 1.0
- Added Type Property Info to Teds
- Added CVar to only store Verse properties by default
- Added Hierarchy information between type rows and property rows
- Added type rows to compatability to handle garbage collection

diff --git a/Engine/Plugins/Experimental/EditorDataStorageFeatures/Source/TedsTypeInfo/Private/TedsTypeInfoFactory.cpp b/Engine/Plugins/Experimental/EditorDataStorageFeatures/Source/TedsTypeInfo/Private/TedsTypeInfoFactory.cpp
--- a/Engine/Plugins/Experimental/EditorDataStorageFeatures/Source/TedsTypeInfo/Private/TedsTypeInfoFactory.cpp
+++ b/Engine/Plugins/Experimental/EditorDataStorageFeatures/Source/TedsTypeInfo/Private/TedsTypeInfoFactory.cpp
@@ -11,6 +11,7 @@
 #include "TedsTypeInfoModule.h"
 
 #include "DataStorage/Features.h"
+#include "Elements/Interfaces/TypedElementDataStorageCompatibilityInterface.h"
 #include "Elements/Interfaces/TypedElementDataStorageInterface.h"
 #include "Elements/Columns/TypedElementLabelColumns.h"
 #include "Elements/Framework/TypedElementQueryBuilder.h"
@@ -24,8 +25,14 @@
 #include "UObject/UObjectIterator.h"
 #include "VerseVM/VVMVerseClass.h"
 
-const FName UTypeInfoFactory::TypeTableName = "Type Information";
+const FName UTypeInfoFactory::PropertyTableName = "Property Information";
 const FName UTypeInfoFactory::ClassSetupActionName = "ClassSetupAction";
+const FName UTypeInfoFactory::ClassHierarchyName = "ClassHierarchy";
+const FName UTypeInfoFactory::PropertyHierarchyName = "PropertyHierarchy";
+
+TAutoConsoleVariable<bool> CVarTEDSTypeInfoOnlyVerseProperties(TEXT("TEDS.Feature.TypeInfoIntegration.OnlyStoreVerseProperties"),
+	true,
+	TEXT("When true we will only store properties of Verse Classes"));
 
 void UTypeInfoFactory::PreRegister(UE::Editor::DataStorage::ICoreProvider& DataStorage)
 {
@@ -38,23 +45,32 @@ void UTypeInfoFactory::RegisterTables(UE::Editor::DataStorage::ICoreProvider& Da
 
 	Super::RegisterTables(DataStorage, DataStorageCompatability);
 
-	TypeTable = DataStorage.RegisterTable({
-		FTypedElementLabelColumn::StaticStruct(),
-		FTypeInfoTag::StaticStruct(),
-		FTypeInfoRequiresHierarchyUpdateTag::StaticStruct()
-		}, TypeTableName);
+	PropertyTable = DataStorage.RegisterTable({
+		FPropertyInfoTag::StaticStruct(),
+		FPropertyNameColumn::StaticStruct(),
+		FPropertyOwnerTypeRowHandleColumn::StaticStruct(),
+		FPropertyTypeNameInfoColumn::StaticStruct(),
+		FPropertyTypeNameColumn::StaticStruct(),
+	}, PropertyTableName);
 }
 
 void UTypeInfoFactory::RegisterHierarchies(UE::Editor::DataStorage::ICoreProvider& DataStorage)
 {
 	using namespace UE::Editor::DataStorage;
 
-	const FHierarchyRegistrationParams Params
+	const FHierarchyRegistrationParams ClassHierarchyParams
+	{
+		.Name = ClassHierarchyName,
+	};
+
+	DataStorage.RegisterHierarchy(ClassHierarchyParams);
+
+	const FHierarchyRegistrationParams PropertyHierarchyParams
 	{
-		.Name = TEXT("ClassHierarchy"),
+		.Name = PropertyHierarchyName,
 	};
 
-	FHierarchyHandle Handle = DataStorage.RegisterHierarchy(Params);
+	PropertyHierarchyHandle = DataStorage.RegisterHierarchy(PropertyHierarchyParams);
 }
 
 void UTypeInfoFactory::RegisterQueries(UE::Editor::DataStorage::ICoreProvider& DataStorage)
@@ -80,7 +96,7 @@ void UTypeInfoFactory::RegisterQueries(UE::Editor::DataStorage::ICoreProvider& D
 
 				Context.RemoveColumns(Row, { FTypeInfoRequiresHierarchyUpdateTag::StaticStruct() });
 			})
-			.AccessesHierarchy(TEXT("ClassHierarchy"))
+			.AccessesHierarchy(ClassHierarchyName)
 			.Where()
 				.All<FTypeInfoRequiresHierarchyUpdateTag>()
 			.Compile());
@@ -109,15 +125,11 @@ void UTypeInfoFactory::OnDatabaseUpdateCompleted()
 	}
 }
 
-UE::Editor::DataStorage::TableHandle UTypeInfoFactory::GetTypeTable() const
-{
-	return TypeTable;
-}
-
 void UTypeInfoFactory::ClearAllTypeInfo()
 {
 	using namespace UE::Editor::DataStorage;
 	ClearTypeInfoByTag<FTypeInfoTag>();
+	ClearTypeInfoByTag<FPropertyInfoTag>();
 }
 
 void UTypeInfoFactory::RefreshAllTypeInfo()
@@ -130,6 +142,8 @@ template <class Type>
 void UTypeInfoFactory::PopulateTypeInfo()
 {
 	using namespace UE::Editor::DataStorage;
+	ICoreProvider* DataStorage = GetMutableDataStorageFeature<ICoreProvider>(StorageFeatureName);
+	checkf(DataStorage, TEXT("Attempted to add type info to Teds with a null DataStorage"));
 
 	// Block adding anything to the data table if our CVar is off
 	static const auto TedsTypeInfoEnabledCVar = IConsoleManager::Get().FindConsoleVariable(TEXT("TEDS.Feature.TypeInfoIntegration.Enable"));
@@ -138,63 +152,83 @@ void UTypeInfoFactory::PopulateTypeInfo()
 		return;
 	}
 
+	bool bOnlyAddVerseProperties = CVarTEDSTypeInfoOnlyVerseProperties.GetValueOnGameThread();
+
+	TMap<Type*, UE::Editor::DataStorage::RowHandle> TypesWithPendingPropertyInfo;
+	UE::Editor::DataStorage::RowHandle OutRowhandle = UE::Editor::DataStorage::InvalidRowHandle;
 	for (TObjectIterator<Type> TypeIterator; TypeIterator; ++TypeIterator)
 	{
-		TryAddTypeInfoRow(*TypeIterator);
+		Type* CurrentType = *TypeIterator;
+		if (CurrentType && TryAddTypeInfo(*DataStorage, CurrentType, OutRowhandle))
+		{
+			if (!bOnlyAddVerseProperties || CurrentType->template IsChildOf<UVerseClass>())
+			{
+				TypesWithPendingPropertyInfo.Emplace(CurrentType, OutRowhandle);
+			}
+		}
+	}
+
+	// Separate type info pass from property info pass. We can lookup types during property pass and time slice properties in the future
+	for (const TPair<Type*, UE::Editor::DataStorage::RowHandle> TypeWithPendingPropertyInfo : TypesWithPendingPropertyInfo)
+	{
+		AddPropertyInfo(*DataStorage, TypeWithPendingPropertyInfo.Value, TypeWithPendingPropertyInfo.Key);
 	}
 
-	ICoreProvider* DataStorage = GetMutableDataStorageFeature<ICoreProvider>(StorageFeatureName);
 	DataStorage->ActivateQueries(ClassSetupActionName);
 }
 
-bool UTypeInfoFactory::TryAddTypeInfoRow(const UStruct* TypeInfo)
+bool UTypeInfoFactory::TryAddTypeInfo(UE::Editor::DataStorage::ICoreProvider& DataStorage, const UStruct* InTypeInfo,
+	UE::Editor::DataStorage::RowHandle& OutTypeRowHandle)
 {
 	// Do not include InstanceDataObject types in TEDS
-	if (UE::IsClassOfInstanceDataObjectClass(TypeInfo))
+	if (UE::IsClassOfInstanceDataObjectClass(InTypeInfo))
 	{
 		return false;
 	}
 
 	using namespace UE::Editor::DataStorage;
 
-	ICoreProvider* DataStorage = GetMutableDataStorageFeature<ICoreProvider>(StorageFeatureName);
-	checkf(DataStorage, TEXT("Attempted to add type info to Teds with a null DataStorage"));
+	ICompatibilityProvider* Compatibility = GetMutableDataStorageFeature<ICompatibilityProvider>(CompatibilityFeatureName);
+	checkf(Compatibility, TEXT("Attempted to add type info to Teds with a null CompatibilityProvider"));
 
 	// We have 2 major type cases: Classes and Structs
 	
 	// Case 1: Classes and Verse Classes
-	const UClass* ClassInfo = Cast<UClass>(TypeInfo);
+	const UClass* ClassInfo = Cast<UClass>(InTypeInfo);
 	if (ClassInfo && FilterClassInfo(ClassInfo))
 	{
-		if (RowHandle TypeRowHandle = DataStorage->AddRow(TypeTable))
+		if (RowHandle TypeRowHandle = Compatibility->AddCompatibleObject(const_cast<UStruct*>(InTypeInfo)))
 		{
-			AddCommonColumns(*DataStorage, TypeRowHandle, ClassInfo);
-			AddClassColumns(*DataStorage, TypeRowHandle, ClassInfo);
+			AddCommonColumns(DataStorage, TypeRowHandle, ClassInfo);
+			AddClassColumns(DataStorage, TypeRowHandle, ClassInfo);
 
 			// Only Classes can be UVerseClasses
 			if (const UVerseClass* VerseClassInfo = Cast<UVerseClass>(ClassInfo))
 			{
-				AddVerseColumns(*DataStorage, TypeRowHandle, VerseClassInfo);
-				
-				AddIdoColumns(*DataStorage, TypeRowHandle, VerseClassInfo);
+				AddVerseColumns(DataStorage, TypeRowHandle, VerseClassInfo);
+				AddIdoColumns(DataStorage, TypeRowHandle, VerseClassInfo);
 			}
 
-			DataStorage->MapRow(TypeInfo::TypeMappingDomain, FMapKey(ClassInfo), TypeRowHandle);
+			DataStorage.MapRow(TypeInfo::TypeMappingDomain, FMapKey(ClassInfo), TypeRowHandle);
+			OutTypeRowHandle = TypeRowHandle;
+
 			return true;
 		}
 	}
 
 	// Case 2: Structs
 	// Currently unsupported as filter will return false
-	const UScriptStruct* StructData = Cast<UScriptStruct>(TypeInfo);
+	const UScriptStruct* StructData = Cast<UScriptStruct>(InTypeInfo);
 	if (StructData && FilterStructInfo(StructData))
 	{
-		if (RowHandle TypeRowHandle = DataStorage->AddRow(TypeTable))
+		if (RowHandle TypeRowHandle = Compatibility->AddCompatibleObject(const_cast<UStruct*>(InTypeInfo)))
 		{
-			AddCommonColumns(*DataStorage, TypeRowHandle, StructData);
-			AddStructColumns(*DataStorage, TypeRowHandle, StructData);
+			AddCommonColumns(DataStorage, TypeRowHandle, StructData);
+			AddStructColumns(DataStorage, TypeRowHandle, StructData);
+
+			DataStorage.MapRow(TypeInfo::TypeMappingDomain, FMapKey(StructData), TypeRowHandle);
+			OutTypeRowHandle = TypeRowHandle;
 
-			DataStorage->MapRow(TypeInfo::TypeMappingDomain, FMapKey(StructData), TypeRowHandle);
 			return true;
 		}
 	}
@@ -202,6 +236,25 @@ bool UTypeInfoFactory::TryAddTypeInfoRow(const UStruct* TypeInfo)
 	return false;
 }
 
+void UTypeInfoFactory::AddPropertyInfo(UE::Editor::DataStorage::ICoreProvider& DataStorage, UE::Editor::DataStorage::RowHandle OwnerTypeRowHandle, const UStruct* InTypeInfo)
+{
+	using namespace UE::Editor::DataStorage;
+	checkf(InTypeInfo, TEXT("Attempted to add property info row(s) with null TypeInfo"));
+
+	// Use a TFieldIterator with IncludeSuper and IncludeInterface properties default off
+	// We only want to store the immediate properties per class and will reference interface and super through class hierarchy
+	for (TFieldIterator<FProperty> CurrentPropertyIt(InTypeInfo); CurrentPropertyIt; ++CurrentPropertyIt)
+	{
+		if (RowHandle PropertyRowHandle = DataStorage.AddRow(PropertyTable))
+		{
+			// TODO: Cleanup property rows when associated type row is destroyed
+			// We get away with it for now because Verse Properties are destroyed during compiles and we refresh the whole type info table then
+			DataStorage.SetParentRow(PropertyHierarchyHandle, PropertyRowHandle, OwnerTypeRowHandle);
+			AddCommonPropertyColumns(DataStorage, OwnerTypeRowHandle, PropertyRowHandle, *CurrentPropertyIt);
+		}
+	}
+}
+
 bool UTypeInfoFactory::FilterStructInfo(const UStruct* /*StructInfo*/)
 {
 	// Structs not supported currently
@@ -315,7 +368,26 @@ void UTypeInfoFactory::AddIdoColumns(UE::Editor::DataStorage::ICoreProvider& Dat
 	}
 }
 
-template <class TypeToClear>
+void UTypeInfoFactory::AddCommonPropertyColumns(UE::Editor::DataStorage::ICoreProvider& DataStorage, UE::Editor::DataStorage::RowHandle OwnerTypeRowHandle,
+	UE::Editor::DataStorage::RowHandle PropertyRowHandle, const FProperty* InPropertyInfo)
+{
+	using namespace UE::Editor::DataStorage;
+
+	checkf(InPropertyInfo, TEXT("Attempted to add property info columns with null PropertyInfo"));
+
+	DataStorage.AddColumn(PropertyRowHandle, FPropertyNameColumn{ .PropertyName = InPropertyInfo->GetDisplayNameText() });
+	DataStorage.AddColumn(PropertyRowHandle, FPropertyOwnerTypeRowHandleColumn{ .PropertyOwnerTypeRowHandle = OwnerTypeRowHandle });
+	DataStorage.AddColumn(PropertyRowHandle, FPropertyTypeNameInfoColumn{ .PropertyTypeNameInfo = UE::FPropertyTypeName(InPropertyInfo) });
+
+	if (FFieldClass* PropertyClass = InPropertyInfo->GetClass())
+	{
+		DataStorage.AddColumn(PropertyRowHandle, FPropertyTypeNameColumn{ .PropertyTypeName = InPropertyInfo->GetClass()->GetDisplayNameText() });
+	}
+
+	DataStorage.AddColumn<FPropertyInfoTag>(PropertyRowHandle);
+}
+
+template <class TagToClear>
 void UTypeInfoFactory::ClearTypeInfoByTag()
 {
 	using namespace UE::Editor::DataStorage;
@@ -323,5 +395,5 @@ void UTypeInfoFactory::ClearTypeInfoByTag()
 
 	checkf(DataStorage, TEXT("Attempted to clear type info from Teds with a null DataStorage"));
 
-	DataStorage->RemoveAllRowsWith<TypeToClear>();
+	DataStorage->RemoveAllRowsWith<TagToClear>();
 }
diff --git a/Engine/Plugins/Experimental/EditorDataStorageFeatures/Source/TedsTypeInfo/Private/TedsTypeInfoFactory.h b/Engine/Plugins/Experimental/EditorDataStorageFeatures/Source/TedsTypeInfo/Private/TedsTypeInfoFactory.h
--- a/Engine/Plugins/Experimental/EditorDataStorageFeatures/Source/TedsTypeInfo/Private/TedsTypeInfoFactory.h
+++ b/Engine/Plugins/Experimental/EditorDataStorageFeatures/Source/TedsTypeInfo/Private/TedsTypeInfoFactory.h
@@ -35,8 +35,10 @@ class UTypeInfoFactory final : public UEditorDataStorageFactory
 	GENERATED_BODY()
 
 public:
-	static const FName TypeTableName;
+	static const FName PropertyTableName;
 	static const FName ClassSetupActionName;
+	static const FName ClassHierarchyName;
+	static const FName PropertyHierarchyName;
 
 	virtual ~UTypeInfoFactory() override = default;
 
@@ -48,8 +50,6 @@ public:
 	void PreShutdown(UE::Editor::DataStorage::ICoreProvider& DataStorage) override;
 	//~ End UEditorDataStorageFactory interface
 
-	UE::Editor::DataStorage::TableHandle GetTypeTable() const;
-
 	TEDSTYPEINFO_API void ClearAllTypeInfo();
 	TEDSTYPEINFO_API void RefreshAllTypeInfo();
 
@@ -60,10 +60,11 @@ private:
 	template <class Type>
 	void PopulateTypeInfo();
 
-	template <class TypeToClear>
+	template <class TagToClear>
 	void ClearTypeInfoByTag();
 
-	bool TryAddTypeInfoRow(const UStruct* TypeInfo);
+	bool TryAddTypeInfo(UE::Editor::DataStorage::ICoreProvider& DataStorage, const UStruct* InTypeInfo, UE::Editor::DataStorage::RowHandle& OutTypeRowHandle);
+	void AddPropertyInfo(UE::Editor::DataStorage::ICoreProvider& DataStorage, UE::Editor::DataStorage::RowHandle OwnerTypeRowHandle, const UStruct* InTypeInfo);
 
 	bool FilterStructInfo(const UStruct* StructInfo);
 	bool FilterClassInfo(const UClass* ClassInfo);
@@ -74,11 +75,14 @@ private:
 	void AddVerseColumns(UE::Editor::DataStorage::ICoreProvider& DataStorage, UE::Editor::DataStorage::RowHandle TypeRowHandle, const UVerseClass* InVerseTypeInfo);
 	void AddIdoColumns(UE::Editor::DataStorage::ICoreProvider& DataStorage, UE::Editor::DataStorage::RowHandle TypeRowHandle, const UVerseClass* InVerseTypeInfo);
 
-	bool bRefreshTypeInfoQueued = false;
+	void AddCommonPropertyColumns(UE::Editor::DataStorage::ICoreProvider& DataStorage, UE::Editor::DataStorage::RowHandle TypeRowhandle,
+		UE::Editor::DataStorage::RowHandle PropertyRowhandle, const FProperty* InPropertyInfo);
 
-	const UE::Editor::DataStorage::IHierarchyAccessInterface* ClassHierarchyAccessInterface = nullptr;
+	bool bRefreshTypeInfoQueued = false;
 
 	UE::Editor::DataStorage::TableHandle TypeTable = UE::Editor::DataStorage::InvalidTableHandle;
+	UE::Editor::DataStorage::TableHandle PropertyTable = UE::Editor::DataStorage::InvalidTableHandle;
+	UE::Editor::DataStorage::FHierarchyHandle PropertyHierarchyHandle;
 };
 
 namespace UE::Editor::DataStorage
diff --git a/Engine/Plugins/Experimental/EditorDataStorageFeatures/Source/TedsTypeInfo/Public/Elements/Columns/TedsTypeInfoColumns.h b/Engine/Plugins/Experimental/EditorDataStorageFeatures/Source/TedsTypeInfo/Public/Elements/Columns/TedsTypeInfoColumns.h
--- a/Engine/Plugins/Experimental/EditorDataStorageFeatures/Source/TedsTypeInfo/Public/Elements/Columns/TedsTypeInfoColumns.h
+++ b/Engine/Plugins/Experimental/EditorDataStorageFeatures/Source/TedsTypeInfo/Public/Elements/Columns/TedsTypeInfoColumns.h
@@ -2,10 +2,12 @@
 #pragma once
 
 #include "DataStorage/CommonTypes.h"
+#include "DataStorage/Handles.h"
 #include "Elements/Columns/TypedElementTypeInfoColumns.h"
 
 #include "TedsTypeInfoColumns.generated.h"
 
+// Type Info Columns Start
 USTRUCT(meta = (DisplayName = "Row is a type"))
 struct FDataStorageTypeInfoTag : public FEditorDataStorageTag
 {
@@ -47,9 +49,53 @@ struct FDataStorageHasUnknownPropertiesTag : public FEditorDataStorageTag
 {
 	GENERATED_BODY()
 };
+// Type Info Columns End
+
+// Property Info Columns Start
+USTRUCT(meta = (DisplayName = "Property Name"))
+struct FDataStoragePropertyNameColumn final : public FEditorDataStorageColumn
+{
+	GENERATED_BODY()
+
+	UPROPERTY(meta = (Searchable, Sortable))
+	FText PropertyName;
+};
+
+USTRUCT(meta = (DisplayName = "Property Owner Type Row Handle"))
+struct FDataStoragePropertyOwnerTypeRowHandleColumn final : public FEditorDataStorageColumn
+{
+	GENERATED_BODY()
+
+	UE::Editor::DataStorage::RowHandle PropertyOwnerTypeRowHandle = UE::Editor::DataStorage::InvalidRowHandle;
+};
+
+USTRUCT(meta = (DisplayName = "Property Type Name"))
+struct FDataStoragePropertyTypeNameColumn final : public FEditorDataStorageColumn
+{
+	GENERATED_BODY()
+
+	UPROPERTY(meta = (Searchable, Sortable))
+	FText PropertyTypeName;
+};
+
+USTRUCT(meta = (DisplayName = "Property Type Name Info"))
+struct FDataStoragePropertyTypeNameInfoColumn final : public FEditorDataStorageColumn
+{
+	GENERATED_BODY()
+
+	UE::FPropertyTypeName PropertyTypeNameInfo;
+};
+
+USTRUCT(meta = (DisplayName = "Row is a property"))
+struct FDataStoragePropertyInfoTag : public FEditorDataStorageTag
+{
+	GENERATED_BODY()
+};
+// Property Info Columns End
 
 namespace UE::Editor::DataStorage
 {
+	// Type Info Columns Start
 	using FTypeInfoTag = FDataStorageTypeInfoTag;
 	using FClassTypeInfoTag = FDataStorageClassTypeInfoTag;
 	using FStructTypeInfoTag = FDataStorageStructTypeInfoTag;
@@ -57,6 +103,15 @@ namespace UE::Editor::DataStorage
 	using FVerseTypeInfoTag = FDataStorageVerseTypeInfoTag;
 	using FVerseTypeInfoAccessLevel = FDataStorageVerseTypeInfoAccessLevel;
 	using FHasUnknownPropertiesTag = FDataStorageHasUnknownPropertiesTag;
+	// Type Info Tags End
+
+	// Property Info Columns Start
+	using FPropertyNameColumn = FDataStoragePropertyNameColumn;
+	using FPropertyOwnerTypeRowHandleColumn = FDataStoragePropertyOwnerTypeRowHandleColumn;
+	using FPropertyTypeNameColumn = FDataStoragePropertyTypeNameColumn;
+	using FPropertyTypeNameInfoColumn = FDataStoragePropertyTypeNameInfoColumn;
+	using FPropertyInfoTag = FDataStoragePropertyInfoTag;
+	// Property Info Tags End
 }
 
 namespace UE::Editor::DataStorage::TypeInfo
