[TedsOperations] Refactored teds operations into one plugin. API now supports batching.

diff --git a/Engine/Plugins/Experimental/EditorDataStorageFeatures/EditorDataStorageFeatures.uplugin b/Engine/Plugins/Experimental/EditorDataStorageFeatures/EditorDataStorageFeatures.uplugin
--- a/Engine/Plugins/Experimental/EditorDataStorageFeatures/EditorDataStorageFeatures.uplugin
+++ b/Engine/Plugins/Experimental/EditorDataStorageFeatures/EditorDataStorageFeatures.uplugin
@@ -79,10 +79,6 @@
 		{
 			"Name": "TedsOperations",
 			"Type": "Editor"
-		},
-		{
-			"Name": "TedsDropOperations",
-			"Type": "Editor"
 		}
 	],
 	"Plugins": [
diff --git a/Engine/Plugins/Experimental/EditorDataStorageFeatures/Source/TedsDropOperations/Private/TedsDragDropOp.cpp b/Engine/Plugins/Experimental/EditorDataStorageFeatures/Source/TedsDropOperations/Private/TedsDragDropOp.cpp
deleted file mode 100644
diff --git a/Engine/Plugins/Experimental/EditorDataStorageFeatures/Source/TedsDropOperations/Private/TedsDropOperationsFactory.cpp b/Engine/Plugins/Experimental/EditorDataStorageFeatures/Source/TedsDropOperations/Private/TedsDropOperationsFactory.cpp
deleted file mode 100644
diff --git a/Engine/Plugins/Experimental/EditorDataStorageFeatures/Source/TedsDropOperations/Private/TedsDropOperationsModule.cpp b/Engine/Plugins/Experimental/EditorDataStorageFeatures/Source/TedsDropOperations/Private/TedsDropOperationsModule.cpp
deleted file mode 100644
diff --git a/Engine/Plugins/Experimental/EditorDataStorageFeatures/Source/TedsDropOperations/Private/TedsDropOperationsModule.h b/Engine/Plugins/Experimental/EditorDataStorageFeatures/Source/TedsDropOperations/Private/TedsDropOperationsModule.h
deleted file mode 100644
diff --git a/Engine/Plugins/Experimental/EditorDataStorageFeatures/Source/TedsDropOperations/Private/TedsDropOperationsUtility.cpp b/Engine/Plugins/Experimental/EditorDataStorageFeatures/Source/TedsDropOperations/Private/TedsDropOperationsUtility.cpp
deleted file mode 100644
diff --git a/Engine/Plugins/Experimental/EditorDataStorageFeatures/Source/TedsDropOperations/Public/TedsDragDropOp.h b/Engine/Plugins/Experimental/EditorDataStorageFeatures/Source/TedsDropOperations/Public/TedsDragDropOp.h
deleted file mode 100644
diff --git a/Engine/Plugins/Experimental/EditorDataStorageFeatures/Source/TedsDropOperations/Public/TedsDropOperations.h b/Engine/Plugins/Experimental/EditorDataStorageFeatures/Source/TedsDropOperations/Public/TedsDropOperations.h
deleted file mode 100644
diff --git a/Engine/Plugins/Experimental/EditorDataStorageFeatures/Source/TedsDropOperations/Public/TedsDropOperationsColumns.h b/Engine/Plugins/Experimental/EditorDataStorageFeatures/Source/TedsDropOperations/Public/TedsDropOperationsColumns.h
deleted file mode 100644
diff --git a/Engine/Plugins/Experimental/EditorDataStorageFeatures/Source/TedsDropOperations/Public/TedsDropOperationsColumnsHelpers.h b/Engine/Plugins/Experimental/EditorDataStorageFeatures/Source/TedsDropOperations/Public/TedsDropOperationsColumnsHelpers.h
deleted file mode 100644
diff --git a/Engine/Plugins/Experimental/EditorDataStorageFeatures/Source/TedsDropOperations/Public/TedsDropOperationsFactory.h b/Engine/Plugins/Experimental/EditorDataStorageFeatures/Source/TedsDropOperations/Public/TedsDropOperationsFactory.h
deleted file mode 100644
diff --git a/Engine/Plugins/Experimental/EditorDataStorageFeatures/Source/TedsDropOperations/Public/TedsDropOperationsUtility.h b/Engine/Plugins/Experimental/EditorDataStorageFeatures/Source/TedsDropOperations/Public/TedsDropOperationsUtility.h
deleted file mode 100644
diff --git a/Engine/Plugins/Experimental/EditorDataStorageFeatures/Source/TedsDropOperations/TedsDropOperations.build.cs b/Engine/Plugins/Experimental/EditorDataStorageFeatures/Source/TedsDropOperations/TedsDropOperations.build.cs
deleted file mode 100644
diff --git a/Engine/Plugins/Experimental/EditorDataStorageFeatures/Source/TedsOperations/Private/DragAndDrop/DropOperationSystem.cpp b/Engine/Plugins/Experimental/EditorDataStorageFeatures/Source/TedsOperations/Private/DragAndDrop/DropOperationSystem.cpp
new file mode 100644
--- /dev/null
+++ b/Engine/Plugins/Experimental/EditorDataStorageFeatures/Source/TedsOperations/Private/DragAndDrop/DropOperationSystem.cpp
@@ -0,0 +1,61 @@
+// Copyright Epic Games, Inc. All Rights Reserved.
+
+#include "DragAndDrop/DropOperationSystem.h"
+
+#include "DragAndDrop/DropOperationInput.h"
+#include "Elements/Framework/TypedElementQueryBuilder.h"
+#include "TedsOperationColumns.h"
+#include "TedsQueryNode.h"
+
+UDropOperationSystem::UDropOperationSystem()
+{
+}
+
+UDropOperationSystem::~UDropOperationSystem()
+{
+}
+
+void UDropOperationSystem::RegisterTables(UE::Editor::DataStorage::ICoreProvider& Storage)
+{
+	using namespace UE::Editor::DataStorage::Operations;
+
+	OperationTable  = Storage.RegisterTable<FApplyColumn, FDropTag>("DropOperationTable");
+	InputTable      = Storage.RegisterTable<FDropTargetColumn, FSourceColumn>("DropOperationInputTable");
+	InputTableDescr = Storage.RegisterTable<FDropTargetColumn, FSourceColumn, FDescriptionColumn>("DropOperationInputTableDescr");
+}
+
+bool UDropOperationSystem::CreateInputRows(UE::Editor::DataStorage::FRowHandleArray& OutRows, UE::Editor::DataStorage::FRowHandleArrayView SourceRows,
+	UE::Editor::DataStorage::RowHandle TargetRow, bool bAddDescription) const
+{
+	using namespace UE::Editor::DataStorage;
+	using namespace UE::Editor::DataStorage::Operations;
+	
+	OutRows.Reset();
+	if (SourceRows.IsEmpty())
+	{
+		return true;
+	}
+	
+	TableHandle Table = bAddDescription ? InputTableDescr : InputTable;
+	if (Table == InvalidTableHandle)
+	{
+		return false;
+	}
+	
+	OutRows.Reserve(SourceRows.Num());
+	const RowHandle* SourceRowIt = SourceRows.begin();
+	return StoragePtr->BatchAddRow(Table, SourceRows.Num(), [this, &OutRows, &SourceRowIt, TargetRow](RowHandle Row)
+		{
+			StoragePtr->AddColumn(Row, FSourceColumn{ .Value = *SourceRowIt++ });
+			StoragePtr->AddColumn(Row, FDropTargetColumn{ .Value = TargetRow });
+			OutRows.Add(Row);
+		});
+}
+
+void UDropOperationSystem::GetQueryDescription(UE::Editor::DataStorage::FQueryDescription& OutDescription) const
+{
+	using namespace UE::Editor::DataStorage::Queries;
+	using namespace UE::Editor::DataStorage::Operations;
+
+	OutDescription = Select().Where().All<FApplyColumn, FDropTag>().Compile();
+}
diff --git a/Engine/Plugins/Experimental/EditorDataStorageFeatures/Source/TedsOperations/Private/DragAndDrop/TedsDragDropOp.cpp b/Engine/Plugins/Experimental/EditorDataStorageFeatures/Source/TedsOperations/Private/DragAndDrop/TedsDragDropOp.cpp
new file mode 100644
--- /dev/null
+++ b/Engine/Plugins/Experimental/EditorDataStorageFeatures/Source/TedsOperations/Private/DragAndDrop/TedsDragDropOp.cpp
@@ -0,0 +1,19 @@
+// Copyright Epic Games, Inc. All Rights Reserved
+// .
+#include "DragAndDrop/TedsDragDropOp.h"
+
+namespace UE::Editor::DataStorage::Widgets
+{
+	void FTedsDragDropOp::Init(FRowHandleArray InDraggedRows)
+	{
+		DraggedRows = MoveTemp(InDraggedRows);
+	}
+	
+	TSharedRef<FTedsDragDropOp> FTedsDragDropOp::New(FRowHandleArray InDraggedRows)
+	{
+		TSharedRef<FTedsDragDropOp> Operation = MakeShareable(new FTedsDragDropOp);
+		Operation->Init(MoveTemp(InDraggedRows));
+		Operation->Construct();
+		return Operation;
+	}
+}
diff --git a/Engine/Plugins/Experimental/EditorDataStorageFeatures/Source/TedsOperations/Private/DragAndDrop/TedsDragDropOpUtility.cpp b/Engine/Plugins/Experimental/EditorDataStorageFeatures/Source/TedsOperations/Private/DragAndDrop/TedsDragDropOpUtility.cpp
new file mode 100644
--- /dev/null
+++ b/Engine/Plugins/Experimental/EditorDataStorageFeatures/Source/TedsOperations/Private/DragAndDrop/TedsDragDropOpUtility.cpp
@@ -0,0 +1,59 @@
+// Copyright Epic Games, Inc. All Rights Reserved.
+
+#include "DragAndDrop/TedsDragDropOpUtility.h"
+
+#include "DataStorage/Features.h"
+#include "DataStorage/MapKey.h"
+#include "DragAndDrop/AssetDragDropOp.h"
+#include "DragAndDrop/ClassDragDropOp.h"
+#include "DragAndDrop/TedsDragDropOp.h"
+#include "Elements/Framework/TypedElementRowHandleArray.h"
+#include "Elements/Interfaces/TypedElementDataStorageInterface.h"
+#include "Input/DragAndDrop.h"
+
+namespace UE::Editor::DataStorage::DragAndDrop
+{
+
+void GetRowsFromData(FRowHandleArray& OutRows, const FDragDropOperation& InData)
+{
+	if (InData.IsOfType<Widgets::FTedsDragDropOp>())
+	{
+		OutRows = static_cast<const Widgets::FTedsDragDropOp&>(InData).GetRowArray();
+	}
+	else if (InData.IsOfType<FAssetDragDropOp>())
+	{
+		if (const ICoreProvider* Storage = GetDataStorageFeature<ICoreProvider>(StorageFeatureName))
+		{
+			for (const FAssetData& Element : static_cast<const FAssetDragDropOp&>(InData).GetAssets())
+			{
+				if (Element.IsValid())
+				{
+					RowHandle Row = Storage->LookupMappedRow(/*UE::Editor::AssetData::MappingDomain*/"Asset", FMapKey(Element.GetSoftObjectPath()));
+					if (Row != InvalidRowHandle)
+					{
+						OutRows.Add(Row);
+					}
+				}
+			}
+		}
+	}
+	else if (InData.IsOfType<FClassDragDropOp>())
+	{
+		if (const ICoreProvider* Storage = GetDataStorageFeature<ICoreProvider>(StorageFeatureName))
+		{
+			for (const TWeakObjectPtr<UClass>& ClassPtr : static_cast<const FClassDragDropOp&>(InData).ClassesToDrop)
+			{
+				if (const UClass* Class = ClassPtr.Get())
+				{
+					RowHandle Row = Storage->LookupMappedRow(/*UE::Editor::AssetData::MappingDomain*/"Asset", FMapKey(Class->GetPathName()));
+					if (Row != InvalidRowHandle)
+					{
+						OutRows.Add(Row);
+					}
+				}
+			}
+		}
+	}
+}
+
+}
\ No newline at end of file
diff --git a/Engine/Plugins/Experimental/EditorDataStorageFeatures/Source/TedsOperations/Private/DragAndDrop/Widgets/WidgetDropHandler.cpp b/Engine/Plugins/Experimental/EditorDataStorageFeatures/Source/TedsOperations/Private/DragAndDrop/Widgets/WidgetDropHandler.cpp
new file mode 100644
--- /dev/null
+++ b/Engine/Plugins/Experimental/EditorDataStorageFeatures/Source/TedsOperations/Private/DragAndDrop/Widgets/WidgetDropHandler.cpp
@@ -0,0 +1,334 @@
+// Copyright Epic Games, Inc. All Rights Reserved.
+
+#include "DragAndDrop/Widgets/WidgetDropHandler.h"
+
+#include "Containers/Ticker.h"
+#include "DataStorage/Features.h"
+#include "DragAndDrop/DecoratedDragDropOp.h"
+#include "DragAndDrop/DropOperationSystem.h"
+#include "DragAndDrop/TedsDragDropOpUtility.h"
+#include "Elements/Columns/TypedElementHandleColumn.h"
+#include "Elements/Framework/TypedElementSelectionSet.h"
+#include "Elements/Interfaces/TypedElementDataStorageInterface.h"
+#include "Input/DragAndDrop.h"
+#include "ScopedTransaction.h"
+#include "TedsOperationInput.h"
+
+#define LOCTEXT_NAMESPACE "WidgetDropHandler"
+
+namespace UE::Editor
+{
+
+namespace WidgetDropHandler_Private
+{
+	static void SelectRows(DataStorage::ICoreProvider& Storage, UTypedElementSelectionSet& SelectionSet, DataStorage::FRowHandleArray RowsToSelect)
+	{
+		using namespace UE::Editor::DataStorage;
+
+		if (RowsToSelect.IsEmpty())
+		{
+			return;
+		}
+		
+		// Select next tick, so that all columns have been set and updated.
+		FTSTicker::GetCoreTicker().AddTicker(FTickerDelegate::CreateLambda(
+			[Storage = &Storage, SelectionSet = TWeakObjectPtr(&SelectionSet), RowsToSelect = MoveTemp(RowsToSelect)](float)
+			{
+				if (TStrongObjectPtr<UTypedElementSelectionSet> SelectionSetPin = SelectionSet.Pin())
+				{
+					TArray<FTypedElementHandle> Elements;
+					Elements.Reserve(RowsToSelect.Num());
+					for (RowHandle Row : RowsToSelect.GetRows())
+					{
+						Compatibility::FTypedElementColumn* Column = Storage->GetColumn<Compatibility::FTypedElementColumn>(Row);
+						if (!Column || !Column->Handle)
+						{
+							continue;
+						}
+						
+						Elements.Add(Column->Handle);
+					}
+					SelectionSetPin->SetSelection(Elements, {});
+				}
+				return false;
+			}));
+	}
+}
+
+FWidgetDropHandler::FWidgetDropHandler(TWeakObjectPtr<UTypedElementSelectionSet> SelectionSet)
+	: SelectionSet(MoveTemp(SelectionSet))
+	, CurrentOperation(nullptr)
+	, bIsValidDrop(false)
+{
+}
+
+FWidgetDropHandler::~FWidgetDropHandler()
+{
+}
+
+bool FWidgetDropHandler::OnDragEnter(const FGeometry& Geometry, const FDragDropEvent& InputEvent)
+{
+	using namespace UE::Editor::DataStorage;
+	
+	if (TSharedPtr<FDragDropOperation> Operation = InputEvent.GetOperation())
+	{
+		if (ICoreProvider* Storage = GetMutableDataStorageFeature<ICoreProvider>(StorageFeatureName))
+		{
+			if (UDropOperationSystem* DropSystem = Storage->FindFactory<UDropOperationSystem>())
+			{
+				FDropResult Result = Start(*Storage, *DropSystem, Geometry, InputEvent, *Operation);
+				UpdateOperationDecorator(*Operation, Result);
+				return true;
+			}
+		}
+	}
+	return false;
+}
+
+FReply FWidgetDropHandler::OnDragOver(const FGeometry& Geometry, const FDragDropEvent& InputEvent)
+{
+	using namespace UE::Editor::DataStorage;
+
+	TSharedPtr<FDragDropOperation> Operation = InputEvent.GetOperation();
+	if (Operation && CurrentOperation == Operation.Get())
+	{
+		if (ICoreProvider* Storage = GetMutableDataStorageFeature<ICoreProvider>(StorageFeatureName))
+		{
+			if (UDropOperationSystem* DropSystem = Storage->FindFactory<UDropOperationSystem>())
+			{
+				if (TOptional<FDropResult> Result = Update(*Storage, *DropSystem, Geometry, InputEvent))
+				{
+					// Only update for changes.
+					UpdateOperationDecorator(*Operation, Result);
+				}
+			}
+		}
+		return FReply::Handled();
+	}
+	return FReply::Unhandled();
+}
+
+bool FWidgetDropHandler::OnDragLeave(const FDragDropEvent& InputEvent)
+{
+	using namespace UE::Editor::DataStorage;
+
+	TSharedPtr<FDragDropOperation> Operation = InputEvent.GetOperation();
+	if (Operation && CurrentOperation == Operation.Get())
+	{
+		if (ICoreProvider* Storage = GetMutableDataStorageFeature<ICoreProvider>(StorageFeatureName))
+		{
+			if (UDropOperationSystem* DropSystem = Storage->FindFactory<UDropOperationSystem>())
+			{
+				Stop(*Storage, *DropSystem);
+				UpdateOperationDecorator(*Operation, {});
+			}
+		}
+		return true;
+	}
+	return false;
+}
+
+FReply FWidgetDropHandler::OnDrop(const FGeometry& Geometry, const FDragDropEvent& InputEvent)
+{
+	using namespace UE::Editor::DataStorage;
+
+	TSharedPtr<FDragDropOperation> Operation = InputEvent.GetOperation();
+	if (Operation && CurrentOperation == Operation.Get())
+	{
+		ICoreProvider* Storage = GetMutableDataStorageFeature<ICoreProvider>(StorageFeatureName);
+		if (Storage && bIsValidDrop)
+		{
+			if (UDropOperationSystem* DropSystem = Storage->FindFactory<UDropOperationSystem>())
+			{
+				Drop(*Storage, *DropSystem, Geometry, InputEvent);
+				UpdateOperationDecorator(*Operation, {});
+			}
+		}
+		return FReply::Handled();
+	}
+	return FReply::Unhandled();
+}
+
+FDragDropOperation* FWidgetDropHandler::GetCurrentOperation() const
+{
+	return CurrentOperation;
+}
+
+bool FWidgetDropHandler::PrepareInput(DataStorage::ICoreProvider& Storage, UDropOperationSystem& DropSystem, const FDragDropOperation& Operation)
+{
+	using namespace UE::Editor::DataStorage;
+	
+	FRowHandleArray SourceRows;
+	DragAndDrop::GetRowsFromData(SourceRows, Operation);
+	SourceRows.Remove(InvalidRowHandle);
+	
+	return PrepareInput(Storage, DropSystem, SourceRows.GetRows());
+}
+
+bool FWidgetDropHandler::PrepareInput(DataStorage::ICoreProvider& Storage, UDropOperationSystem& DropSystem,
+	DataStorage::FRowHandleArrayView SourceRows)
+{
+	using namespace UE::Editor::DataStorage;
+	checkf(InputRows.IsEmpty(), TEXT("Expected previous operation to be cleared."));
+
+	return DropSystem.CreateInputRows(InputRows, SourceRows, InvalidRowHandle, true) && !InputRows.IsEmpty();
+}
+
+FWidgetDropHandler::FDropResult FWidgetDropHandler::ResetOperations(DataStorage::ICoreProvider& Storage, UDropOperationSystem& DropSystem)
+{
+	using namespace UE::Editor::DataStorage;
+	using namespace UE::Editor::DataStorage::Operations;
+
+	// We only accept the drop if all input rows are valid.
+	bIsValidDrop = DropSystem.Test(InputRows.GetRows()) == InputRows.Num();
+
+	FText ToolTip;
+	for (RowHandle InputRow : InputRows.GetRows())
+	{
+		// Find an error text when we failed and a success text if we succeeded.
+		if (bIsValidDrop != Storage.HasColumns<FTestResultTag>(InputRow))
+		{
+			continue;
+		}
+		
+		ToolTip = Utilities::GetDescription(Storage, InputRow);
+		if (!ToolTip.IsEmpty())
+		{
+			break;
+		}
+	}
+	
+	if (!bIsValidDrop && ToolTip.IsEmpty())
+	{
+		ToolTip = LOCTEXT("Drop_Unsupported", "Drop operation is unsupported.");
+	}
+	
+	return { bIsValidDrop, MoveTemp(ToolTip) };
+}
+
+FWidgetDropHandler::EUpdateParameters FWidgetDropHandler::UpdateParameters(DataStorage::ICoreProvider& Storage, const FGeometry& Geometry,
+	const FDragDropEvent& InputEvent)
+{
+	return EUpdateParameters::NoChanges;
+}
+
+void FWidgetDropHandler::SetRowsFromParameters(DataStorage::ICoreProvider& Storage, DataStorage::FRowHandleArrayView Rows) const
+{
+}
+
+FWidgetDropHandler::FDropResult FWidgetDropHandler::Start(DataStorage::ICoreProvider& Storage, UDropOperationSystem& DropSystem,
+                                                          const FGeometry& Geometry, const FDragDropEvent& InputEvent, FDragDropOperation& Operation)
+{
+	using namespace UE::Editor::DataStorage;
+
+	if (CurrentOperation)
+	{
+		Stop(Storage, DropSystem);
+	}
+	
+	CurrentOperation = &Operation;
+	
+	if (!PrepareInput(Storage, DropSystem, Operation))
+	{
+		Stop(Storage, DropSystem);
+		return { false, LOCTEXT("Drop_Unsupported", "Drop operation is unsupported.") };
+	}
+	(void)UpdateParameters(Storage, Geometry, InputEvent);
+	SetRowsFromParameters(Storage, InputRows.GetRows());
+	return ResetOperations(Storage, DropSystem);
+}
+
+TOptional<FWidgetDropHandler::FDropResult> FWidgetDropHandler::Update(DataStorage::ICoreProvider& Storage, UDropOperationSystem& DropSystem,
+	const FGeometry& Geometry, const FDragDropEvent& InputEvent)
+{
+	using namespace UE::Editor::DataStorage;
+	check(CurrentOperation != nullptr);
+
+	EUpdateParameters Result = UpdateParameters(Storage, Geometry, InputEvent);
+	if (Result == EUpdateParameters::ResetOperations)
+	{
+		SetRowsFromParameters(Storage, InputRows.GetRows()); // Only update the input rows when we need to.
+		return ResetOperations(Storage, DropSystem);
+	}	
+	return {};
+}
+	
+void FWidgetDropHandler::Stop(DataStorage::ICoreProvider& Storage, UDropOperationSystem& DropSystem)
+{
+	using namespace UE::Editor::DataStorage;
+	check(CurrentOperation != nullptr);
+
+	DropSystem.RemoveInputRows(InputRows.GetRows());
+	
+	InputRows.Empty();
+	CurrentOperation = nullptr;
+	bIsValidDrop = false;	
+}
+
+void FWidgetDropHandler::Drop(DataStorage::ICoreProvider& Storage, UDropOperationSystem& DropSystem, const FGeometry& Geometry,
+		const FDragDropEvent& InputEvent)
+{
+	using namespace UE::Editor::DataStorage;
+	using namespace UE::Editor::DataStorage::Operations;
+	check(bIsValidDrop);
+
+	// Update one last time.
+	(void)UpdateParameters(Storage, Geometry, InputEvent);
+	SetRowsFromParameters(Storage, InputRows.GetRows());
+	{
+		FScopedTransaction Transaction(LOCTEXT("Drop_Transaction", "Drop on widget."));
+		(void)DropSystem.Apply(InputRows.GetRows());
+	}
+
+	UTypedElementSelectionSet* SelectionSetPtr = SelectionSet.Get();
+	FRowHandleArray RowsToSelect;
+	if (SelectionSetPtr)
+	{
+		RowsToSelect.Reserve(InputRows.Num());
+		for (RowHandle InputRow : InputRows.GetRows())
+		{
+			if (const FResultColumn* Column = Storage.GetColumn<FResultColumn>(InputRow))
+			{
+				RowsToSelect.Append(Column->Value.Created.GetRows());
+			}
+		}
+	}
+	
+	Stop(Storage, DropSystem);
+
+	if (SelectionSetPtr)
+	{
+		WidgetDropHandler_Private::SelectRows(Storage, *SelectionSetPtr, MoveTemp(RowsToSelect));
+	}
+}
+
+void FWidgetDropHandler::UpdateOperationDecorator(FDragDropOperation& Operation, TOptional<FDropResult> Result) const
+{
+	Operation.SetDecoratorVisibility(true);
+	if (Operation.IsOfType<FDecoratedDragDropOp>())
+	{
+		static constexpr TCHAR BrushName_Ok[] = TEXT("Graph.ConnectorFeedback.OK");
+		static constexpr TCHAR BrushName_Error[] = TEXT("Graph.ConnectorFeedback.Error");
+		FDecoratedDragDropOp& DecoratedOp = static_cast<FDecoratedDragDropOp&>(Operation);
+		if (Result)
+		{
+			DecoratedOp.SetToolTip(Result->Tooltip, FAppStyle::GetBrush(Result->bCanDrop ? BrushName_Ok : BrushName_Error));
+		}
+		else
+		{
+			DecoratedOp.ResetToDefaultToolTip();
+		}
+	}
+	else if (Result)
+	{
+		Operation.SetCursorOverride(Result->bCanDrop ? EMouseCursor::Default : EMouseCursor::SlashedCircle);
+	}
+	else
+	{
+		Operation.SetCursorOverride({});
+	}
+}
+	
+}
+
+#undef LOCTEXT_NAMESPACE
diff --git a/Engine/Plugins/Experimental/EditorDataStorageFeatures/Source/TedsOperations/Private/TedsOperationSystem.cpp b/Engine/Plugins/Experimental/EditorDataStorageFeatures/Source/TedsOperations/Private/TedsOperationSystem.cpp
new file mode 100644
--- /dev/null
+++ b/Engine/Plugins/Experimental/EditorDataStorageFeatures/Source/TedsOperations/Private/TedsOperationSystem.cpp
@@ -0,0 +1,287 @@
+// Copyright Epic Games, Inc. All Rights Reserved.
+
+#include "TedsOperationSystem.h"
+
+#include "Elements/Framework/TypedElementQueryBuilder.h"
+#include "Elements/Interfaces/TypedElementDataStorageInterface.h"
+#include "TedsOperationColumns.h"
+#include "TedsOperationInput.h"
+#include "TedsQueryNode.h"
+#include "TedsQueryStackExecutor.h"
+#include "TedsRowQueryResultsNode.h"
+
+bool UOperationSystem::FDefaultLess::operator()(const UE::Editor::DataStorage::ICoreProvider& Storage, UE::Editor::DataStorage::RowHandle A,
+	UE::Editor::DataStorage::RowHandle B) const
+{
+	using namespace UE::Editor::DataStorage;
+	// First sort by priority column.
+	{
+		const Operations::FPriorityColumn* ColumnA = Storage.GetColumn<Operations::FPriorityColumn>(A);
+		const Operations::FPriorityColumn* ColumnB = Storage.GetColumn<Operations::FPriorityColumn>(B);
+		int64 PriorityA = ColumnA ? ColumnA->Value : 0;
+		int64 PriorityB = ColumnB ? ColumnB->Value : 0;
+		if (PriorityA != PriorityB)
+		{
+			return PriorityA > PriorityB; // Highest priority should be the first element.
+		}
+	}
+	// Second sort by name column.
+	{
+		const Operations::FNameColumn* ColumnA = Storage.GetColumn<Operations::FNameColumn>(A);
+		const Operations::FNameColumn* ColumnB = Storage.GetColumn<Operations::FNameColumn>(B);
+		FName NameA = ColumnA ? ColumnA->Value : FName();
+		FName NameB = ColumnB ? ColumnB->Value : FName();
+		if (int32 Result = NameA.Compare(NameB); Result != 0)
+		{
+			return Result < 0;
+		}
+	}
+	// Third sort by row handle.
+	return A < B;
+}
+
+UOperationSystem::UOperationSystem()
+	: StoragePtr(nullptr)
+	, OperationTable(UE::Editor::DataStorage::InvalidTableHandle)
+	, InputTable(UE::Editor::DataStorage::InvalidTableHandle)
+	, InputTableDescr(UE::Editor::DataStorage::InvalidTableHandle)
+{
+}
+
+UOperationSystem::~UOperationSystem()
+{
+}
+
+void UOperationSystem::PreRegister(UE::Editor::DataStorage::ICoreProvider& Storage)
+{
+	StoragePtr = &Storage;
+}
+
+void UOperationSystem::RegisterTables(UE::Editor::DataStorage::ICoreProvider& Storage)
+{
+	using namespace UE::Editor::DataStorage::Operations;
+	
+	OperationTable  = Storage.RegisterTable<FApplyColumn>("OperationTable");
+	InputTable      = Storage.RegisterTable<FSourceColumn>("OperationInputTable");
+	InputTableDescr = Storage.RegisterTable<FSourceColumn, FDescriptionColumn>("OperationInputTableDescr");
+}
+
+void UOperationSystem::RegisterQueries(UE::Editor::DataStorage::ICoreProvider& Storage)
+{
+	using namespace UE::Editor::DataStorage;
+	using namespace UE::Editor::DataStorage::Queries;
+	using namespace UE::Editor::DataStorage::QueryStack;
+
+	FQueryDescription QueryDescription;
+	GetQueryDescription(QueryDescription);
+
+	QueryNode   = MakeShared<FQueryNode>(Storage, MoveTemp(QueryDescription));
+	ResultsNode = MakeShared<FRowQueryResultsNode>(Storage, QueryNode, FRowQueryResultsNode::ESyncFlags::RefreshOnUpdate);
+	Executor    = MakeShared<FExplicitUpdateExecutor>("Operations", ResultsNode);
+}
+
+UE::Editor::DataStorage::RowHandle UOperationSystem::AddOperation(FName Name,
+	TNotNull<UE::Editor::DataStorage::Operations::FApplyCallback> Operation,
+	UE::Editor::DataStorage::Operations::FTestCallback Test,
+	UE::Editor::DataStorage::Operations::FProbeCallback Probe, int64 Priority)
+{
+	using namespace UE::Editor::DataStorage;
+	using namespace UE::Editor::DataStorage::Operations;
+
+	if (OperationTable == InvalidTableHandle)
+	{
+		return InvalidRowHandle;
+	}
+
+	RowHandle Row = StoragePtr->AddRow(OperationTable);
+	if (Row != InvalidRowHandle)
+	{
+		StoragePtr->AddColumn(Row, FApplyColumn{ .Callback = MoveTemp(Operation) });
+		if (Name != NAME_None)
+		{
+			StoragePtr->AddColumn(Row, FNameColumn{ .Value = MoveTemp(Name) });
+		}
+		if (Test)
+		{
+			StoragePtr->AddColumn(Row, FTestColumn{ .Callback = MoveTemp(Test) });
+		}
+		if (Probe)
+		{
+			StoragePtr->AddColumn(Row, FProbeColumn{ .Callback = MoveTemp(Probe) });
+		}
+		if (Priority != DefaultPriority)
+		{
+			StoragePtr->AddColumn(Row, FPriorityColumn{ .Value = Priority });
+		}
+	}
+	return Row;	
+}
+
+void UOperationSystem::GetOperations(TArray<UE::Editor::DataStorage::RowHandle>& OutRows, UE::Editor::DataStorage::RowHandle InputRow,
+	const FFilter& Filter, const FLess& Less) const
+{
+	using namespace UE::Editor::DataStorage;
+	
+	OutRows.Empty();
+	if (!Executor)
+	{
+		return;
+	}
+	
+	Executor->Update();
+	FRowHandleArrayView AllRows = ResultsNode->GetRows();
+	if (AllRows.IsEmpty())
+	{
+		return;
+	}
+
+	for (RowHandle Row : AllRows)
+	{
+		if (Row == InvalidRowHandle)
+		{
+			continue;
+		}
+
+		if (Filter && !Filter(*StoragePtr, Row))
+		{
+			continue;
+		}
+
+		const Operations::FProbeColumn* ProbeColumn = StoragePtr->GetColumn<Operations::FProbeColumn>(Row);
+		if (ProbeColumn && ProbeColumn->Callback && !ProbeColumn->Callback(*StoragePtr, InputRow))
+		{
+			continue;
+		}
+
+		OutRows.Add(Row);
+	}
+
+	if (Less)
+	{
+		Algo::Sort(OutRows, [this, &Less](RowHandle A, RowHandle B) { return Less(*StoragePtr, A, B); });
+	}
+}
+
+int32 UOperationSystem::Test(UE::Editor::DataStorage::FRowHandleArrayView InputRows, const FFilter& Filter, const FLess& Less) const
+{
+	using namespace UE::Editor::DataStorage;
+	using namespace UE::Editor::DataStorage::Operations;
+
+	// @todo: Improve for batch processing.
+
+	int32 SuccessNum = 0;	
+
+	TArray<RowHandle> OperationRows;
+	for (RowHandle InputRow : InputRows)
+	{
+		bool bResult = false;
+		OperationRows.Reset();
+		StoragePtr->RemoveColumn<FTestResultTag>(InputRow);
+		
+		GetOperations(OperationRows, InputRow, Filter, Less);
+		for (RowHandle OperationRow : OperationRows)
+		{
+			// No column equals validity.
+			const FTestColumn* TestColumn = StoragePtr->GetColumn<FTestColumn>(OperationRow);
+			bResult = !TestColumn || !TestColumn->Callback || TestColumn->Callback(*StoragePtr, InputRow);
+			if (bResult)
+			{
+				break;
+			}
+		}
+
+		if (bResult)
+		{
+			StoragePtr->AddColumn<FTestResultTag>(InputRow);
+			++SuccessNum;
+		}
+		else
+		{
+			StoragePtr->RemoveColumn<FTestResultTag>(InputRow);
+		}
+	}
+
+	return SuccessNum;
+}
+
+int32 UOperationSystem::Apply(UE::Editor::DataStorage::FRowHandleArrayView InputRows, const FFilter& Filter, const FLess& Less) const
+{
+	using namespace UE::Editor::DataStorage;
+	using namespace UE::Editor::DataStorage::Operations;
+	
+	// @todo: Improve for batch processing.
+
+	int32 SuccessNum = 0;
+	
+	TArray<RowHandle> OperationRows;
+	for (RowHandle InputRow : InputRows)
+	{
+		TOptional<FResult> Result;
+		OperationRows.Reset();
+		
+		GetOperations(OperationRows, InputRow, Filter, Less);
+		for (RowHandle OperationRow : OperationRows)
+		{
+			if (const FApplyColumn* ApplyColumn = StoragePtr->GetColumn<FApplyColumn>(OperationRow))
+			{
+				Result = ApplyColumn->Callback(*StoragePtr, InputRow);
+				if (Result)
+				{
+					break;
+				}
+			}
+		}
+
+		if (Result)
+		{
+			StoragePtr->AddColumn(InputRow, FResultColumn{ .Value = MoveTemp(*Result) });
+			++SuccessNum;
+		}
+		else
+		{
+			StoragePtr->RemoveColumn<FResultColumn>(InputRow);
+		}
+	}
+
+	return SuccessNum;
+}
+
+bool UOperationSystem::CreateInputRows(UE::Editor::DataStorage::FRowHandleArray& OutRows, UE::Editor::DataStorage::FRowHandleArrayView SourceRows,
+	bool bAddDescription) const
+{
+	using namespace UE::Editor::DataStorage;
+	using namespace UE::Editor::DataStorage::Operations;
+	
+	OutRows.Reset();
+	if (SourceRows.IsEmpty())
+	{
+		return true;
+	}
+	
+	TableHandle Table = bAddDescription ? InputTableDescr : InputTable;
+	if (Table == InvalidTableHandle)
+	{
+		return false;
+	}
+	
+	OutRows.Reserve(SourceRows.Num());
+	const RowHandle* SourceRowIt = SourceRows.begin();
+	return StoragePtr->BatchAddRow(Table, SourceRows.Num(), [this, &SourceRowIt, &OutRows](RowHandle Row)
+		{
+			StoragePtr->AddColumn(Row, FSourceColumn{ .Value = *SourceRowIt++ });
+			OutRows.Add(Row);
+		});
+}
+
+void UOperationSystem::RemoveInputRows(UE::Editor::DataStorage::FRowHandleArrayView RowsToRemove) const
+{
+	StoragePtr->BatchRemoveRows(RowsToRemove);
+}
+
+void UOperationSystem::GetQueryDescription(UE::Editor::DataStorage::FQueryDescription& OutDescription) const
+{
+	using namespace UE::Editor::DataStorage::Queries;
+	using namespace UE::Editor::DataStorage::Operations;
+
+	OutDescription = Select().Where().All<FApplyColumn>().Compile();
+}
diff --git a/Engine/Plugins/Experimental/EditorDataStorageFeatures/Source/TedsOperations/Private/TedsOperationsFactory.cpp b/Engine/Plugins/Experimental/EditorDataStorageFeatures/Source/TedsOperations/Private/TedsOperationsFactory.cpp
deleted file mode 100644
diff --git a/Engine/Plugins/Experimental/EditorDataStorageFeatures/Source/TedsOperations/Public/DragAndDrop/DropOperationInput.h b/Engine/Plugins/Experimental/EditorDataStorageFeatures/Source/TedsOperations/Public/DragAndDrop/DropOperationInput.h
new file mode 100644
--- /dev/null
+++ b/Engine/Plugins/Experimental/EditorDataStorageFeatures/Source/TedsOperations/Public/DragAndDrop/DropOperationInput.h
@@ -0,0 +1,92 @@
+// Copyright Epic Games, Inc. All Rights Reserved.
+
+#pragma once
+
+#include "DataStorage/CommonTypes.h"
+#include "TedsOperationInput.h"
+
+#include "DropOperationInput.generated.h"
+
+USTRUCT(meta = (DisplayName = "Drop Operation Target"))
+struct FDropOperationTargetColumn final : public FEditorDataStorageColumn
+{
+	GENERATED_BODY()
+
+	UE::Editor::DataStorage::RowHandle Value;
+};
+
+USTRUCT(meta = (DisplayName = "Drop Operation Name"))
+struct FDropOperationNameColumn final : public FEditorDataStorageColumn
+{
+	GENERATED_BODY()
+
+	FName Value;
+};
+
+USTRUCT(meta = (DisplayName = "Drop Operation Transform"))
+struct FDropOperationTransformColumn final : public FEditorDataStorageColumn
+{
+	GENERATED_BODY()
+
+	FTransform Value;
+};
+
+USTRUCT(meta = (DisplayName = "Drop Operation Folder"))
+struct FDropOperationFolderColumn final : public FEditorDataStorageColumn
+{
+	GENERATED_BODY()
+
+	FName Path;
+};
+
+USTRUCT(meta = (DisplayName = "Drop Operation Preview"))
+struct FDropOperationPreviewTag final : public FEditorDataStorageTag
+{
+	GENERATED_BODY()
+};
+
+namespace UE::Editor::DataStorage::Operations
+{
+/** Holds the target row of the drop operation (e.g. a level). */	
+using FDropTargetColumn = FDropOperationTargetColumn;
+/** Informs the operation to create a preview. */
+using FDropPreviewTag = FDropOperationPreviewTag;
+/** Holds the name that should be used by the operation. */
+using FDropNameColumn = FDropOperationNameColumn;
+/** Holds the transform that should be used by the operation. */
+using FDropTransformColumn = FDropOperationTransformColumn;
+/** Holds the folder path that should be used by the operation. */
+using FDropFolderColumn = FDropOperationFolderColumn;
+
+namespace Utilities
+{
+	inline RowHandle GetDropTargetRow(const ICoreProvider& Storage, RowHandle InputRow)
+	{
+		const FDropTargetColumn* Column = Storage.GetColumn<FDropTargetColumn>(InputRow);
+		return Column ? Column->Value : InvalidRowHandle;
+	}
+	
+	inline FName GetDropName(const ICoreProvider& Storage, RowHandle InputRow)
+	{
+		const FDropNameColumn* Column = Storage.GetColumn<FDropNameColumn>(InputRow);
+		return Column ? Column->Value : NAME_None;
+	}
+	
+	inline FTransform GetDropTransform(const ICoreProvider& Storage, RowHandle InputRow)
+	{
+		const FDropTransformColumn* Column = Storage.GetColumn<FDropTransformColumn>(InputRow);
+		return Column ? Column->Value : FTransform::Identity;
+	}
+	
+	inline FName GetDropFolderPath(const ICoreProvider& Storage, RowHandle InputRow)
+	{
+		const FDropFolderColumn* Column = Storage.GetColumn<FDropFolderColumn>(InputRow);
+		return Column ? Column->Path : NAME_None;
+	}
+
+	inline bool HasDropPreviewTag(const ICoreProvider& Storage, RowHandle InputRow)
+	{
+		return Storage.HasColumns<FDropPreviewTag>(InputRow);
+	}
+}
+}
diff --git a/Engine/Plugins/Experimental/EditorDataStorageFeatures/Source/TedsOperations/Public/DragAndDrop/DropOperationSystem.h b/Engine/Plugins/Experimental/EditorDataStorageFeatures/Source/TedsOperations/Public/DragAndDrop/DropOperationSystem.h
new file mode 100644
--- /dev/null
+++ b/Engine/Plugins/Experimental/EditorDataStorageFeatures/Source/TedsOperations/Public/DragAndDrop/DropOperationSystem.h
@@ -0,0 +1,49 @@
+// Copyright Epic Games, Inc. All Rights Reserved.
+
+#pragma once
+
+#include "CoreMinimal.h"
+#include "DataStorage/ScopedRowHandle.h"
+#include "TedsOperationSystem.h"
+
+#include "DropOperationSystem.generated.h"
+
+#define UE_API TEDSOPERATIONS_API
+
+USTRUCT(meta = (DisplayName = "Drop Operation"))
+struct FDropOperationTag final : public FEditorDataStorageTag
+{
+	GENERATED_BODY()
+};
+
+/**
+ * Operations system for drop operations.
+ * Only operation rows with the "Drop Tag" will be considered by this system.
+ */
+UCLASS(MinimalAPI)
+class UDropOperationSystem : public UOperationSystem
+{
+	GENERATED_BODY()
+public:
+	UE_API UDropOperationSystem();
+	UE_API virtual ~UDropOperationSystem() override;
+	
+	UE_API virtual void RegisterTables(UE::Editor::DataStorage::ICoreProvider& Storage) override;
+	
+	/** Utility method to create a batch of rows for the given source rows. */
+	UE_API bool CreateInputRows(UE::Editor::DataStorage::FRowHandleArray& OutRows, UE::Editor::DataStorage::FRowHandleArrayView SourceRows,
+		UE::Editor::DataStorage::RowHandle TargetRow, bool bAddDescription = false) const;
+
+protected:
+	UE_API virtual void GetQueryDescription(UE::Editor::DataStorage::FQueryDescription& OutDescription) const;
+};
+
+namespace UE::Editor::DataStorage::Operations
+{
+	
+/** Tag to identify this operation as a drop operation. UDropOperationSystem will only consider rows with this tag. */
+using FDropTag = FDropOperationTag;
+	
+}
+
+#undef UE_API
diff --git a/Engine/Plugins/Experimental/EditorDataStorageFeatures/Source/TedsOperations/Public/DragAndDrop/TedsDragDropOp.h b/Engine/Plugins/Experimental/EditorDataStorageFeatures/Source/TedsOperations/Public/DragAndDrop/TedsDragDropOp.h
new file mode 100644
--- /dev/null
+++ b/Engine/Plugins/Experimental/EditorDataStorageFeatures/Source/TedsOperations/Public/DragAndDrop/TedsDragDropOp.h
@@ -0,0 +1,33 @@
+// Copyright Epic Games, Inc. All Rights Reserved.
+
+#pragma once
+
+#include "DragAndDrop/DecoratedDragDropOp.h"
+#include "Elements/Framework/TypedElementRowHandleArray.h"
+#include "Input/DragAndDrop.h"
+
+#define UE_API TEDSOPERATIONS_API
+
+namespace UE::Editor::DataStorage::Widgets
+{
+	// TEDS UI TODO: Instead of using a decorated operation use a default one and allow decorator customization using TEDS UI
+	class FTedsDragDropOp : public FDecoratedDragDropOp
+	{
+	public:
+		DRAG_DROP_OPERATOR_TYPE(FTedsDragDropOp, FDecoratedDragDropOp)
+
+		UE_API static TSharedRef<FTedsDragDropOp> New(FRowHandleArray InDraggedRows);
+
+		FRowHandleArrayView    GetRows() const     { return DraggedRows.GetRows(); }
+		FRowHandleArray&       GetRowArray()       { return DraggedRows; }
+		const FRowHandleArray& GetRowArray() const { return DraggedRows; }
+		
+	protected:
+		void Init(FRowHandleArray InDraggedRows);
+		
+	protected:
+		FRowHandleArray DraggedRows;
+	};
+}
+
+#undef UE_API
\ No newline at end of file
diff --git a/Engine/Plugins/Experimental/EditorDataStorageFeatures/Source/TedsOperations/Public/DragAndDrop/TedsDragDropOpUtility.h b/Engine/Plugins/Experimental/EditorDataStorageFeatures/Source/TedsOperations/Public/DragAndDrop/TedsDragDropOpUtility.h
new file mode 100644
--- /dev/null
+++ b/Engine/Plugins/Experimental/EditorDataStorageFeatures/Source/TedsOperations/Public/DragAndDrop/TedsDragDropOpUtility.h
@@ -0,0 +1,22 @@
+// Copyright Epic Games, Inc. All Rights Reserved.
+
+#pragma once
+
+#include "CoreTypes.h"
+
+#define UE_API TEDSOPERATIONS_API
+
+class FDragDropOperation;
+namespace UE::Editor::DataStorage
+{
+class FRowHandleArray;
+}
+
+namespace UE::Editor::DataStorage::DragAndDrop
+{
+
+UE_API void GetRowsFromData(FRowHandleArray& OutRows, const FDragDropOperation& InData);
+
+}
+
+#undef UE_API
diff --git a/Engine/Plugins/Experimental/EditorDataStorageFeatures/Source/TedsOperations/Public/DragAndDrop/Widgets/WidgetDropHandler.h b/Engine/Plugins/Experimental/EditorDataStorageFeatures/Source/TedsOperations/Public/DragAndDrop/Widgets/WidgetDropHandler.h
new file mode 100644
--- /dev/null
+++ b/Engine/Plugins/Experimental/EditorDataStorageFeatures/Source/TedsOperations/Public/DragAndDrop/Widgets/WidgetDropHandler.h
@@ -0,0 +1,82 @@
+// Copyright Epic Games, Inc. All Rights Reserved.
+
+#pragma once
+
+#include "DataStorage/Handles.h"
+#include "DragAndDrop/IWidgetDropHandler.h"
+#include "Elements/Framework/TypedElementRowHandleArray.h"
+#include "Elements/Framework/TypedElementRowHandleArrayView.h"
+#include "Misc/Optional.h"
+
+class UTypedElementSelectionSet;
+class UDropOperationSystem;
+namespace UE::Editor::DataStorage
+{
+class ICoreProvider;
+}
+
+#define UE_API TEDSOPERATIONS_API
+
+namespace UE::Editor
+{
+
+class FWidgetDropHandler : public IWidgetDropHandler
+{
+public:
+	struct FDropResult
+	{
+		bool  bCanDrop = false;
+		FText Tooltip  = FText();
+	};
+
+protected:
+	enum class EUpdateParameters
+	{
+		NoChanges,       // Nothing has changed.
+		MinorChanges,    // Some changes to input, but no need to reset operations.
+		ResetOperations, // Input has major changes, reset operations.
+	};
+
+public:
+	UE_API FWidgetDropHandler(TWeakObjectPtr<UTypedElementSelectionSet> SelectionSet);
+	UE_API virtual ~FWidgetDropHandler() override;
+
+	UE_API virtual bool                OnDragEnter(const FGeometry& Geometry, const FDragDropEvent& InputEvent) override;
+	UE_API virtual FReply              OnDragOver(const FGeometry& Geometry, const FDragDropEvent& InputEvent) override;
+	UE_API virtual bool                OnDragLeave(const FDragDropEvent& InputEvent) override;
+	UE_API virtual FReply              OnDrop(const FGeometry& Geometry, const FDragDropEvent& InputEvent) override;
+	UE_API virtual FDragDropOperation* GetCurrentOperation() const override;
+
+protected:
+	/** Creates input rows for the given drag&drop input event. */
+	UE_API virtual bool PrepareInput(DataStorage::ICoreProvider& Storage, UDropOperationSystem& DropSystem, const FDragDropOperation& Operation);
+	/** Creates input rows for the given drag&drop input event. */
+	UE_API virtual bool PrepareInput(DataStorage::ICoreProvider& Storage, UDropOperationSystem& DropSystem, DataStorage::FRowHandleArrayView SourceRows);
+	/** Finds suitable operations for the given input rows and sets 'bIsValidDrop'. */
+	UE_API virtual FDropResult ResetOperations(DataStorage::ICoreProvider& Storage, UDropOperationSystem& DropSystem);
+	
+	/** Updates all input rows for the updated input parameters. */
+	UE_API virtual EUpdateParameters UpdateParameters(DataStorage::ICoreProvider& Storage, const FGeometry& Geometry, const FDragDropEvent& InputEvent);
+	UE_API virtual void SetRowsFromParameters(DataStorage::ICoreProvider& Storage, DataStorage::FRowHandleArrayView Rows) const;
+	
+	UE_API virtual FDropResult Start(DataStorage::ICoreProvider& Storage, UDropOperationSystem& DropSystem, const FGeometry& Geometry,
+		const FDragDropEvent& InputEvent, FDragDropOperation& Operation);
+	UE_API virtual TOptional<FDropResult> Update(DataStorage::ICoreProvider& Storage, UDropOperationSystem& DropSystem, const FGeometry& Geometry,
+		const FDragDropEvent& InputEvent);
+	UE_API virtual void Stop(DataStorage::ICoreProvider& Storage, UDropOperationSystem& DropSystem);
+	UE_API virtual void Drop(DataStorage::ICoreProvider& Storage, UDropOperationSystem& DropSystem, const FGeometry& Geometry,
+		const FDragDropEvent& InputEvent);
+
+	UE_API virtual void UpdateOperationDecorator(FDragDropOperation& Operation, TOptional<FDropResult> Result) const;
+
+protected:
+	TWeakObjectPtr<UTypedElementSelectionSet> SelectionSet;
+
+	FDragDropOperation*          CurrentOperation;
+	DataStorage::FRowHandleArray InputRows;
+	bool                         bIsValidDrop;
+};
+
+}
+
+#undef UE_API
\ No newline at end of file
diff --git a/Engine/Plugins/Experimental/EditorDataStorageFeatures/Source/TedsOperations/Public/TedsOperationCallbacks.h b/Engine/Plugins/Experimental/EditorDataStorageFeatures/Source/TedsOperations/Public/TedsOperationCallbacks.h
new file mode 100644
--- /dev/null
+++ b/Engine/Plugins/Experimental/EditorDataStorageFeatures/Source/TedsOperations/Public/TedsOperationCallbacks.h
@@ -0,0 +1,23 @@
+// Copyright Epic Games, Inc. All Rights Reserved.
+
+#pragma once
+
+#include "DataStorage/Handles.h"
+#include "Misc/Optional.h"
+#include "TedsOperationResult.h"
+#include "Templates/Function.h"
+
+namespace UE::Editor::DataStorage
+{
+class ICoreProvider;
+}
+
+namespace UE::Editor::DataStorage::Operations
+{
+	/** Probe call to test superficially if the given input is potentially acceptable for the operation. */
+	using FProbeCallback = TFunction<bool(const ICoreProvider&, RowHandle)>;
+	/** Test call to query if the operation for the given input can be executed. */
+	using FTestCallback  = TFunction<bool(ICoreProvider&, RowHandle)>;
+	/** Apply call to execute the operation for the given input. */
+	using FApplyCallback = TFunction<TOptional<FResult>(ICoreProvider&, RowHandle)>;
+}
diff --git a/Engine/Plugins/Experimental/EditorDataStorageFeatures/Source/TedsOperations/Public/TedsOperationColumns.h b/Engine/Plugins/Experimental/EditorDataStorageFeatures/Source/TedsOperations/Public/TedsOperationColumns.h
new file mode 100644
--- /dev/null
+++ b/Engine/Plugins/Experimental/EditorDataStorageFeatures/Source/TedsOperations/Public/TedsOperationColumns.h
@@ -0,0 +1,62 @@
+// Copyright Epic Games, Inc. All Rights Reserved.
+
+#pragma once
+
+#include "DataStorage/CommonTypes.h"
+#include "TedsOperationCallbacks.h"
+
+#include "TedsOperationColumns.generated.h"
+
+USTRUCT(meta = (DisplayName = "Operation Name"))
+struct FOperationNameColumn final : public FEditorDataStorageColumn
+{
+	GENERATED_BODY()
+
+	FName Value;
+};
+
+USTRUCT(meta = (DisplayName = "Operation Probe"))
+struct FOperationProbeColumn final : public FEditorDataStorageColumn
+{
+	GENERATED_BODY()
+
+	UE::Editor::DataStorage::Operations::FProbeCallback Callback;
+};
+
+USTRUCT(meta = (DisplayName = "Operation Test"))
+struct FOperationTestColumn final : public FEditorDataStorageColumn
+{
+	GENERATED_BODY()
+
+	UE::Editor::DataStorage::Operations::FTestCallback Callback;
+};
+
+USTRUCT(meta = (DisplayName = "Operation Apply"))
+struct FOperationApplyColumn final : public FEditorDataStorageColumn
+{
+	GENERATED_BODY()
+
+	UE::Editor::DataStorage::Operations::FApplyCallback Callback;
+};
+
+USTRUCT(meta = (DisplayName = "Operation Priority"))
+struct FOperationPriorityColumn final : public FEditorDataStorageColumn
+{
+	GENERATED_BODY()
+
+	int64 Value;
+};
+
+namespace UE::Editor::DataStorage::Operations
+{
+	/** Column for the name of the operation. */
+	using FNameColumn = FOperationNameColumn;
+	/** Probe call column, to test superficially if a given input is potentially acceptable for the operation. */
+	using FProbeColumn = FOperationProbeColumn;
+	/** Test call column, to query if the operation for a given input can be executed. */
+	using FTestColumn = FOperationTestColumn;
+	/** Apply call column, to execute the operation for a given input. */
+	using FApplyColumn = FOperationApplyColumn;
+	/** Priority value of this operation. Used for the default sorting when calling the operation system. */
+	using FPriorityColumn = FOperationPriorityColumn;
+}
diff --git a/Engine/Plugins/Experimental/EditorDataStorageFeatures/Source/TedsOperations/Public/TedsOperationInput.h b/Engine/Plugins/Experimental/EditorDataStorageFeatures/Source/TedsOperations/Public/TedsOperationInput.h
new file mode 100644
--- /dev/null
+++ b/Engine/Plugins/Experimental/EditorDataStorageFeatures/Source/TedsOperations/Public/TedsOperationInput.h
@@ -0,0 +1,84 @@
+// Copyright Epic Games, Inc. All Rights Reserved.
+
+#pragma once
+
+#include "DataStorage/CommonTypes.h"
+#include "Elements/Interfaces/TypedElementDataStorageInterface.h"
+#include "TedsOperationResult.h"
+
+#include "TedsOperationInput.generated.h"
+
+USTRUCT(meta = (DisplayName = "Operation Input Source"))
+struct FOperationSourceColumn final : public FEditorDataStorageColumn
+{
+	GENERATED_BODY()
+
+	UE::Editor::DataStorage::RowHandle Value;
+};
+
+USTRUCT(meta = (DisplayName = "Operation Input Description"))
+struct FOperationDescriptionColumn final : public FEditorDataStorageColumn
+{
+	GENERATED_BODY()
+
+	FText Value;
+};
+
+USTRUCT(meta = (DisplayName = "Operation Test Result"))
+struct FOperationTestResultTag final : public FEditorDataStorageTag
+{
+	GENERATED_BODY()
+};
+
+USTRUCT(meta = (DisplayName = "Operation Result"))
+struct FOperationResultColumn final : public FEditorDataStorageColumn
+{
+	GENERATED_BODY()
+
+	UE::Editor::DataStorage::Operations::FResult Value;
+};
+
+namespace UE::Editor::DataStorage::Operations
+{
+/** Holds a general purpose source row for the operation. */
+using FSourceColumn = FOperationSourceColumn;
+/** Holds a description which the operation may write to. */
+using FDescriptionColumn = FOperationDescriptionColumn;
+/** Denotes that the test operation has been successful. */
+using FTestResultTag = FOperationTestResultTag;
+/** Denotes that the operation has been sucessful. */
+using FResultColumn = FOperationResultColumn;
+
+namespace Utilities
+{
+/** Returns the value of the source column, if set. */
+inline RowHandle GetSourceRow(const ICoreProvider& Storage, RowHandle InputRow)
+{
+	const FSourceColumn* Column = Storage.GetColumn<FSourceColumn>(InputRow);
+	return Column ? Column->Value : InvalidRowHandle;
+}
+
+/** Returns a pointer to the description text, if set.*/
+inline FText* GetDescriptionPtr(ICoreProvider& Storage, RowHandle InputRow)
+{
+	FDescriptionColumn* Column = Storage.GetColumn<FDescriptionColumn>(InputRow);
+	return Column ? &Column->Value : nullptr;
+}
+
+/** Returns a pointer to the description text, if set.*/
+inline const FText* GetDescriptionPtr(const ICoreProvider& Storage, RowHandle InputRow)
+{
+	const FDescriptionColumn* Column = Storage.GetColumn<FDescriptionColumn>(InputRow);
+	return Column ? &Column->Value : nullptr;
+}
+
+/** Returns the description text as a value. */
+inline FText GetDescription(const ICoreProvider& Storage, RowHandle InputRow)
+{
+	const FText* Description = GetDescriptionPtr(Storage, InputRow);
+	return Description ? *Description : FText();
+}
+}
+}
+
+
diff --git a/Engine/Plugins/Experimental/EditorDataStorageFeatures/Source/TedsOperations/Public/TedsOperationResult.h b/Engine/Plugins/Experimental/EditorDataStorageFeatures/Source/TedsOperations/Public/TedsOperationResult.h
new file mode 100644
--- /dev/null
+++ b/Engine/Plugins/Experimental/EditorDataStorageFeatures/Source/TedsOperations/Public/TedsOperationResult.h
@@ -0,0 +1,17 @@
+// Copyright Epic Games, Inc. All Rights Reserved.
+
+#pragma once
+
+#include "Elements/Framework/TypedElementRowHandleArray.h"
+
+namespace UE::Editor::DataStorage::Operations
+{
+	
+struct FResult
+{
+	FRowHandleArray Created;
+	FRowHandleArray Changed;
+	FRowHandleArray Removed;
+};
+
+}
diff --git a/Engine/Plugins/Experimental/EditorDataStorageFeatures/Source/TedsOperations/Public/TedsOperationSystem.h b/Engine/Plugins/Experimental/EditorDataStorageFeatures/Source/TedsOperations/Public/TedsOperationSystem.h
new file mode 100644
--- /dev/null
+++ b/Engine/Plugins/Experimental/EditorDataStorageFeatures/Source/TedsOperations/Public/TedsOperationSystem.h
@@ -0,0 +1,106 @@
+// Copyright Epic Games, Inc. All Rights Reserved.
+
+#pragma once
+
+#include "DataStorage/Handles.h"
+#include "Elements/Interfaces/TypedElementDataStorageFactory.h"
+#include "TedsOperationCallbacks.h"
+#include "TedsOperationResult.h"
+
+#include "TedsOperationSystem.generated.h"
+
+#define UE_API TEDSOPERATIONS_API
+
+namespace UE::Editor::DataStorage
+{
+struct FQueryDescription;
+}
+namespace UE::Editor::DataStorage::QueryStack
+{
+class IQueryNode;
+class IRowNode;
+class FRowFilterNode;
+class FExplicitUpdateExecutor;
+}
+
+/** 
+ * Generic system to provide functionality for the operation columns. Operations
+ * that were added to the data storage can be selected and "tested" or "applied"
+ * in a consistent way by this system depending on an input row.
+ */
+UCLASS(MinimalAPI)
+class UOperationSystem : public UEditorDataStorageFactory
+{
+	GENERATED_BODY()
+public:
+	static constexpr int64 DefaultPriority = 0;
+	
+	using FFilter = TFunction<bool(const UE::Editor::DataStorage::ICoreProvider&, UE::Editor::DataStorage::RowHandle)>;
+	using FLess   = TFunction<bool(const UE::Editor::DataStorage::ICoreProvider&, UE::Editor::DataStorage::RowHandle, UE::Editor::DataStorage::RowHandle)>;
+
+	struct FDefaultLess
+	{
+		UE_API bool operator()(const UE::Editor::DataStorage::ICoreProvider& Storage, UE::Editor::DataStorage::RowHandle A, UE::Editor::DataStorage::RowHandle B) const;
+	};
+
+public:
+	UE_API UOperationSystem();
+	UE_API virtual ~UOperationSystem() override;
+
+	UE_API virtual void PreRegister(UE::Editor::DataStorage::ICoreProvider& Storage) override;
+	UE_API virtual void RegisterTables(UE::Editor::DataStorage::ICoreProvider& Storage) override;
+	UE_API virtual void RegisterQueries(UE::Editor::DataStorage::ICoreProvider& Storage) override;
+
+	/** Utility method to create a row with columns for the given input operation. */
+	UE_API UE::Editor::DataStorage::RowHandle AddOperation(FName Name, TNotNull<UE::Editor::DataStorage::Operations::FApplyCallback> Operation,
+		UE::Editor::DataStorage::Operations::FTestCallback Test, UE::Editor::DataStorage::Operations::FProbeCallback Probe,
+		int64 Priority = DefaultPriority);
+
+	/** Retrieves all operation rows that accept (probe call) the given input row.
+	 * @param OutRows The array to fill with the operation rows.
+	 * @param InputRow The row that contains input data that will be passed to the operations.
+	 * @param Filter [optional] Predicate that can be used to customarily filter the operation rows. False = Skip, True = Keep.
+	 * @param Less [optional] Sorter that can be used to customarily sort the operation rows.
+	 */
+	UE_API void GetOperations(TArray<UE::Editor::DataStorage::RowHandle>& OutRows, UE::Editor::DataStorage::RowHandle InputRow,
+		const FFilter& Filter = {}, const FLess& Less = FDefaultLess()) const;
+
+	/** Invokes the Test callback for each input row for all operations that accept (probe call) it and writes the result of the first success to the input row.
+	 * @param InputRows The rows that contain input data that will be passed to the operations.
+	 * @param Filter [optional] Predicate that can be used to customarily filter the operation rows. False = Skip, True = Keep.
+	 * @param Less [optional] Sorter that can be used to customarily sort the operation rows.
+	 * @returns The number of successfully tested input rows.
+	 */
+	UE_API int32 Test(UE::Editor::DataStorage::FRowHandleArrayView InputRows, const FFilter& Filter = {}, const FLess& Less = FDefaultLess()) const;
+
+	/** Invokes the Apply callback for each input row for all operations that accept (probe call) it and writes the result of the first success to the input row.
+	 * @param InputRows The rows that contain input data that will be passed to the operations.
+	 * @param Filter [optional] Predicate that can be used to customarily filter the operation rows. False = Skip, True = Keep.
+	 * @param Less [optional] Sorter that can be used to customarily sort the operation rows.
+	 * @returns The number of successfully applied input rows.
+	 */
+	UE_API int32 Apply(UE::Editor::DataStorage::FRowHandleArrayView InputRows, const FFilter& Filter = {}, const FLess& Less = FDefaultLess()) const;
+	
+	/** Utility method to create a batch of rows for the given source rows. */
+	UE_API bool CreateInputRows(UE::Editor::DataStorage::FRowHandleArray& OutRows, UE::Editor::DataStorage::FRowHandleArrayView SourceRows,
+		bool bAddDescription = false) const;
+	
+	/** Utility method to remove a batch of rows. */
+	UE_API void RemoveInputRows(UE::Editor::DataStorage::FRowHandleArrayView RowsToRemove) const;
+	
+protected:
+	UE_API virtual void GetQueryDescription(UE::Editor::DataStorage::FQueryDescription& OutDescription) const;
+	
+protected:
+	UE::Editor::DataStorage::ICoreProvider*                     StoragePtr;
+	UE::Editor::DataStorage::TableHandle                        OperationTable;
+	TSharedPtr<UE::Editor::DataStorage::QueryStack::IQueryNode> QueryNode;
+	TSharedPtr<UE::Editor::DataStorage::QueryStack::IRowNode>   ResultsNode;
+	TSharedPtr<UE::Editor::DataStorage::QueryStack::FExplicitUpdateExecutor> Executor;
+
+protected:
+	UE::Editor::DataStorage::TableHandle InputTable;
+	UE::Editor::DataStorage::TableHandle InputTableDescr;
+};
+
+#undef UE_API
diff --git a/Engine/Plugins/Experimental/EditorDataStorageFeatures/Source/TedsOperations/Public/TedsOperations.h b/Engine/Plugins/Experimental/EditorDataStorageFeatures/Source/TedsOperations/Public/TedsOperations.h
deleted file mode 100644
diff --git a/Engine/Plugins/Experimental/EditorDataStorageFeatures/Source/TedsOperations/Public/TedsOperationsCallbacks.h b/Engine/Plugins/Experimental/EditorDataStorageFeatures/Source/TedsOperations/Public/TedsOperationsCallbacks.h
deleted file mode 100644
diff --git a/Engine/Plugins/Experimental/EditorDataStorageFeatures/Source/TedsOperations/Public/TedsOperationsColumns.h b/Engine/Plugins/Experimental/EditorDataStorageFeatures/Source/TedsOperations/Public/TedsOperationsColumns.h
deleted file mode 100644
diff --git a/Engine/Plugins/Experimental/EditorDataStorageFeatures/Source/TedsOperations/Public/TedsOperationsFactory.h b/Engine/Plugins/Experimental/EditorDataStorageFeatures/Source/TedsOperations/Public/TedsOperationsFactory.h
deleted file mode 100644
diff --git a/Engine/Plugins/Experimental/EditorDataStorageFeatures/Source/TedsOperations/Public/TedsOperationsResult.h b/Engine/Plugins/Experimental/EditorDataStorageFeatures/Source/TedsOperations/Public/TedsOperationsResult.h
deleted file mode 100644
diff --git a/Engine/Plugins/Experimental/EditorDataStorageFeatures/Source/TedsOperations/TedsOperations.build.cs b/Engine/Plugins/Experimental/EditorDataStorageFeatures/Source/TedsOperations/TedsOperations.build.cs
--- a/Engine/Plugins/Experimental/EditorDataStorageFeatures/Source/TedsOperations/TedsOperations.build.cs
+++ b/Engine/Plugins/Experimental/EditorDataStorageFeatures/Source/TedsOperations/TedsOperations.build.cs
@@ -24,13 +24,20 @@ public class TedsOperations : ModuleRules
 				{
 					"Core",
 					"CoreUObject",
+					"SlateCore",
+					"ToolWidgets",
 					"TypedElementFramework",
-					"TedsQueryStack",
+					"UnrealEd",
 				});
 
 			PrivateDependencyModuleNames.AddRange(
 				new string[]
 				{
+					"Engine",
+					"SceneOutliner",
+					"Slate",
+					"TedsQueryStack",
+					"TypedElementRuntime",
 				});
 
 		}
diff --git a/Engine/Plugins/Experimental/EditorDataStorageFeatures/Source/TedsOutliner/Private/TedsOutlinerMode.cpp b/Engine/Plugins/Experimental/EditorDataStorageFeatures/Source/TedsOutliner/Private/TedsOutlinerMode.cpp
--- a/Engine/Plugins/Experimental/EditorDataStorageFeatures/Source/TedsOutliner/Private/TedsOutlinerMode.cpp
+++ b/Engine/Plugins/Experimental/EditorDataStorageFeatures/Source/TedsOutliner/Private/TedsOutlinerMode.cpp
@@ -5,10 +5,10 @@
 #include "Elements/Columns/TypedElementCompatibilityColumns.h"
 #include "Elements/Columns/TypedElementSelectionColumns.h"
 #include "FolderTreeItem.h"
-#include "TedsDragDropOp.h"
-#include "TedsDropOperationsColumnsHelpers.h"
-#include "TedsDropOperationsFactory.h"
-#include "TedsDropOperationsUtility.h"
+#include "DragAndDrop/TedsDragDropOp.h"
+#include "DragAndDrop/TedsDragDropOpUtility.h"
+#include "DragAndDrop/DropOperationInput.h"
+#include "DragAndDrop/DropOperationSystem.h"
 #include "TedsOutlinerHierarchy.h"
 #include "TedsOutlinerItem.h"
 #include "ToolMenus.h"
@@ -158,29 +158,33 @@ FSceneOutlinerDragValidationInfo FTedsOutlinerMode::ValidateDrop(const ISceneOut
 	// that info to the operations somehow
 	if (const FTedsOutlinerTreeItem* TedsDropTarget = DropTarget.CastTo<FTedsOutlinerTreeItem>())
 	{
-		if (UDropOperationsFactory* DropFactory = Storage->FindFactory<UDropOperationsFactory>())
+		if (UDropOperationSystem* DropSystem = Storage->FindFactory<UDropOperationSystem>())
 		{
-			FText Description;
-			
-			for (RowHandle SourceRow : DraggedRowHandles.GetRows())
+			FRowHandleArray InputRows;
+			if (DropSystem->CreateInputRows(InputRows, DraggedRowHandles.GetRows(), TedsDropTarget->GetRowHandle(), true))
 			{
-				FScopedRow InputRow = DropFactory->CreateInputRowScoped(TedsDropTarget->GetRowHandle(), SourceRow, true);
- 
-				bool bResult = DropFactory->Test(InputRow.GetRow()) != InvalidRowHandle;
+				bool bResult = DropSystem->Test(InputRows.GetRows()) == InputRows.Num();
 
-				// Currently the description of the last operation will take precedence, we might want support for batched tests in the future
-				if (FText* DescriptionPtr = DragAndDrop::Helpers::GetDescription(*Storage, InputRow.GetRow()))
+				FText Description;
+				for (RowHandle InputRow : InputRows.GetRows())
 				{
-					Description = *DescriptionPtr;
+					// Find an error text when we failed and a success text if we succeeded.
+					if (bResult != Storage->HasColumns<Operations::FTestResultTag>(InputRow))
+					{
+						continue;
+					}
+		
+					Description = Operations::Utilities::GetDescription(*Storage, InputRow);
+					if (!Description.IsEmpty())
+					{
+						break;
+					}
 				}
 
-				// Currently only accept drops if all the rows are compatible with an operation
-				if (!bResult)
-				{
-					return FSceneOutlinerDragValidationInfo(ESceneOutlinerDropCompatibility::IncompatibleGeneric, Description);
-				}
+				DropSystem->RemoveInputRows(InputRows.GetRows());
+
+				return { bResult ? ESceneOutlinerDropCompatibility::CompatibleGeneric : ESceneOutlinerDropCompatibility::IncompatibleGeneric, MoveTemp(Description) };
 			}
-			return FSceneOutlinerDragValidationInfo(ESceneOutlinerDropCompatibility::CompatibleGeneric, Description);
 		}
 	}
 
@@ -209,12 +213,13 @@ void FTedsOutlinerMode::OnDrop(ISceneOutlinerTreeItem& DropTarget, const FSceneO
 	{
 		if (ValidationInfo.CompatibilityType == ESceneOutlinerDropCompatibility::CompatibleGeneric)
 		{
-			if (UDropOperationsFactory* DropFactory = Storage->FindFactory<UDropOperationsFactory>())
+			if (UDropOperationSystem* DropSystem = Storage->FindFactory<UDropOperationSystem>())
 			{
-				for (RowHandle SourceRow : DraggedRowHandles.GetRows())
+				FRowHandleArray InputRows;
+				if (DropSystem->CreateInputRows(InputRows, DraggedRowHandles.GetRows(), TedsItem->GetRowHandle(), true))
 				{
-					FScopedRow InputRow = DropFactory->CreateInputRowScoped(TedsItem->GetRowHandle(), SourceRow, false);
-					(void)DropFactory->Apply(InputRow.GetRow());
+					(void)DropSystem->Apply(InputRows.GetRows());
+					DropSystem->RemoveInputRows(InputRows.GetRows());
 				}
 			}
 		}
diff --git a/Engine/Plugins/Experimental/EditorDataStorageFeatures/Source/TedsOutliner/TedsOutliner.Build.cs b/Engine/Plugins/Experimental/EditorDataStorageFeatures/Source/TedsOutliner/TedsOutliner.Build.cs
--- a/Engine/Plugins/Experimental/EditorDataStorageFeatures/Source/TedsOutliner/TedsOutliner.Build.cs
+++ b/Engine/Plugins/Experimental/EditorDataStorageFeatures/Source/TedsOutliner/TedsOutliner.Build.cs
@@ -42,8 +42,7 @@ public class TedsOutliner : ModuleRules
 					"LevelEditor",
 					"TedsQueryStack",
 					"TedsActorCompatibility",
-					"TedsOperations",
-					"TedsDropOperations"
+					"TedsOperations"
 				});
 
 			DynamicallyLoadedModuleNames.AddRange(new string[] {});
diff --git a/Engine/Plugins/Experimental/EditorDataStorageFeatures/Source/TedsTableViewer/Private/Widgets/STedsHierarchyViewer.cpp b/Engine/Plugins/Experimental/EditorDataStorageFeatures/Source/TedsTableViewer/Private/Widgets/STedsHierarchyViewer.cpp
--- a/Engine/Plugins/Experimental/EditorDataStorageFeatures/Source/TedsTableViewer/Private/Widgets/STedsHierarchyViewer.cpp
+++ b/Engine/Plugins/Experimental/EditorDataStorageFeatures/Source/TedsTableViewer/Private/Widgets/STedsHierarchyViewer.cpp
@@ -3,17 +3,12 @@
 #include "Widgets/STedsHierarchyViewer.h"
 
 #include "TableViewerWidgetDropHandler.h"
-#include "TedsDragDropOp.h"
-#include "TedsDropOperationsColumnsHelpers.h"
-#include "TedsDropOperationsFactory.h"
-#include "TedsDropOperationsUtility.h"
 #include "TedsTableViewerColumn.h"
 #include "TedsTableViewerWidgetColumns.h"
 #include "Columns/SlateDelegateColumns.h"
 #include "DataStorage/Features.h"
 #include "Elements/Columns/TypedElementHiearchyColumns.h"
 #include "Elements/Columns/TypedElementUIColumns.h"
-#include "Elements/Framework/TypedElementAttributeBinding.h"
 #include "Elements/Framework/TypedElementDataStorageWidget.h"
 #include "QueryStack/TopLevelRowsNode.h"
 #include "Widgets/STedsTableViewerRow.h"
diff --git a/Engine/Plugins/Experimental/EditorDataStorageFeatures/Source/TedsTableViewer/Private/Widgets/STedsTableViewerRow.cpp b/Engine/Plugins/Experimental/EditorDataStorageFeatures/Source/TedsTableViewer/Private/Widgets/STedsTableViewerRow.cpp
--- a/Engine/Plugins/Experimental/EditorDataStorageFeatures/Source/TedsTableViewer/Private/Widgets/STedsTableViewerRow.cpp
+++ b/Engine/Plugins/Experimental/EditorDataStorageFeatures/Source/TedsTableViewer/Private/Widgets/STedsTableViewerRow.cpp
@@ -4,12 +4,9 @@
 #include "Widgets/STedsTableViewerRow.h"
 
 #include "Columns/SlateDelegateColumns.h"
-#include "DragAndDrop/DecoratedDragDropOp.h"
 #include "Editor.h"
 #include "Elements/Columns/TypedElementHiearchyColumns.h"
-#include "TedsDragDropOp.h"
-#include "TedsDropOperations.h"
-#include "TedsDropOperationsUtility.h"
+#include "DragAndDrop/TedsDragDropOp.h"
 #include "TedsTableViewerColumn.h"
 #include "TedsTableViewerUtils.h"
 #include "TedsTableViewerWidgetColumns.h"
diff --git a/Engine/Plugins/Experimental/EditorDataStorageFeatures/Source/TedsTableViewer/Private/Widgets/TableViewerWidgetDropHandler.cpp b/Engine/Plugins/Experimental/EditorDataStorageFeatures/Source/TedsTableViewer/Private/Widgets/TableViewerWidgetDropHandler.cpp
--- a/Engine/Plugins/Experimental/EditorDataStorageFeatures/Source/TedsTableViewer/Private/Widgets/TableViewerWidgetDropHandler.cpp
+++ b/Engine/Plugins/Experimental/EditorDataStorageFeatures/Source/TedsTableViewer/Private/Widgets/TableViewerWidgetDropHandler.cpp
@@ -2,8 +2,8 @@
 
 #include "TableViewerWidgetDropHandler.h"
 
+#include "DragAndDrop/DropOperationInput.h"
 #include "ITedsTableViewer.h"
-#include "TedsDropOperationsColumns.h"
 #include "TedsTableViewerWidgetColumns.h"
 
 namespace UE::Editor::DataStorage::Widgets
@@ -14,23 +14,27 @@ namespace UE::Editor::DataStorage::Widgets
 	{
 	}
 
-	FWidgetDropHandler::EUpdateInput FTableViewerWidgetDropHandler::UpdateInput(ICoreProvider& Storage, const FGeometry& Geometry,
+	FWidgetDropHandler::EUpdateParameters FTableViewerWidgetDropHandler::UpdateParameters(ICoreProvider& Storage, const FGeometry& Geometry,
 		const FDragDropEvent& InputEvent)
 	{
 		RowHandle TargetRow = GetDropTarget(Storage);
 
 		bool bTargetRowChanged = LastTargetRow != TargetRow;
-
+		LastTargetRow = TargetRow;
 		if (bTargetRowChanged)
 		{
-			for (FDropEntry& Entry : DropEntries)
-			{
-				Storage.AddColumn(Entry.InputRow, DragAndDrop::FTargetColumn{ .Value = TargetRow });
-			}
-			LastTargetRow = TargetRow;
-			return EUpdateInput::ResetOperations;
+			return EUpdateParameters::ResetOperations;
+		}
+		return FWidgetDropHandler::UpdateParameters(Storage, Geometry, InputEvent);
+	}
+
+	void FTableViewerWidgetDropHandler::SetRowsFromParameters(ICoreProvider& Storage, FRowHandleArrayView Rows) const
+	{
+		for (RowHandle Row : Rows)
+		{
+			Storage.AddColumn(Row, Operations::FDropTargetColumn{ .Value = LastTargetRow });
 		}
-		return EUpdateInput::NoChanges;
+		FWidgetDropHandler::SetRowsFromParameters(Storage, Rows);
 	}
 
 	RowHandle FTableViewerWidgetDropHandler::GetDropTarget(ICoreProvider& Storage) const
@@ -57,7 +61,7 @@ namespace UE::Editor::DataStorage::Widgets
 		return InvalidRowHandle;
 	}
 
-	void FTableViewerWidgetDropHandler::Stop(ICoreProvider& Storage)
+	void FTableViewerWidgetDropHandler::Stop(ICoreProvider& Storage, UDropOperationSystem& DropSystem)
 	{
 		LastTargetRow = InvalidRowHandle;
 		
@@ -66,6 +70,6 @@ namespace UE::Editor::DataStorage::Widgets
 			Storage.RemoveColumn<FTableViewerDropInfoColumn>(TableViewerPtr->GetWidgetRowHandle());
 		}
 
-		FWidgetDropHandler::Stop(Storage);
+		FWidgetDropHandler::Stop(Storage, DropSystem);
 	}
 }
diff --git a/Engine/Plugins/Experimental/EditorDataStorageFeatures/Source/TedsTableViewer/Private/Widgets/TableViewerWidgetDropHandler.h b/Engine/Plugins/Experimental/EditorDataStorageFeatures/Source/TedsTableViewer/Private/Widgets/TableViewerWidgetDropHandler.h
--- a/Engine/Plugins/Experimental/EditorDataStorageFeatures/Source/TedsTableViewer/Private/Widgets/TableViewerWidgetDropHandler.h
+++ b/Engine/Plugins/Experimental/EditorDataStorageFeatures/Source/TedsTableViewer/Private/Widgets/TableViewerWidgetDropHandler.h
@@ -2,7 +2,7 @@
 
 #pragma once
 
-#include "Widgets/WidgetDropHandler.h"
+#include "DragAndDrop/Widgets/WidgetDropHandler.h"
 
 namespace UE::Editor::DataStorage
 {
@@ -19,9 +19,11 @@ namespace UE::Editor::DataStorage::Widgets
 		virtual ~FTableViewerWidgetDropHandler() override = default;
 
 	protected:
-		virtual EUpdateInput UpdateInput(ICoreProvider& Storage, const FGeometry& Geometry, const FDragDropEvent& InputEvent) override;
+		virtual EUpdateParameters UpdateParameters(ICoreProvider& Storage, const FGeometry& Geometry, const FDragDropEvent& InputEvent) override;
+		virtual void SetRowsFromParameters(ICoreProvider& Storage, FRowHandleArrayView Rows) const override;
+		virtual void Stop(ICoreProvider& Storage, UDropOperationSystem& DropSystem) override;
+		
 		RowHandle GetDropTarget(ICoreProvider& Storage) const;
-		virtual void Stop(ICoreProvider& Storage) override;
 	private:
 		TWeakPtr<ITableViewer> TableViewer;
 		RowHandle LastTargetRow = InvalidRowHandle;
diff --git a/Engine/Plugins/Experimental/EditorDataStorageFeatures/Source/TedsTableViewer/Private/Widgets/WidgetDropHandler.cpp b/Engine/Plugins/Experimental/EditorDataStorageFeatures/Source/TedsTableViewer/Private/Widgets/WidgetDropHandler.cpp
deleted file mode 100644
diff --git a/Engine/Plugins/Experimental/EditorDataStorageFeatures/Source/TedsTableViewer/Public/Widgets/WidgetDropHandler.h b/Engine/Plugins/Experimental/EditorDataStorageFeatures/Source/TedsTableViewer/Public/Widgets/WidgetDropHandler.h
deleted file mode 100644
diff --git a/Engine/Plugins/Experimental/EditorDataStorageFeatures/Source/TedsTableViewer/TedsTableViewer.Build.cs b/Engine/Plugins/Experimental/EditorDataStorageFeatures/Source/TedsTableViewer/TedsTableViewer.Build.cs
--- a/Engine/Plugins/Experimental/EditorDataStorageFeatures/Source/TedsTableViewer/TedsTableViewer.Build.cs
+++ b/Engine/Plugins/Experimental/EditorDataStorageFeatures/Source/TedsTableViewer/TedsTableViewer.Build.cs
@@ -17,9 +17,8 @@ public class TedsTableViewer : ModuleRules
 			PublicDependencyModuleNames.AddRange(
 				new string[]
 				{
-					"TypedElementFramework",
 					"TedsQueryStack",
-					"ToolWidgets"
+					"TypedElementFramework",
 				});
 
 			PrivateDependencyModuleNames.AddRange(
@@ -32,12 +31,10 @@ public class TedsTableViewer : ModuleRules
 					"InputCore",
 					"Slate",
 					"SlateCore",
-					"WorkspaceMenuStructure",
-					"ToolMenus",
 					"TedsOperations",
-					"TedsDropOperations",
-					"UnrealEd",
-					"TypedElementRuntime"
+					"ToolMenus",
+					"ToolWidgets",
+					"WorkspaceMenuStructure",
 				});
 		}
 	}
