[TEDS FilterBar] Add functions to add, remove, and replace filters with a given array of new filters

diff --git a/Engine/Plugins/Experimental/EditorDataStorageFeatures/Source/TedsTableViewer/Private/Widgets/Composite/STedsCompositeHierarchyViewer.cpp b/Engine/Plugins/Experimental/EditorDataStorageFeatures/Source/TedsTableViewer/Private/Widgets/Composite/STedsCompositeHierarchyViewer.cpp
--- a/Engine/Plugins/Experimental/EditorDataStorageFeatures/Source/TedsTableViewer/Private/Widgets/Composite/STedsCompositeHierarchyViewer.cpp
+++ b/Engine/Plugins/Experimental/EditorDataStorageFeatures/Source/TedsTableViewer/Private/Widgets/Composite/STedsCompositeHierarchyViewer.cpp
@@ -181,6 +181,54 @@ namespace UE::Editor::DataStorage
 		return HierarchyViewer->IsSelectedExclusively(InRow);
 	}
 
+	void STedsCompositeHierarchyViewer::AddFilter(FTedsFilterData& InFilterData) const
+	{
+		if (FilterBar)
+		{
+			FilterBar->AddFilter(InFilterData);
+		}
+	}
+
+	void STedsCompositeHierarchyViewer::AddFilters(const TArray<FTedsFilterData>& InFiltersData) const
+	{
+		if (FilterBar)
+		{
+			FilterBar->AddFilters(InFiltersData);
+		}
+	}
+
+	void STedsCompositeHierarchyViewer::DeleteFilter(const FName& InFilterName) const
+	{
+		if (FilterBar)
+		{
+			FilterBar->DeleteFilter(InFilterName);
+		}
+	}
+
+	void STedsCompositeHierarchyViewer::DeleteAllFilters() const
+	{
+		if (FilterBar)
+		{
+			FilterBar->DeleteAllFilters();
+		}
+	}
+
+	void STedsCompositeHierarchyViewer::SetCommonFilters(const TArray<FName>& InFilters) const
+	{
+		if (FilterBar)
+		{
+			FilterBar->SetCommonFilters(InFilters);
+		}
+	}
+
+	void STedsCompositeHierarchyViewer::ReplaceFilters(const TArray<FTedsFilterData>& InFilters, const bool bPersistFilterStates /*= true*/) const
+	{
+		if (FilterBar)
+		{
+			FilterBar->ReplaceFilters(InFilters, bPersistFilterStates);
+		}
+	}
+
 	void STedsCompositeHierarchyViewer::ExpandAll() const
 	{
 		HierarchyViewer->ExpandAll();
diff --git a/Engine/Plugins/Experimental/EditorDataStorageFeatures/Source/TedsTableViewer/Private/Widgets/STedsFilterBar.cpp b/Engine/Plugins/Experimental/EditorDataStorageFeatures/Source/TedsTableViewer/Private/Widgets/STedsFilterBar.cpp
--- a/Engine/Plugins/Experimental/EditorDataStorageFeatures/Source/TedsTableViewer/Private/Widgets/STedsFilterBar.cpp
+++ b/Engine/Plugins/Experimental/EditorDataStorageFeatures/Source/TedsTableViewer/Private/Widgets/STedsFilterBar.cpp
@@ -28,7 +28,7 @@ namespace UE::Editor::DataStorage
 		TSharedRef<STedsFilterBar> STedsFilterBarRef = SharedThis(this);
 		Storage = GetMutableDataStorageFeature<DataStorage::ICoreProvider>(StorageFeatureName);
 
-		TSharedRef<FFilterCategory> OtherFiltersCategory = MakeShared<FFilterCategory>(LOCTEXT("OtherFilters", "Other Filters"), LOCTEXT("OtherFiltersToolTip", "Other Filters"));
+		OtherFiltersCategory = MakeShared<FFilterCategory>(LOCTEXT("OtherFilters", "Other Filters"), LOCTEXT("OtherFiltersToolTip", "Other Filters"));
 		for (FTedsFilterData FilterData : InArgs._Filters)
 		{
 			if (!FilterData.FilterCategory)
@@ -62,6 +62,96 @@ namespace UE::Editor::DataStorage
 			RecompileQueries();
 		}
 	}
+
+	TSharedRef<FTedsFilter> STedsFilterBar::AddFilter(FTedsFilterData& InFilterData)
+	{
+		if (!InFilterData.FilterCategory)
+		{
+			InFilterData.FilterCategory = OtherFiltersCategory;
+		}
+		
+		if (const TSharedPtr<FFilterBase<FTedsRowHandle&>> FoundFilter = GetFilter(InFilterData.FilterName.ToString()))
+		{
+			// If we did find an existing filter with that name, delete the old version so we can overwrite it
+			DeleteFilter(InFilterData.FilterName);
+		}
+
+		const TSharedRef<FTedsFilter> FilterToAdd = MakeShared<FTedsFilter>(InFilterData, SharedThis(this));
+		SBasicFilterBar<FTedsRowHandle&>::AddFilter(FilterToAdd);
+		return FilterToAdd;
+	}
+
+	void STedsFilterBar::AddFilters(const TArray<FTedsFilterData>& InFiltersData)
+	{
+		for (FTedsFilterData FilterData : InFiltersData)
+		{
+			AddFilter(FilterData);
+		}
+	}
+
+	void STedsFilterBar::DeleteFilter(const FName& InFilterName)
+	{
+		if (const TSharedPtr<FFilterBase<FTedsRowHandle&>> FilterToRemove = GetFilter(InFilterName.ToString()))
+		{
+			// Check if the filter is active and remove from any active filter maps
+			if (IsFilterActive(FilterToRemove.ToSharedRef()))
+			{
+				RemoveExternalQuery(InFilterName);
+				RemoveExternalQueryFunction(InFilterName);
+				RemoveClassQueryFunction(InFilterName);
+			}
+			RemoveFilter(FilterToRemove.ToSharedRef());
+			AllFrontendFilters.Remove(FilterToRemove.ToSharedRef());
+		}
+	}
+
+	void STedsFilterBar::DeleteAllFilters()
+	{
+		// Delete all filters and empty all active queries
+		ExternalQueries.Empty();
+		ExternalQueryFunctions.Empty();
+		ClassFilters.Empty();
+
+		RemoveAllFilters();
+		AllFilterCategories.Empty();
+		AllFrontendFilters.Empty();
+	}
+
+	void STedsFilterBar::SetCommonFilters(const TArray<FName>& InFilters)
+	{
+		CommonSectionFilters = InFilters;
+	}
+
+	void STedsFilterBar::ReplaceFilters(const TArray<FTedsFilterData>& InFilters, const bool bPersistFilterStates /*= true*/)
+	{
+		// Find all current filters on the filter bar and which ones are active so that state can be persisted if 
+		// filters with the same name are re-added
+		TMap<FString, bool> PreviousFilterStates;
+		if (bPersistFilterStates)
+		{
+			for(TSharedRef<SFilter> FilterOnBar : Filters)
+			{
+				PreviousFilterStates.Add(FilterOnBar->GetFilterName(), FilterOnBar->IsEnabled());
+			}
+		}
+		
+		DeleteAllFilters();
+
+		// Add the new given filters
+		for (FTedsFilterData FilterData : InFilters)
+		{
+			const TSharedRef<FTedsFilter> NewFilter = AddFilter(FilterData);
+			if (const bool* FilterState = PreviousFilterStates.Find(NewFilter->GetName()))
+			{
+				const TSharedRef<SFilter> AddedFilter = AddFilterToBar(NewFilter);
+				if (*FilterState) // Enabled defaults to false when a filter is created so only perform if we need to enable
+				{
+					AddedFilter->SetEnabled(true);
+					SetFrontendFilterActive(NewFilter, true);
+				}
+			}
+		}
+	}
 	
 	void STedsFilterBar::PopulateAddFilterMenu(UToolMenu* InMenu)
 	{
@@ -90,9 +180,9 @@ namespace UE::Editor::DataStorage
 		else
 		{
 			FToolMenuSection& CommonFiltersSection = InMenu->AddSection("FilterBarCommonFilters", LOCTEXT("TedsFilterBarCommonFiltersSection", "Common Filters"));
-			for (const TSharedRef<FFilterBase<FTedsRowHandle&>>& FrontendFilter : AllFrontendFilters)
+			for(const FName& CommonFilter : CommonSectionFilters)
 			{
-				if(CommonSectionFilters.Contains(FName(*FrontendFilter->GetName())))
+				if (const TSharedPtr<FFilterBase<FTedsRowHandle&>>& FrontendFilter = GetFilter(CommonFilter.ToString()))
 				{
 					CommonFiltersSection.AddMenuEntry(
 						NAME_None,
@@ -100,9 +190,9 @@ namespace UE::Editor::DataStorage
 						FrontendFilter->GetToolTipText(),
 						FSlateIcon(FAppStyle::GetAppStyleSetName(), FrontendFilter->GetIconName()),
 						FUIAction(
-						FExecuteAction::CreateSP(this, &STedsFilterBar::FrontendFilterClicked, FrontendFilter),
+						FExecuteAction::CreateSP(this, &STedsFilterBar::FrontendFilterClicked, FrontendFilter.ToSharedRef()),
 						FCanExecuteAction(),
-						FIsActionChecked::CreateSP(this, &STedsFilterBar::IsFrontendFilterInUse, FrontendFilter)),
+						FIsActionChecked::CreateSP(this, &STedsFilterBar::IsFrontendFilterInUse, FrontendFilter.ToSharedRef())),
 						EUserInterfaceActionType::ToggleButton
 						);
 				}
diff --git a/Engine/Plugins/Experimental/EditorDataStorageFeatures/Source/TedsTableViewer/Public/Widgets/Composite/STedsCompositeHierarchyViewer.h b/Engine/Plugins/Experimental/EditorDataStorageFeatures/Source/TedsTableViewer/Public/Widgets/Composite/STedsCompositeHierarchyViewer.h
--- a/Engine/Plugins/Experimental/EditorDataStorageFeatures/Source/TedsTableViewer/Public/Widgets/Composite/STedsCompositeHierarchyViewer.h
+++ b/Engine/Plugins/Experimental/EditorDataStorageFeatures/Source/TedsTableViewer/Public/Widgets/Composite/STedsCompositeHierarchyViewer.h
@@ -51,8 +51,8 @@ namespace UE::Editor::DataStorage
 		// If we want to enable the filter bar on the table
 		SLATE_ARGUMENT(bool, EnableFiltering)
 		
-		/** Set of filter names to also add to the 'common filters' section (section at the top of the menu) */
-        SLATE_ARGUMENT(TSet<FName>, CommonSectionFilters)
+		/** Array of filter names to also add to the 'common filters' section (section at the top of the menu) */
+        SLATE_ARGUMENT(TArray<FName>, CommonSectionFilters)
 
 		// Array of filters to add to the filter bar *AND Filters
 		SLATE_ARGUMENT(TArray<FTedsFilterData>, Filters)
@@ -99,6 +99,18 @@ namespace UE::Editor::DataStorage
 
 		TEDSTABLEVIEWER_API virtual bool IsSelectedExclusively(RowHandle InRow) const override;
 
+		TEDSTABLEVIEWER_API void AddFilter(FTedsFilterData& InFilterData) const;
+
+		TEDSTABLEVIEWER_API void AddFilters(const TArray<FTedsFilterData>& InFiltersData) const;
+
+		TEDSTABLEVIEWER_API void DeleteFilter(const FName& InFilterName) const;
+
+		TEDSTABLEVIEWER_API void DeleteAllFilters() const;
+		
+		TEDSTABLEVIEWER_API void SetCommonFilters(const TArray<FName>& InFilters) const;
+
+		TEDSTABLEVIEWER_API void ReplaceFilters(const TArray<FTedsFilterData>& InFilters, const bool bPersistFilterStates = true) const;
+		
 		// Recursively Expand all rows currently present in the viewer
 		TEDSTABLEVIEWER_API void ExpandAll() const;
 
diff --git a/Engine/Plugins/Experimental/EditorDataStorageFeatures/Source/TedsTableViewer/Public/Widgets/STedsFilterBar.h b/Engine/Plugins/Experimental/EditorDataStorageFeatures/Source/TedsTableViewer/Public/Widgets/STedsFilterBar.h
--- a/Engine/Plugins/Experimental/EditorDataStorageFeatures/Source/TedsTableViewer/Public/Widgets/STedsFilterBar.h
+++ b/Engine/Plugins/Experimental/EditorDataStorageFeatures/Source/TedsTableViewer/Public/Widgets/STedsFilterBar.h
@@ -54,8 +54,8 @@ namespace UE::Editor::DataStorage
 			/** Array of class filters to add to the filter bar *OR Filters */
 			SLATE_ARGUMENT(TArray<UClass*>, ClassFilters)
 			
-			/** Set of filter names to also add to the 'common filters' section (section at the top of the menu) */
-            SLATE_ARGUMENT(TSet<FName>, CommonSectionFilters)
+			/** Array of filter names to also add to the 'common filters' section (section at the top of the menu) */
+            SLATE_ARGUMENT(TArray<FName>, CommonSectionFilters)
 			
 			/** Delegate for when filters have changed */
 			SLATE_EVENT(FOnFilterChanged, OnFilterChanged)
@@ -79,6 +79,23 @@ namespace UE::Editor::DataStorage
 
 		/** Override tick so that filters are only recompiled once if multiple were added at once */
 		virtual void Tick(const FGeometry& AllottedGeometry, const double InCurrentTime, const float InDeltaTime) override;
+
+		/** Add an array of new Filters to the menu post-initialization, will overwrite any filters with the same name */
+		TSharedRef<FTedsFilter> AddFilter(FTedsFilterData& InFilterData);
+		/** Add a new Filter to the menu post-initialization, will overwrite any filters with the same name */
+		void AddFilters(const TArray<FTedsFilterData>& InFiltersData);
+		/** Remove a Filter from the menu and remove any place it may be referenced */
+		void DeleteFilter(const FName& InFilterName);
+		/** Remove all Filters from the menu and all references */
+		void DeleteAllFilters();
+		/** Set a new array of Common filters to appear in that section */
+		void SetCommonFilters(const TArray<FName>& InFilters);
+		
+		/** 
+		 * Remove all filters and replace them with a new array of filters with an option to persist their state (is pinned & is active) for new
+		 * filters with the same FName identifier
+		 */
+		void ReplaceFilters(const TArray<FTedsFilterData>& InFilters, const bool bPersistFilterStates = true);
 		
 		// Add an external query description to the FilterBar
 		void AddExternalQuery(const FName& QueryName, const QueryHandle& InQueryHandle);
@@ -107,9 +124,9 @@ namespace UE::Editor::DataStorage
 
 		ICoreProvider* Storage = nullptr;
 
-		// Set of Filter FNames used to populate the 'common filters' section that filters can have a duplicate entry
+		// Array of Filter FNames used to populate the 'common filters' section that filters can have a duplicate entry
 		// placed at the top of the add filter menu
-		TSet<FName> CommonSectionFilters;
+		TArray<FName> CommonSectionFilters;
 
 		// The row node passed in as the data to filter on
 		TSharedPtr<QueryStack::IRowNode> InFilterableRowNode;
@@ -147,6 +164,9 @@ namespace UE::Editor::DataStorage
 		// Delegate for additional functionality to perform after the filters have been updated, performed at the end of the RecompileQueries function
 		FOnPostFiltersChanged OnPostFiltersChanged;
 
+		// Default Category if the category ptr on the Filter Data is a nullptr
+		TSharedPtr<FFilterCategory> OtherFiltersCategory;
+
 		// Responsible for seeing if a filter was added or removed so multiple recompiles aren't performed every frame
 		bool bMarkDirty = false;
 	};
