Merging CL#47821165

[Prefab Editor] Have the Component Filters refresh on recompile

diff --git a/Engine/Plugins/Experimental/EditorDataStorageFeatures/Source/TedsTypeInfo/Private/TedsTypeInfoFactory.cpp b/Engine/Plugins/Experimental/EditorDataStorageFeatures/Source/TedsTypeInfo/Private/TedsTypeInfoFactory.cpp
--- a/Engine/Plugins/Experimental/EditorDataStorageFeatures/Source/TedsTypeInfo/Private/TedsTypeInfoFactory.cpp
+++ b/Engine/Plugins/Experimental/EditorDataStorageFeatures/Source/TedsTypeInfo/Private/TedsTypeInfoFactory.cpp
@@ -9,6 +9,7 @@
 #include "ClassViewerModule.h"
 
 #include "TedsTypeInfoModule.h"
+#include "TedsTypeInfoUtils.h"
 
 #include "DataStorage/Features.h"
 #include "Elements/Interfaces/TypedElementDataStorageCompatibilityInterface.h"
@@ -26,7 +27,6 @@
 #include "VerseVM/VVMVerseClass.h"
 
 const FName UTypeInfoFactory::PropertyTableName = "Property Information";
-const FName UTypeInfoFactory::ClassSetupActionName = "ClassSetupAction";
 const FName UTypeInfoFactory::ClassHierarchyName = "ClassHierarchy";
 const FName UTypeInfoFactory::PropertyHierarchyName = "PropertyHierarchy";
 
@@ -82,7 +82,7 @@ void UTypeInfoFactory::RegisterQueries(UE::Editor::DataStorage::ICoreProvider& D
 		Select(
 			TEXT("Setup Class Hierarchy Info"),
 			FProcessor(EQueryTickPhase::DuringPhysics, DataStorage.GetQueryTickGroupName(EQueryTickGroups::Default))
-				.MakeActivatable(ClassSetupActionName),
+				.MakeActivatable(TypeInfo::ClassSetupActionName),
 			[this](IQueryContext& Context, RowHandle Row, const FTypedElementClassTypeInfoColumn& ClassTypeColumn)
 			{
 				if (const UClass* ClassInfo = ClassTypeColumn.TypeInfo.Get())
@@ -174,7 +174,8 @@ void UTypeInfoFactory::PopulateTypeInfo()
 		AddPropertyInfo(*DataStorage, TypeWithPendingPropertyInfo.Value, TypeWithPendingPropertyInfo.Key);
 	}
 
-	DataStorage->ActivateQueries(ClassSetupActionName);
+	DataStorage->ActivateQueries(TypeInfo::ClassSetupActionName);
+	DataStorage->ActivateQueries(TypeInfo::PostClassSetupActionName);
 }
 
 bool UTypeInfoFactory::TryAddTypeInfo(UE::Editor::DataStorage::ICoreProvider& DataStorage, const UStruct* InTypeInfo,
diff --git a/Engine/Plugins/Experimental/EditorDataStorageFeatures/Source/TedsTypeInfo/Private/TedsTypeInfoFactory.h b/Engine/Plugins/Experimental/EditorDataStorageFeatures/Source/TedsTypeInfo/Private/TedsTypeInfoFactory.h
--- a/Engine/Plugins/Experimental/EditorDataStorageFeatures/Source/TedsTypeInfo/Private/TedsTypeInfoFactory.h
+++ b/Engine/Plugins/Experimental/EditorDataStorageFeatures/Source/TedsTypeInfo/Private/TedsTypeInfoFactory.h
@@ -36,7 +36,6 @@ class UTypeInfoFactory final : public UEditorDataStorageFactory
 
 public:
 	static const FName PropertyTableName;
-	static const FName ClassSetupActionName;
 	static const FName ClassHierarchyName;
 	static const FName PropertyHierarchyName;
 
diff --git a/Engine/Plugins/Experimental/EditorDataStorageFeatures/Source/TedsTypeInfo/Public/TedsTypeInfoUtils.h b/Engine/Plugins/Experimental/EditorDataStorageFeatures/Source/TedsTypeInfo/Public/TedsTypeInfoUtils.h
new file mode 100644
--- /dev/null
+++ b/Engine/Plugins/Experimental/EditorDataStorageFeatures/Source/TedsTypeInfo/Public/TedsTypeInfoUtils.h
@@ -0,0 +1,10 @@
+// Copyright Epic Games, Inc. All Rights Reserved.
+#pragma once
+
+#include "UObject/NameTypes.h"
+
+namespace UE::Editor::DataStorage::TypeInfo
+{
+	static const FName ClassSetupActionName = TEXT("ClassSetupAction");
+	static const FName PostClassSetupActionName = TEXT("PostClassSetupAction");
+}
\ No newline at end of file
