[TedsOperations] Added a few comments to operation systems and tagged UOperationSystem as abstract.

diff --git a/Engine/Plugins/Experimental/EditorDataStorageFeatures/Source/TedsOperations/Private/TedsOperationSystem.cpp b/Engine/Plugins/Experimental/EditorDataStorageFeatures/Source/TedsOperations/Private/TedsOperationSystem.cpp
--- a/Engine/Plugins/Experimental/EditorDataStorageFeatures/Source/TedsOperations/Private/TedsOperationSystem.cpp
+++ b/Engine/Plugins/Experimental/EditorDataStorageFeatures/Source/TedsOperations/Private/TedsOperationSystem.cpp
@@ -246,13 +246,13 @@ int32 UOperationSystem::Apply(UE::Editor::DataStorage::FRowHandleArrayView Input
 	return SuccessNum;
 }
 
-bool UOperationSystem::CreateInputRows(UE::Editor::DataStorage::FRowHandleArray& OutRows, UE::Editor::DataStorage::FRowHandleArrayView SourceRows,
+bool UOperationSystem::CreateInputRows(UE::Editor::DataStorage::FRowHandleArray& OutInputRows, UE::Editor::DataStorage::FRowHandleArrayView SourceRows,
 	bool bAddDescription) const
 {
 	using namespace UE::Editor::DataStorage;
 	using namespace UE::Editor::DataStorage::Operations;
 	
-	OutRows.Reset();
+	OutInputRows.Reset();
 	if (SourceRows.IsEmpty())
 	{
 		return true;
@@ -264,12 +264,12 @@ bool UOperationSystem::CreateInputRows(UE::Editor::DataStorage::FRowHandleArray&
 		return false;
 	}
 	
-	OutRows.Reserve(SourceRows.Num());
+	OutInputRows.Reserve(SourceRows.Num());
 	const RowHandle* SourceRowIt = SourceRows.begin();
-	return StoragePtr->BatchAddRow(Table, SourceRows.Num(), [this, &SourceRowIt, &OutRows](RowHandle Row)
+	return StoragePtr->BatchAddRow(Table, SourceRows.Num(), [this, &SourceRowIt, &OutInputRows](RowHandle Row)
 		{
 			StoragePtr->AddColumn(Row, FSourceColumn{ .Value = *SourceRowIt++ });
-			OutRows.Add(Row);
+			OutInputRows.Add(Row);
 		});
 }
 
diff --git a/Engine/Plugins/Experimental/EditorDataStorageFeatures/Source/TedsOperations/Public/DragAndDrop/DropOperationSystem.h b/Engine/Plugins/Experimental/EditorDataStorageFeatures/Source/TedsOperations/Public/DragAndDrop/DropOperationSystem.h
--- a/Engine/Plugins/Experimental/EditorDataStorageFeatures/Source/TedsOperations/Public/DragAndDrop/DropOperationSystem.h
+++ b/Engine/Plugins/Experimental/EditorDataStorageFeatures/Source/TedsOperations/Public/DragAndDrop/DropOperationSystem.h
@@ -15,11 +15,15 @@ struct FDropOperationTag final : public FEditorDataStorageTag
 {
 	GENERATED_BODY()
 };
+namespace UE::Editor::DataStorage::Operations
+{
+	/** Tag to identify this operation as a drop operation. UDropOperationSystem will only consider rows with this tag. */
+	using FDropTag = FDropOperationTag;
+}
 
 /**
- * Operations system for drop operations.
- * Only operation rows with the "Drop Tag" will be considered by this system.
- */
+ * Operation system to execute operations that drag&drop something on a target. Only considers operations with the FDropTag.
+*/
 UCLASS(MinimalAPI)
 class UDropOperationSystem : public UOperationSystem
 {
@@ -30,20 +34,20 @@ public:
 	
 	UE_API virtual void RegisterTables(UE::Editor::DataStorage::ICoreProvider& Storage) override;
 	
-	/** Utility method to create a batch of rows for the given source rows. */
-	UE_API bool CreateInputRows(UE::Editor::DataStorage::FRowHandleArray& OutRows, UE::Editor::DataStorage::FRowHandleArrayView SourceRows,
-		UE::Editor::DataStorage::RowHandle TargetRow, bool bAddDescription = false) const;
+	/**
+	 * Utility method to create a batch of rows that can be used as input for drop operations.
+	 * The method will create as many rows as source rows were passed in and automatically assign the FSourceColumn and FDropTag.
+	 * The source rows will be the ones the drop operation should consider.
+	 * @OutInputRows The array to fill with the input rows that were created.
+	 * @SourceRows The array holding the source rows that should be considered by the drop operation. Automatically assigned to the FSourceColumn.
+	 * @bAddDescription Whether the FDescriptionColumn should be added to the input rows. Holds an FText which the operation may write to.
+	 * @returns True if the requested number of rows were successfully created, otherwise false.
+	 */
+	UE_API bool CreateInputRows(UE::Editor::DataStorage::FRowHandleArray& OutInputRows, UE::Editor::DataStorage::FRowHandleArrayView SourceRows,
+		UE::Editor::DataStorage::RowHandle TargetRow, bool bAddDescription) const;
 
 protected:
 	UE_API virtual void GetQueryDescription(UE::Editor::DataStorage::FQueryDescription& OutDescription) const;
 };
 
-namespace UE::Editor::DataStorage::Operations
-{
-	
-/** Tag to identify this operation as a drop operation. UDropOperationSystem will only consider rows with this tag. */
-using FDropTag = FDropOperationTag;
-	
-}
-
 #undef UE_API
diff --git a/Engine/Plugins/Experimental/EditorDataStorageFeatures/Source/TedsOperations/Public/TedsOperationSystem.h b/Engine/Plugins/Experimental/EditorDataStorageFeatures/Source/TedsOperations/Public/TedsOperationSystem.h
--- a/Engine/Plugins/Experimental/EditorDataStorageFeatures/Source/TedsOperations/Public/TedsOperationSystem.h
+++ b/Engine/Plugins/Experimental/EditorDataStorageFeatures/Source/TedsOperations/Public/TedsOperationSystem.h
@@ -24,11 +24,12 @@ class FExplicitUpdateExecutor;
 }
 
 /** 
- * Generic system to provide functionality for the operation columns. Operations
- * that were added to the data storage can be selected and "tested" or "applied"
- * in a consistent way by this system depending on an input row.
+ * System to execute externally customizable logic for a given purpose and input.
+ * 
+ * By deriving this type, an "OperationSystem" for a specific purpose may be created, to which customizable "Operations" can be registered.
+ * The system may then be invoked to "test" or "apply" a registered operation that fits the given input.
  */
-UCLASS(MinimalAPI)
+UCLASS(MinimalAPI, Abstract)
 class UOperationSystem : public UEditorDataStorageFactory
 {
 	GENERATED_BODY()
@@ -38,6 +39,7 @@ public:
 	using FFilter = TFunction<bool(const UE::Editor::DataStorage::ICoreProvider&, UE::Editor::DataStorage::RowHandle)>;
 	using FLess   = TFunction<bool(const UE::Editor::DataStorage::ICoreProvider&, UE::Editor::DataStorage::RowHandle, UE::Editor::DataStorage::RowHandle)>;
 
+	/** Default sorter for invoking an OperationSystem. Sorts by priority column first and name column second. */
 	struct FDefaultLess
 	{
 		UE_API bool operator()(const UE::Editor::DataStorage::ICoreProvider& Storage, UE::Editor::DataStorage::RowHandle A, UE::Editor::DataStorage::RowHandle B) const;
@@ -81,7 +83,15 @@ public:
 	 */
 	UE_API int32 Apply(UE::Editor::DataStorage::FRowHandleArrayView InputRows, const FFilter& Filter = {}, const FLess& Less = FDefaultLess()) const;
 	
-	/** Utility method to create a batch of rows for the given source rows. */
+	/**
+	 * Utility method to create a batch of rows that can be used as input for operations.
+	 * The method will create as many rows as source rows were passed in and automatically assign the FSourceColumn.
+	 * The source rows will be the ones the operation should consider.
+	 * @OutInputRows The array to fill with the input rows that were created.
+	 * @SourceRows The array holding the source rows that should be considered by the operation. Automatically assigned to the FSourceColumn.
+	 * @bAddDescription Whether the FDescriptionColumn should be added to the input rows. Holds an FText which the operation may write to.
+	 * @returns True if the requested number of rows were successfully created, otherwise false.
+	 */
 	UE_API bool CreateInputRows(UE::Editor::DataStorage::FRowHandleArray& OutRows, UE::Editor::DataStorage::FRowHandleArrayView SourceRows,
 		bool bAddDescription = false) const;
 	
