Teds hierarchy viewer: Defer the expansion call until we have actual rows to expand.
This is a temporary solution for expansion related issues until we have persistent expansion state.

diff --git a/Engine/Plugins/Experimental/EditorDataStorageFeatures/Source/TedsTableViewer/Private/Widgets/STedsHierarchyViewer.cpp b/Engine/Plugins/Experimental/EditorDataStorageFeatures/Source/TedsTableViewer/Private/Widgets/STedsHierarchyViewer.cpp
--- a/Engine/Plugins/Experimental/EditorDataStorageFeatures/Source/TedsTableViewer/Private/Widgets/STedsHierarchyViewer.cpp
+++ b/Engine/Plugins/Experimental/EditorDataStorageFeatures/Source/TedsTableViewer/Private/Widgets/STedsHierarchyViewer.cpp
@@ -133,9 +133,13 @@ namespace UE::Editor::DataStorage
 		{
 			if (TSharedPtr<QueryStack::IRowNode> AllRowsNode = Model->GetRowNode())
 			{
-				SetExpansionState_Internal(AllRowsNode->GetRows(), bExpansionStateChangeRequested.GetValue());
+				FRowHandleArrayView AllRows = AllRowsNode->GetRows();
+				if (AllRows.Num())
+				{
+					SetExpansionState_Internal(AllRows, bExpansionStateChangeRequested.GetValue());
+					bExpansionStateChangeRequested.Reset();
+				}
 			}
-			bExpansionStateChangeRequested.Reset();
 		}
 		
 		TreeView->RequestListRefresh();
