[TEDS] Added Query Stack nodes that trigger delegates when a change is detected.

With the new executors not all stacks are updated every frame. As a result it's not always possible to easily track if there have been changes and when those  should be checked for. These new nodes simplify this be allowing a delegate to be injected at any point in the Query Stack and triggering it whenever any parent in the chain changes.
The original callback in some executors has been removed as these new nodes are more convient to place at various spots in the stack and this also keeps the principle of one task per node.

diff --git a/Engine/Plugins/Experimental/EditorDataStorageFeatures/Source/TedsQueryStack/Private/TedsQueryChangeNotifyNode.cpp b/Engine/Plugins/Experimental/EditorDataStorageFeatures/Source/TedsQueryStack/Private/TedsQueryChangeNotifyNode.cpp
new file mode 100644
--- /dev/null
+++ b/Engine/Plugins/Experimental/EditorDataStorageFeatures/Source/TedsQueryStack/Private/TedsQueryChangeNotifyNode.cpp
@@ -0,0 +1,39 @@
+// Copyright Epic Games, Inc. All Rights Reserved.
+
+#include "TedsQueryChangeNotifyNode.h"
+
+namespace UE::Editor::DataStorage::QueryStack
+{
+	FQueryChangeNotifyNode::FQueryChangeNotifyNode(TSharedPtr<IQueryNode> Parent, FOnQueryNodeChange ChangeEvent)
+		: Parent(MoveTemp(Parent))
+		, OnChangeEvent(MoveTemp(ChangeEvent))
+	{
+		checkf(Parent, TEXT("FQueryChangeNotifyNode requires that the parent is set."));
+		checkf(OnChangeEvent.IsBound(), TEXT("FRowChangFQueryChangeNotifyNodeeNotifyNode requires a delegate."))
+	}
+
+	INode::RevisionId FQueryChangeNotifyNode::GetRevision() const
+	{
+		return Parent->GetRevision();
+	}
+
+	void FQueryChangeNotifyNode::VisitParents(ParentListCallback Callback)
+	{
+		Callback(Parent);
+	}
+
+	QueryHandle FQueryChangeNotifyNode::GetQuery() const
+	{
+		return Parent->GetQuery();
+	}
+
+	void FQueryChangeNotifyNode::Update(FTimespan AllottedTime)
+	{
+		RevisionId CurrentRevision = Parent->GetRevision();
+		if (CurrentRevision != CachedParentRevision)
+		{
+			OnChangeEvent.Execute(Parent);
+			CachedParentRevision = CurrentRevision;
+		}
+	}
+} // namespace UE::Editor::DataStorage::QueryStack
diff --git a/Engine/Plugins/Experimental/EditorDataStorageFeatures/Source/TedsQueryStack/Private/TedsQueryStackExecutor.cpp b/Engine/Plugins/Experimental/EditorDataStorageFeatures/Source/TedsQueryStack/Private/TedsQueryStackExecutor.cpp
--- a/Engine/Plugins/Experimental/EditorDataStorageFeatures/Source/TedsQueryStack/Private/TedsQueryStackExecutor.cpp
+++ b/Engine/Plugins/Experimental/EditorDataStorageFeatures/Source/TedsQueryStack/Private/TedsQueryStackExecutor.cpp
@@ -28,76 +28,45 @@ namespace UE::Editor::DataStorage::QueryStack
 		return RootNode;
 	}
 
-	void FExecutorBase::SetOnChangedCallback(OnChangedFunction Callback)
-	{
-		OnChangedCallback = MoveTemp(Callback);
-	}
-
-	void FExecutorBase::ClearOnChangedCallback()
-	{
-		OnChangedCallback.Reset();
-	}
-
-	FExecutorBase::FExecutorBase(const FName& StackName, TSharedPtr<INode> RootNode)
-		: RootNode(MoveTemp(RootNode))
-		, StackName(StackName)
+	FExecutorBase::FExecutorBase(const FName& InStackName, TSharedPtr<INode> InRootNode)
+		: RootNode(MoveTemp(InRootNode))
+		, StackName(InStackName)
 	{
+		checkf(RootNode, TEXT("Root node for Query Stack executor '%s' can't be a nullptr. "), *StackName.ToString());
 	}
 
 	void FExecutorBase::FullUpdate()
 	{
-		// Expectation is that the implementers check if RootNode is valid. This means that the versions that registers
-		// to callbacks only have to check once.
-		if (OnChangedCallback)
-		{
-			INode::RevisionId Revision = RootNode->GetRevision();
-			FullNodeUpdate(RootNode);
-			if (RootNode->GetRevision() != Revision)
-			{
-				OnChangedCallback(StackName);
-			}
-		}
-		else
-		{
-			FullNodeUpdate(RootNode);
-		}
+		checkf(RootNode, TEXT("Root node for Query Stack executor '%s' cannot be a nullptr for a full update."), *StackName.ToString());
+		FullNodeUpdate(StackName, RootNode);
 	}
 
-	void FExecutorBase::FullNodeUpdate(const TSharedPtr<INode>& Node)
+	void FExecutorBase::FullNodeUpdate(const FName& StackName, const TSharedPtr<INode>& Node)
 	{
 		Node->VisitParents(
-			[](TSharedPtr<INode> Parent)
+			[&Node, &StackName](TSharedPtr<INode> Parent)
 			{
-				FullNodeUpdate(Parent);
+				checkf(Parent, TEXT("Parent node for Query Stack executor '%s' cannot be a nullptr for a full update."), *StackName.ToString());
+				checkf(Parent != Node, TEXT("Parent node for Query Stack executor '%s' cannot be itself for a full update."), *StackName.ToString());
+				FullNodeUpdate(StackName, Parent);
 			});
 		Node->Update(INode::UnlimitedTime);
 	}
 
 	void FExecutorBase::TimedUpdate(FTimespan AllottedTime)
 	{
-		// Expectation is that the implementers check if RootNode is valid. This means that the versions that registers
-		// to callbacks only have to check once.
-		if (OnChangedCallback)
-		{
-			INode::RevisionId Revision = RootNode->GetRevision();
-			TimedNodeUpdate(RootNode, FPlatformTime::Cycles64(), AllottedTime);
-			if (RootNode->GetRevision() != Revision)
-			{
-				OnChangedCallback(StackName);
-			}
-		}
-		else
-		{
-			TimedNodeUpdate(RootNode, FPlatformTime::Cycles64(), AllottedTime);
-		}
+		checkf(RootNode, TEXT("Root node for Query Stack executor '%s' cannot be a nullptr for a timed update."), *StackName.ToString()); 
+		TimedNodeUpdate(StackName, RootNode, FPlatformTime::Cycles64(), AllottedTime);
 	}
 
-	void FExecutorBase::TimedNodeUpdate(const TSharedPtr<INode>& Node, uint64 StartTime, FTimespan AllottedTime)
+	void FExecutorBase::TimedNodeUpdate(const FName& StackName, const TSharedPtr<INode>& Node, uint64 StartTime, FTimespan AllottedTime)
 	{
 		Node->VisitParents(
-			[StartTime, AllottedTime](TSharedPtr<INode> Parent)
+			[&Node, &StackName, StartTime, AllottedTime](TSharedPtr<INode> Parent)
 			{
-				TimedNodeUpdate(Parent, StartTime, AllottedTime);
+				checkf(Parent, TEXT("Parent node for Query Stack executor '%s' cannot be a nullptr for a timed update."), *StackName.ToString());
+				checkf(Parent != Node, TEXT("Parent node for Query Stack executor '%s' cannot be itself for a timed update."), *StackName.ToString());
+				TimedNodeUpdate(StackName, Parent, StartTime, AllottedTime);
 			});
 		FTimespan Duration(FTimespan::FromSeconds(FPlatformTime::ToSeconds64(FPlatformTime::Cycles64() - StartTime)));
 		if (Duration < AllottedTime)
@@ -307,4 +276,4 @@ namespace UE::Editor::DataStorage::QueryStack
 			FullUpdate();
 		}
 	}
-} // namespace UE::Editor::DataStorage::QueryStack
\ No newline at end of file
+} // namespace UE::Editor::DataStorage::QueryStack
diff --git a/Engine/Plugins/Experimental/EditorDataStorageFeatures/Source/TedsQueryStack/Private/TedsRowChangeNotifyNode.cpp b/Engine/Plugins/Experimental/EditorDataStorageFeatures/Source/TedsQueryStack/Private/TedsRowChangeNotifyNode.cpp
new file mode 100644
--- /dev/null
+++ b/Engine/Plugins/Experimental/EditorDataStorageFeatures/Source/TedsQueryStack/Private/TedsRowChangeNotifyNode.cpp
@@ -0,0 +1,44 @@
+// Copyright Epic Games, Inc. All Rights Reserved.
+
+#include "TedsRowChangeNotifyNode.h"
+
+namespace UE::Editor::DataStorage::QueryStack
+{
+	FRowChangeNotifyNode::FRowChangeNotifyNode(TSharedPtr<IRowNode> ParentNode, FOnRowNodeChange ChangeEvent)
+		: Parent(MoveTemp(ParentNode))
+		, OnChangeEvent(MoveTemp(ChangeEvent))
+	{
+		checkf(Parent, TEXT("FRowChangeNotifyNode requires that the parent is set."));
+		checkf(OnChangeEvent.IsBound(), TEXT("FRowChangeNotifyNode requires a delegate."))
+	}
+
+	INode::RevisionId FRowChangeNotifyNode::GetRevision() const
+	{
+		return Parent->GetRevision();
+	}
+
+	void FRowChangeNotifyNode::VisitParents(ParentListCallback Callback)
+	{
+		Callback(Parent);
+	}
+
+	FRowHandleArrayView FRowChangeNotifyNode::GetRows() const
+	{
+		return Parent->GetRows();
+	}
+
+	FRowHandleArray& FRowChangeNotifyNode::GetMutableRows()
+	{
+		return Parent->GetMutableRows();
+	}
+
+	void FRowChangeNotifyNode::Update(FTimespan AllottedTime)
+	{
+		RevisionId CurrentRevision = Parent->GetRevision();
+		if (CurrentRevision != CachedParentRevision)
+		{
+			OnChangeEvent.Execute(Parent);
+			CachedParentRevision = CurrentRevision;
+		}
+	}
+} // namespace UE::Editor::DataStorage::QueryStack
diff --git a/Engine/Plugins/Experimental/EditorDataStorageFeatures/Source/TedsQueryStack/Public/TedsQueryChangeNotifyNode.h b/Engine/Plugins/Experimental/EditorDataStorageFeatures/Source/TedsQueryStack/Public/TedsQueryChangeNotifyNode.h
new file mode 100644
--- /dev/null
+++ b/Engine/Plugins/Experimental/EditorDataStorageFeatures/Source/TedsQueryStack/Public/TedsQueryChangeNotifyNode.h
@@ -0,0 +1,33 @@
+// Copyright Epic Games, Inc. All Rights Reserved.
+
+#pragma once
+
+#include "Delegates/Delegate.h"
+#include "TedsQueryStackInterfaces.h"
+
+namespace UE::Editor::DataStorage::QueryStack
+{
+	/**
+	 * Query node that will broadcast a message when one or more of its parent nodes has been updated.
+	 */
+	class FQueryChangeNotifyNode final : public IQueryNode
+	{
+	public:
+		DECLARE_DELEGATE_OneParam(FOnQueryNodeChange, const TSharedPtr<IQueryNode>&);
+
+		TEDSQUERYSTACK_API FQueryChangeNotifyNode(TSharedPtr<IQueryNode> Parent, FOnQueryNodeChange ChangeEvent);
+		virtual ~FQueryChangeNotifyNode() override = default;
+
+		TEDSQUERYSTACK_API virtual RevisionId GetRevision() const override;
+		TEDSQUERYSTACK_API virtual void VisitParents(ParentListCallback Callback) override;
+		TEDSQUERYSTACK_API virtual QueryHandle GetQuery() const override;
+
+	protected:
+		virtual void Update(FTimespan AllottedTime) override;
+
+	private:
+		TSharedPtr<IQueryNode> Parent;
+		RevisionId CachedParentRevision = UninitializedRevisionId;
+		FOnQueryNodeChange OnChangeEvent;
+	};
+} // namespace UE::Editor::DataStorage::QueryStack
diff --git a/Engine/Plugins/Experimental/EditorDataStorageFeatures/Source/TedsQueryStack/Public/TedsQueryStackExecutor.h b/Engine/Plugins/Experimental/EditorDataStorageFeatures/Source/TedsQueryStack/Public/TedsQueryStackExecutor.h
--- a/Engine/Plugins/Experimental/EditorDataStorageFeatures/Source/TedsQueryStack/Public/TedsQueryStackExecutor.h
+++ b/Engine/Plugins/Experimental/EditorDataStorageFeatures/Source/TedsQueryStack/Public/TedsQueryStackExecutor.h
@@ -14,22 +14,16 @@ namespace UE::Editor::DataStorage::QueryStack
 	class FExecutorBase
 	{
 	public:
-		using OnChangedFunction = TFunction<void(const FName&)>;
-
 		TEDSQUERYSTACK_API virtual ~FExecutorBase() = 0;
 
 		TEDSQUERYSTACK_API const FName& GetStackName() const;
 		TEDSQUERYSTACK_API TSharedPtr<const INode> GetRootNode() const;
 		TEDSQUERYSTACK_API TSharedPtr<INode> GetRootNode();
 
-		/** If set, the provided function will be called when the a change in the query stack was reported. */
-		TEDSQUERYSTACK_API void SetOnChangedCallback(OnChangedFunction Callback);
-		TEDSQUERYSTACK_API void ClearOnChangedCallback();
-
 	// Not exporting these protected functions because this is not intended to be globally extended.
 	protected:
 		FExecutorBase() = default;
-		FExecutorBase(const FName& StackName, TSharedPtr<INode> RootNode);
+		FExecutorBase(const FName& InStackName, TSharedPtr<INode> InRootNode);
 
 		// Delete copying to keep the guarantee that the name is unique and remove the risk
 		// of the same query stack being updated multiple times.
@@ -42,13 +36,12 @@ namespace UE::Editor::DataStorage::QueryStack
 		void FullUpdate();
 		void TimedUpdate(FTimespan AllottedTime);
 		
-		OnChangedFunction OnChangedCallback;
 		TSharedPtr<INode> RootNode;
 		FName StackName;
 
 	private:
-		static void FullNodeUpdate(const TSharedPtr<INode>& Node);
-		static void TimedNodeUpdate(const TSharedPtr<INode>& Node, uint64 StartTime, FTimespan AllottedTime);
+		static void FullNodeUpdate(const FName& StackName, const TSharedPtr<INode>& Node);
+		static void TimedNodeUpdate(const FName& StackName, const TSharedPtr<INode>& Node, uint64 StartTime, FTimespan AllottedTime);
 	};
 
 	/**
diff --git a/Engine/Plugins/Experimental/EditorDataStorageFeatures/Source/TedsQueryStack/Public/TedsRowChangeNotifyNode.h b/Engine/Plugins/Experimental/EditorDataStorageFeatures/Source/TedsQueryStack/Public/TedsRowChangeNotifyNode.h
new file mode 100644
--- /dev/null
+++ b/Engine/Plugins/Experimental/EditorDataStorageFeatures/Source/TedsQueryStack/Public/TedsRowChangeNotifyNode.h
@@ -0,0 +1,34 @@
+// Copyright Epic Games, Inc. All Rights Reserved.
+
+#pragma once
+
+#include "Delegates/Delegate.h"
+#include "TedsQueryStackInterfaces.h"
+
+namespace UE::Editor::DataStorage::QueryStack
+{
+	/**
+	 * Row node that will broadcast a message when one or more of its parent nodes has been updated.
+	 */
+	class FRowChangeNotifyNode final : public IRowNode
+	{
+	public:
+		DECLARE_DELEGATE_OneParam(FOnRowNodeChange, const TSharedPtr<IRowNode>&);
+
+		TEDSQUERYSTACK_API FRowChangeNotifyNode(TSharedPtr<IRowNode> ParentNode, FOnRowNodeChange ChangeEvent);
+		virtual ~FRowChangeNotifyNode() override = default;
+
+		TEDSQUERYSTACK_API virtual RevisionId GetRevision() const override;
+		TEDSQUERYSTACK_API virtual void VisitParents(ParentListCallback Callback) override;
+		TEDSQUERYSTACK_API virtual FRowHandleArrayView GetRows() const override;
+		TEDSQUERYSTACK_API virtual FRowHandleArray& GetMutableRows() override;
+
+	protected:
+		virtual void Update(FTimespan AllottedTime) override;
+	
+	private:
+		TSharedPtr<IRowNode> Parent;
+		RevisionId CachedParentRevision = UninitializedRevisionId;
+		FOnRowNodeChange OnChangeEvent;
+	};
+} // namespace UE::Editor::DataStorage::QueryStack
