[TEDS] Move Verse Build handling for Type Info to Type Info Module

diff --git a/Engine/Plugins/Experimental/EditorDataStorageFeatures/EditorDataStorageFeatures.uplugin b/Engine/Plugins/Experimental/EditorDataStorageFeatures/EditorDataStorageFeatures.uplugin
--- a/Engine/Plugins/Experimental/EditorDataStorageFeatures/EditorDataStorageFeatures.uplugin
+++ b/Engine/Plugins/Experimental/EditorDataStorageFeatures/EditorDataStorageFeatures.uplugin
@@ -82,6 +82,10 @@
 		}
 	],
 	"Plugins": [
+		{
+			"Name": "Solaris",
+			"Enabled": true
+		}
 	],
 	"SupportedPrograms": ["ChaosVisualDebugger"]
 }
\ No newline at end of file
diff --git a/Engine/Plugins/Experimental/EditorDataStorageFeatures/Source/TedsTypeInfo/Private/TedsTypeInfoModule.cpp b/Engine/Plugins/Experimental/EditorDataStorageFeatures/Source/TedsTypeInfo/Private/TedsTypeInfoModule.cpp
--- a/Engine/Plugins/Experimental/EditorDataStorageFeatures/Source/TedsTypeInfo/Private/TedsTypeInfoModule.cpp
+++ b/Engine/Plugins/Experimental/EditorDataStorageFeatures/Source/TedsTypeInfo/Private/TedsTypeInfoModule.cpp
@@ -8,6 +8,7 @@
 
 #include "DataStorage/Features.h"
 #include "Elements/Interfaces/TypedElementDataStorageInterface.h"
+#include "SolarisLoadCompilerModule.h"
 
 namespace UE::Editor::DataStorage::TypeInfo
 {
@@ -40,10 +41,15 @@ namespace UE::Editor::DataStorage::TypeInfo
 			{
 				bDataStorageFeaturesEnabled = true;
 			});
+
+		ISolarisLoadCompilerModule::Get().OnBuildComplete().AddRaw(this, &FTedsTypeInfoModule::OnVerseBuildComplete);
 	}
 
 	void FTedsTypeInfoModule::ShutdownModule()
 	{
+		ISolarisLoadCompilerModule::Get().OnBuildComplete().RemoveAll(this);
+		UE::Editor::DataStorage::OnEditorDataStorageFeaturesEnabled().RemoveAll(this);
+
 		IModuleInterface::ShutdownModule();
 	}
 
@@ -121,6 +127,11 @@ namespace UE::Editor::DataStorage::TypeInfo
 
 		TypeInfoFactory->RefreshAllTypeInfo();
 	}
+
+	void FTedsTypeInfoModule::OnVerseBuildComplete(const TSharedRef<FSolBuildResults>& BuildResults)
+	{
+		RefreshTypeInfo();
+	}
 } // namespace UE::Editor::TypeData
 
 IMPLEMENT_MODULE(UE::Editor::DataStorage::TypeInfo::FTedsTypeInfoModule, TedsTypeInfo);
diff --git a/Engine/Plugins/Experimental/EditorDataStorageFeatures/Source/TedsTypeInfo/Public/TedsTypeInfoModule.h b/Engine/Plugins/Experimental/EditorDataStorageFeatures/Source/TedsTypeInfo/Public/TedsTypeInfoModule.h
--- a/Engine/Plugins/Experimental/EditorDataStorageFeatures/Source/TedsTypeInfo/Public/TedsTypeInfoModule.h
+++ b/Engine/Plugins/Experimental/EditorDataStorageFeatures/Source/TedsTypeInfo/Public/TedsTypeInfoModule.h
@@ -3,6 +3,9 @@
 #pragma once
 
 #include "Modules/ModuleInterface.h"
+#include "Templates/SharedPointer.h"
+
+struct FSolBuildResults;
 
 namespace UE::Editor::DataStorage::TypeInfo
 {
@@ -26,6 +29,8 @@ namespace UE::Editor::DataStorage::TypeInfo
 
 	private:
 
+		void OnVerseBuildComplete(const TSharedRef<FSolBuildResults>& BuildResults);
+
 		bool bTypeInfoFactoryEnabled = false;
 		bool bDataStorageFeaturesEnabled = false;
 	};
diff --git a/Engine/Plugins/Experimental/EditorDataStorageFeatures/Source/TedsTypeInfo/TedsTypeInfo.build.cs b/Engine/Plugins/Experimental/EditorDataStorageFeatures/Source/TedsTypeInfo/TedsTypeInfo.build.cs
--- a/Engine/Plugins/Experimental/EditorDataStorageFeatures/Source/TedsTypeInfo/TedsTypeInfo.build.cs
+++ b/Engine/Plugins/Experimental/EditorDataStorageFeatures/Source/TedsTypeInfo/TedsTypeInfo.build.cs
@@ -29,6 +29,7 @@ public class TedsTypeInfo : ModuleRules
 					"TypedElementFramework",
 					"UnrealEd",
 					"TedsAlerts",
+					"SolarisLoadCompiler"
 				});
 
 			PublicDependencyModuleNames.AddRange(
