FHierarchyRowNode : Move this node into the ::QueryStack namespace like all the other node classes

diff --git a/Engine/Plugins/Experimental/EditorDataStorageFeatures/Source/TedsActorCompatibility/Private/ActorComponentDebugHierarchyWidget/ActorComponentTree.cpp b/Engine/Plugins/Experimental/EditorDataStorageFeatures/Source/TedsActorCompatibility/Private/ActorComponentDebugHierarchyWidget/ActorComponentTree.cpp
--- a/Engine/Plugins/Experimental/EditorDataStorageFeatures/Source/TedsActorCompatibility/Private/ActorComponentDebugHierarchyWidget/ActorComponentTree.cpp
+++ b/Engine/Plugins/Experimental/EditorDataStorageFeatures/Source/TedsActorCompatibility/Private/ActorComponentDebugHierarchyWidget/ActorComponentTree.cpp
@@ -78,11 +78,11 @@ namespace UE::Editor::DataStorage::ActorCompatibility
 						);
 
 					// Use the actor-component hierarchy to enumerate the components of all the selected actors
-					TSharedPtr<QueryStack::IRowNode> AllRowNode = MakeShared<FHierarchyRowNode>(
+					TSharedPtr<QueryStack::IRowNode> AllRowNode = MakeShared<QueryStack::FHierarchyRowNode>(
 						*DataStorageInterface,
 						HierarchyHandle,
 						TopLevelRowNode,
-						FHierarchyRowNode::ESyncFlags::Always);
+						QueryStack::FHierarchyRowNode::ESyncFlags::Always);
 
 					// Create a widget to display the view model, widget internally binds the results of the query and displays appropriately
 					TSharedRef<SActorComponentTreeWidget> TreeWidget = SNew(SActorComponentTreeWidget, HierarchyHandle, AllRowNode);
diff --git a/Engine/Plugins/Experimental/EditorDataStorageFeatures/Source/TedsEverythingPicker/Private/Context/TedsPickerContextUtil.cpp b/Engine/Plugins/Experimental/EditorDataStorageFeatures/Source/TedsEverythingPicker/Private/Context/TedsPickerContextUtil.cpp
--- a/Engine/Plugins/Experimental/EditorDataStorageFeatures/Source/TedsEverythingPicker/Private/Context/TedsPickerContextUtil.cpp
+++ b/Engine/Plugins/Experimental/EditorDataStorageFeatures/Source/TedsEverythingPicker/Private/Context/TedsPickerContextUtil.cpp
@@ -166,8 +166,8 @@ namespace UE::Editor::DataStorage::Picker
 
 				// Acquire the row list of BaseClass + SubClasses
 				FHierarchyHandle ClassHierarchyHandle = Storage->FindHierarchyByName(TEXT("ClassHierarchy"));
-				TSharedRef<FHierarchyRowNode> HierarchyNode =
-					MakeShared<FHierarchyRowNode>(*Storage, ClassHierarchyHandle, BaseClassRowNode, FHierarchyRowNode::ESyncFlags::IncrementWhenDifferent);
+				TSharedRef<QueryStack::FHierarchyRowNode> HierarchyNode =
+					MakeShared<QueryStack::FHierarchyRowNode>(*Storage, ClassHierarchyHandle, BaseClassRowNode, QueryStack::FHierarchyRowNode::ESyncFlags::IncrementWhenDifferent);
 
 				// Do an initial update to populate the node
 				QueryStack::FExplicitUpdateExecutor("STypeListContextView Query Stack one-shot", HierarchyNode).Update();
diff --git a/Engine/Plugins/Experimental/EditorDataStorageFeatures/Source/TedsQueryStack/Private/TedsHierarchyNode.cpp b/Engine/Plugins/Experimental/EditorDataStorageFeatures/Source/TedsQueryStack/Private/TedsHierarchyNode.cpp
--- a/Engine/Plugins/Experimental/EditorDataStorageFeatures/Source/TedsQueryStack/Private/TedsHierarchyNode.cpp
+++ b/Engine/Plugins/Experimental/EditorDataStorageFeatures/Source/TedsQueryStack/Private/TedsHierarchyNode.cpp
@@ -5,7 +5,7 @@
 #include "Elements/Interfaces/TypedElementDataStorageInterface.h"
 
 
-namespace UE::Editor::DataStorage
+namespace UE::Editor::DataStorage::QueryStack
 {
 	FHierarchyRowNode::FHierarchyRowNode(
 		ICoreProvider& Storage,
diff --git a/Engine/Plugins/Experimental/EditorDataStorageFeatures/Source/TedsQueryStack/Public/TedsHierarchyNode.h b/Engine/Plugins/Experimental/EditorDataStorageFeatures/Source/TedsQueryStack/Public/TedsHierarchyNode.h
--- a/Engine/Plugins/Experimental/EditorDataStorageFeatures/Source/TedsQueryStack/Public/TedsHierarchyNode.h
+++ b/Engine/Plugins/Experimental/EditorDataStorageFeatures/Source/TedsQueryStack/Public/TedsHierarchyNode.h
@@ -3,12 +3,12 @@
 #pragma once
 #include "TedsRowQueryResultsNode.h"
 
-namespace UE::Editor::DataStorage
+namespace UE::Editor::DataStorage::QueryStack
 {
 	// Row which enumerates and stores all children of the given TopLevelRows (including the TopLevelRows)
 	//
 	// Warning: The resulting row array is not guaranteed to be sorted in depth order
-	class FHierarchyRowNode final : public QueryStack::IRowNode
+	class FHierarchyRowNode final : public IRowNode
 	{
 	public:
 		enum class ESyncFlags
