[TedsOperations][EntityFramework] Introduced Deletion Operations.

#Synced-CL 48152445

diff --git a/Engine/Plugins/Experimental/EditorDataStorageFeatures/Source/TedsOperations/Private/Deletion/DeletionOperationSystem.cpp b/Engine/Plugins/Experimental/EditorDataStorageFeatures/Source/TedsOperations/Private/Deletion/DeletionOperationSystem.cpp
new file mode 100644
--- /dev/null
+++ b/Engine/Plugins/Experimental/EditorDataStorageFeatures/Source/TedsOperations/Private/Deletion/DeletionOperationSystem.cpp
@@ -0,0 +1,72 @@
+// Copyright Epic Games, Inc. All Rights Reserved.
+
+#include "Deletion/DeletionOperationSystem.h"
+
+#include "Deletion/DeletionOperationInput.h"
+#include "Elements/Framework/TypedElementQueryBuilder.h"
+#include "TedsOperationColumns.h"
+
+UDeletionOperationSystem::UDeletionOperationSystem()
+	: InputTableForce(UE::Editor::DataStorage::InvalidTableHandle)
+	, InputTableForceDescr(UE::Editor::DataStorage::InvalidTableHandle)
+{
+}
+
+UDeletionOperationSystem::~UDeletionOperationSystem()
+{
+}
+
+void UDeletionOperationSystem::RegisterTables(UE::Editor::DataStorage::ICoreProvider& Storage)
+{
+	using namespace UE::Editor::DataStorage::Operations;
+	OperationTable  = Storage.RegisterTable<FApplyColumn, FDeletionTag>("DeletionOperationTable");
+
+	InputTable = Storage.RegisterTable<FSourceColumn, FDeletionTag>("DeletionOperationInputTable");
+	InputTableDescr = Storage.RegisterTable<FSourceColumn, FDeletionTag, FDescriptionColumn>("DeletionOperationInputTable_Descr");
+	InputTableForce = Storage.RegisterTable<FSourceColumn, FDeletionTag, FDeletionForceTag>("DeletionOperationInputTable_Force");
+	InputTableForceDescr = Storage.RegisterTable<FSourceColumn, FDeletionTag, FDeletionForceTag, FDescriptionColumn>("DeletionOperationInputTable_ForceDescr");
+}
+
+bool UDeletionOperationSystem::CreateInputRows(UE::Editor::DataStorage::FRowHandleArray& OutInputRows,
+	UE::Editor::DataStorage::FRowHandleArrayView SourceRows, bool bForce, bool bAddDescription) const
+{
+	using namespace UE::Editor::DataStorage;
+	using namespace UE::Editor::DataStorage::Operations;
+
+	OutInputRows.Reset();
+	if (SourceRows.IsEmpty())
+	{
+		return true;
+	}
+	
+	TableHandle Table;
+	if (bForce)
+	{
+		Table = bAddDescription ? InputTableForceDescr : InputTableForce;
+	}
+	else
+	{
+		Table = bAddDescription ? InputTableDescr : InputTable;
+	}
+	
+	if (Table == InvalidTableHandle)
+	{
+		return false;
+	}
+
+	OutInputRows.Reserve(SourceRows.Num());
+	const RowHandle* SourceRowIt = SourceRows.begin();
+	return StoragePtr->BatchAddRow(Table, SourceRows.Num(), [this, &SourceRowIt, &OutInputRows](RowHandle Row)
+		{
+			StoragePtr->AddColumn(Row, FSourceColumn{ .Value = *SourceRowIt++ });
+			OutInputRows.Add(Row);
+		});
+}
+
+void UDeletionOperationSystem::GetQueryDescription(UE::Editor::DataStorage::FQueryDescription& OutDescription) const
+{
+	using namespace UE::Editor::DataStorage::Queries;
+	using namespace UE::Editor::DataStorage::Operations;
+
+	OutDescription = Select().Where().All<FApplyColumn, FDeletionTag>().Compile();
+}
diff --git a/Engine/Plugins/Experimental/EditorDataStorageFeatures/Source/TedsOperations/Public/Deletion/DeletionOperationInput.h b/Engine/Plugins/Experimental/EditorDataStorageFeatures/Source/TedsOperations/Public/Deletion/DeletionOperationInput.h
new file mode 100644
--- /dev/null
+++ b/Engine/Plugins/Experimental/EditorDataStorageFeatures/Source/TedsOperations/Public/Deletion/DeletionOperationInput.h
@@ -0,0 +1,18 @@
+// Copyright Epic Games, Inc. All Rights Reserved.
+
+#pragma once
+
+#include "DataStorage/CommonTypes.h"
+
+#include "DeletionOperationInput.generated.h"
+
+USTRUCT(meta = (DisplayName = "Force Deletion"))
+struct FDeletionOperationForceTag final : public FEditorDataStorageTag
+{
+	GENERATED_BODY()
+};
+namespace UE::Editor::DataStorage::Operations
+{
+	/** Tag to notify the deletion operation that the input should be force deleted, i.e. regardless of restrictions. */
+	using FDeletionForceTag = FDeletionOperationForceTag;
+}
diff --git a/Engine/Plugins/Experimental/EditorDataStorageFeatures/Source/TedsOperations/Public/Deletion/DeletionOperationSystem.h b/Engine/Plugins/Experimental/EditorDataStorageFeatures/Source/TedsOperations/Public/Deletion/DeletionOperationSystem.h
new file mode 100644
--- /dev/null
+++ b/Engine/Plugins/Experimental/EditorDataStorageFeatures/Source/TedsOperations/Public/Deletion/DeletionOperationSystem.h
@@ -0,0 +1,57 @@
+// Copyright Epic Games, Inc. All Rights Reserved.
+
+#pragma once
+
+#include "DataStorage/CommonTypes.h"
+#include "TedsOperationSystem.h"
+#include "TedsOperationInput.h"
+
+#include "DeletionOperationSystem.generated.h"
+
+#define UE_API TEDSOPERATIONS_API
+
+USTRUCT(meta = (DisplayName = "Deletion Operation"))
+struct FDeletionOperationTag final : public FEditorDataStorageTag
+{
+	GENERATED_BODY()
+};
+namespace UE::Editor::DataStorage::Operations
+{
+	/** Tag to identify this operation as a drop operation. UDeletionOperationSystem will only consider operation rows with this tag. */
+	using FDeletionTag = FDeletionOperationTag;
+}
+
+/**
+ * Operation system to execute operations that delete something. Only considers operations with the FDeletionTag.
+*/
+UCLASS(MinimalAPI)
+class UDeletionOperationSystem : public UOperationSystem
+{
+	GENERATED_BODY()
+public:
+	UE_API UDeletionOperationSystem();
+	UE_API virtual ~UDeletionOperationSystem() override;
+	
+	UE_API virtual void RegisterTables(UE::Editor::DataStorage::ICoreProvider& Storage) override;
+
+	/**
+	 * Utility method to create a batch of rows that can be used as input for deletion operations.
+	 * The method will create as many rows as source rows were passed in and automatically assign the FSourceColumn and FDeletionTag.
+	 * The source rows will be the ones the operation should consider for deletion.
+	 * @OutInputRows The array to fill with the input rows that were created.
+	 * @SourceRows The array holding the source rows that should be deleted. Automatically assigned to the FSourceColumn.
+	 * @bForce Whether the FForceDeletionTag should be added to the input rows. Informs the operation to execute deletion regardless of restrictions.
+	 * @bAddDescription Whether the FDescriptionColumn should be added to the input rows. Holds an FText which the operation may write to.
+	 * @returns True if the requested number of rows were successfully created, otherwise false.
+	 */
+	UE_API bool CreateInputRows(UE::Editor::DataStorage::FRowHandleArray& OutInputRows, UE::Editor::DataStorage::FRowHandleArrayView SourceRows,
+		bool bForce, bool bAddDescription) const;
+	
+protected:
+	UE_API virtual void GetQueryDescription(UE::Editor::DataStorage::FQueryDescription& OutDescription) const override;
+
+	UE::Editor::DataStorage::TableHandle InputTableForce;
+	UE::Editor::DataStorage::TableHandle InputTableForceDescr;
+};
+
+#undef UE_API
