[TedsOperations][UnrealEd] Introduced new ViewportDropHandler for generalized drag&drop(-receiving) behavior for viewport widgets using teds drop operations. Added ViewportDropHandlers to both Prefab and Level Editors. Introduced new EntityLevelViewport widget to override the default viewport with.

#Synced-CL 48209826

diff --git a/Engine/Plugins/Experimental/EditorDataStorageFeatures/Source/TedsOperations/Private/DragAndDrop/Widgets/PreviewWidgetDropHandler.cpp b/Engine/Plugins/Experimental/EditorDataStorageFeatures/Source/TedsOperations/Private/DragAndDrop/Widgets/PreviewWidgetDropHandler.cpp
new file mode 100644
--- /dev/null
+++ b/Engine/Plugins/Experimental/EditorDataStorageFeatures/Source/TedsOperations/Private/DragAndDrop/Widgets/PreviewWidgetDropHandler.cpp
@@ -0,0 +1,224 @@
+// Copyright Epic Games, Inc. All Rights Reserved.
+
+#include "DragAndDrop/Widgets/PreviewWidgetDropHandler.h"
+
+#include "Deletion/DeletionOperationSystem.h"
+#include "DragAndDrop/DropOperationInput.h"
+#include "DragAndDrop/DropOperationSystem.h"
+#include "Elements/Interfaces/TypedElementDataStorageInterface.h"
+#include "ISceneOutliner.h"
+#include "LevelEditor.h"
+#include "Modules/ModuleManager.h"
+
+#define LOCTEXT_NAMESPACE "PreviewWidgetDropHandler"
+
+namespace UE::Editor
+{
+	
+namespace PreviewWidgetDropHandler_Private
+{
+	
+static void UpdateOutliner()
+{
+	// We update the level outliners manually, because they don't update automatically yet during drag&drop.
+	const FLevelEditorModule& LevelEditorModule = FModuleManager::GetModuleChecked<FLevelEditorModule>(TEXT("LevelEditor"));
+	if (TSharedPtr<ILevelEditor> LevelEditor = LevelEditorModule.GetLevelEditorInstance().Pin())
+	{
+		for (const TWeakPtr<ISceneOutliner>& SceneOutlinerPtr : LevelEditor->GetAllSceneOutliners())
+		{
+			if (TSharedPtr<ISceneOutliner> SceneOutliner = SceneOutlinerPtr.Pin())
+			{
+				SceneOutliner->FullRefresh();
+			}
+		}
+	}
+}
+	
+}
+	
+FPreviewWidgetDropHandler::FPreviewWidgetDropHandler(TWeakObjectPtr<UTypedElementSelectionSet> SelectionSet)
+	: FWidgetDropHandler(MoveTemp(SelectionSet))
+{
+}
+
+FPreviewWidgetDropHandler::~FPreviewWidgetDropHandler()
+{
+}
+
+FWidgetDropHandler::FDropResult FPreviewWidgetDropHandler::ResetOperations(DataStorage::ICoreProvider& Storage, UDropOperationSystem& DropSystem)
+{
+	using namespace UE::Editor::DataStorage;
+
+	if (HasPreviews())
+	{
+		RemovePreviews(Storage, DropSystem); // @todo: Only previews of operations that change need to be recreated.
+	}
+	
+	FDropResult Result = FWidgetDropHandler::ResetOperations(Storage, DropSystem);
+
+	if (bIsValidDrop)
+	{
+		CreatePreviews(Storage, DropSystem);
+	}
+	
+	return Result;
+}
+	
+TOptional<FWidgetDropHandler::FDropResult> FPreviewWidgetDropHandler::Update(DataStorage::ICoreProvider& Storage, UDropOperationSystem& DropSystem, const FGeometry& Geometry,
+	const FDragDropEvent& InputEvent)
+{
+	using namespace UE::Editor::DataStorage;
+
+	EUpdateParameters Result = UpdateParameters(Storage, Geometry, InputEvent);
+	if (Result == EUpdateParameters::NoChanges)
+	{
+		return {};
+	}
+
+	if (Result == EUpdateParameters::ResetOperations)
+	{
+		SetRowsFromParameters(Storage, InputRows.GetRows()); // Only update the input rows when we need to.
+		return ResetOperations(Storage, DropSystem);
+	}
+
+	if (HasPreviews())
+	{
+		UpdatePreviews(Storage, DropSystem);
+	}
+	return {};
+}
+	
+void FPreviewWidgetDropHandler::Stop(DataStorage::ICoreProvider& Storage, UDropOperationSystem& DropSystem)
+{
+	if (HasPreviews())
+	{
+		RemovePreviews(Storage, DropSystem);
+	}
+	
+	FWidgetDropHandler::Stop(Storage, DropSystem);
+}
+
+void FPreviewWidgetDropHandler::Drop(DataStorage::ICoreProvider& Storage, UDropOperationSystem& DropSystem, const FGeometry& Geometry,
+		const FDragDropEvent& InputEvent)
+{
+	if (HasPreviews())
+	{
+		RemovePreviews(Storage, DropSystem);
+	}
+	
+	FWidgetDropHandler::Drop(Storage, DropSystem, Geometry, InputEvent);
+}
+
+static void GetCreatedRows(DataStorage::FRowHandleArray& OutRows, DataStorage::ICoreProvider& Storage, DataStorage::FRowHandleArrayView ResultRows)
+{
+	using namespace UE::Editor::DataStorage;
+	using namespace UE::Editor::DataStorage::Operations;
+
+	OutRows.Reset();
+	OutRows.Reserve(ResultRows.Num());
+	for (RowHandle InputRow : ResultRows)
+	{
+		if (const FResultColumn* Column = Storage.GetColumn<FResultColumn>(InputRow))
+		{
+			OutRows.Append(Column->Value.Created.GetRows());
+		}
+	}
+}
+	
+void FPreviewWidgetDropHandler::CreatePreviews(DataStorage::ICoreProvider& Storage, UDropOperationSystem& DropSystem)
+{
+	using namespace UE::Editor::DataStorage;
+	using namespace UE::Editor::DataStorage::Operations;
+	if (!ensure(bIsValidDrop && PreviewInputRows.IsEmpty()))
+	{
+		return;
+	}
+
+	if (!Storage.FindFactory<UDeletionOperationSystem>())
+	{
+		return; // We need to be able to revert the previews.
+	}	
+
+	// Prepare input rows for a preview creation call.
+	for (RowHandle InputRow : InputRows.GetRows())
+	{
+		Storage.AddColumn<FDropPreviewTag>(InputRow);
+	}
+
+	// Create the previews from our input.
+	(void)DropSystem.Apply(InputRows.GetRows());
+
+	// Gather the results.
+	FRowHandleArray Created;
+	GetCreatedRows(Created, Storage, InputRows.GetRows());
+
+	// Clean up the input rows.
+	for (RowHandle InputRow : InputRows.GetRows())
+	{
+		Storage.RemoveColumn<FResultColumn>(InputRow);
+		Storage.RemoveColumn<FDropPreviewTag>(InputRow);
+	}
+
+	// Create rows that hold our created previews for updates.
+	DropSystem.CreateInputRows(PreviewInputRows, Created.GetRows(), InvalidRowHandle, false);	
+	
+	PreviewWidgetDropHandler_Private::UpdateOutliner();
+}
+
+void FPreviewWidgetDropHandler::UpdatePreviews(DataStorage::ICoreProvider& Storage, UDropOperationSystem& DropSystem)
+{
+	using namespace UE::Editor::DataStorage;
+	using namespace UE::Editor::DataStorage::Operations;
+	if (!ensure(!PreviewInputRows.IsEmpty()))
+	{
+		return;
+	}
+	
+	SetRowsFromParameters(Storage, PreviewInputRows.GetRows());
+	(void)DropSystem.Apply(PreviewInputRows.GetRows());
+}
+
+void FPreviewWidgetDropHandler::RemovePreviews(DataStorage::ICoreProvider& Storage, UDropOperationSystem& DropSystem)
+{
+	using namespace UE::Editor::DataStorage;
+	using namespace UE::Editor::DataStorage::Operations;
+	if (!ensure(!PreviewInputRows.IsEmpty()))
+	{
+		return;
+	}
+
+	UDeletionOperationSystem* DeletionSystem = Storage.FindFactory<UDeletionOperationSystem>();
+	if (!ensureMsgf(DeletionSystem, TEXT("Deletion system not found. Unable to delete previews!")))
+	{
+		return;
+	}
+
+	// Delete our created previews.
+	(void)DeletionSystem->Apply(PreviewInputRows.GetRows());
+
+	// Free our preview input rows.
+	DropSystem.RemoveInputRows(PreviewInputRows.GetRows());
+	PreviewInputRows.Empty();
+
+	PreviewWidgetDropHandler_Private::UpdateOutliner();
+}
+
+void FPreviewWidgetDropHandler::GetPreviewRows(DataStorage::FRowHandleArray& OutRows, DataStorage::ICoreProvider& Storage) const
+{
+	using namespace UE::Editor::DataStorage;
+	using namespace UE::Editor::DataStorage::Operations;
+	
+	OutRows.Reset();
+	OutRows.Reserve(PreviewInputRows.Num());
+	for (RowHandle PreviewInputRow : PreviewInputRows.GetRows())
+	{
+		if (const FSourceColumn* Column = Storage.GetColumn<FSourceColumn>(PreviewInputRow))
+		{
+			OutRows.Add(Column->Value);
+		}
+	}
+}
+	
+}
+
+#undef LOCTEXT_NAMESPACE
diff --git a/Engine/Plugins/Experimental/EditorDataStorageFeatures/Source/TedsOperations/Private/DragAndDrop/Widgets/ViewportDropHandler.cpp b/Engine/Plugins/Experimental/EditorDataStorageFeatures/Source/TedsOperations/Private/DragAndDrop/Widgets/ViewportDropHandler.cpp
new file mode 100644
--- /dev/null
+++ b/Engine/Plugins/Experimental/EditorDataStorageFeatures/Source/TedsOperations/Private/DragAndDrop/Widgets/ViewportDropHandler.cpp
@@ -0,0 +1,199 @@
+// Copyright Epic Games, Inc. All Rights Reserved.
+
+#include "DragAndDrop/Widgets/ViewportDropHandler.h"
+
+#include "ActorFactories/ActorFactory.h"
+#include "DragAndDrop/DropOperationInput.h"
+#include "Editor/ObjectPositioning.h"
+#include "EditorModeManager.h"
+#include "EditorViewportClient.h"
+#include "Elements/Columns/TypedElementCompatibilityColumns.h"
+#include "Elements/Columns/TypedElementHandleColumn.h"
+#include "Elements/Framework/TypedElementRegistry.h"
+#include "Elements/Interfaces/TypedElementDataStorageInterface.h"
+#include "Elements/Interfaces/TypedElementWorldInterface.h"
+#include "Input/DragAndDrop.h"
+#include "Modules/ModuleManager.h"
+#include "SceneView.h"
+#include "Settings/LevelEditorViewportSettings.h"
+#include "SnappingUtils.h"
+
+#define LOCTEXT_NAMESPACE "ViewportDropHandler"
+
+namespace UE::Editor
+{
+
+namespace ViewportDropHandler_Private
+{
+	static FIntPoint GetViewportCursorPosition(FEditorViewportClient& ViewportClient, const FGeometry& Geometry, const FDragDropEvent& InputEvent)
+	{
+		FIntPoint Pos = (Geometry.AbsoluteToLocal(InputEvent.GetScreenSpacePosition()) * Geometry.Scale).IntPoint();
+
+		FIntPoint ViewportOrigin, ViewportSize;
+		ViewportClient.GetViewportDimensions(ViewportOrigin, ViewportSize);
+		Pos -= ViewportOrigin;
+		return Pos;
+	}
+
+	// TEv1 workaround until tracing properly supports entities (and TEDS).
+	static void AddIgnoredObjects(FCollisionQueryParams& OutParams, DataStorage::ICoreProvider& Storage, DataStorage::FRowHandleArrayView IgnoreRows)
+	{
+		using namespace UE::Editor::DataStorage;
+
+		UTypedElementRegistry* Registry = UTypedElementRegistry::GetInstance();
+		for (RowHandle Row : IgnoreRows)
+		{
+			if (Registry)
+			{
+				if (Compatibility::FTypedElementColumn* Column = Storage.GetColumn<Compatibility::FTypedElementColumn>(Row))
+				{
+					if (TTypedElement<ITypedElementWorldInterface> WorldElement = Registry->GetElement<ITypedElementWorldInterface>(Column->Handle))
+					{
+						WorldElement.AddIgnoredElementToCollisionQueryParams(Column->Handle, OutParams, true);
+						continue;
+					}
+				}
+			}
+
+			if (FTypedElementUObjectColumn* Column = Storage.GetColumn<FTypedElementUObjectColumn>(Row))
+			{
+				if (UObject* Object = Column->Object.Get())
+				{
+					OutParams.AddIgnoredSourceObject(Object);
+				}
+			}
+		}
+	}
+
+	// Optimally this should be provided by some public gizmo functionality instead.
+	static FTransform TraceDropTransform(FEditorViewportClient& Client, const FIntPoint& CursorPos, DataStorage::ICoreProvider& Storage, DataStorage::FRowHandleArrayView IgnoreRows)
+	{
+		FSceneViewFamilyContext ViewFamily(FSceneViewFamily::ConstructionValues(
+			Client.Viewport, Client.GetScene(), Client.EngineShowFlags));
+		FSceneView* View = Client.CalcSceneView(&ViewFamily);
+		FViewportCursorLocation Cursor(View, &Client, CursorPos.X, CursorPos.Y);
+
+		FCollisionQueryParams Params(SCENE_QUERY_STAT(DragDropTrace), true);
+		AddIgnoredObjects(Params, Storage, IgnoreRows);
+
+		const Positioning::FObjectPositioningTraceResult TraceResult =
+			Positioning::TraceWorldForPositionWithDefault(Cursor, *View, &Params);
+
+		FVector Location = TraceResult.Location;
+		FSnappingUtils::SnapPointToGrid(Location, FVector::ZeroVector);
+
+		FVector VertexNormal;
+		if (!FSnappingUtils::SnapLocationToNearestVertex(Location, FVector2D(CursorPos), &Client, VertexNormal, false, {}))
+		{
+			VertexNormal = TraceResult.SurfaceNormal;
+		}
+
+		FQuat Rotation = FQuat::Identity;
+		if (GetDefault<ULevelEditorViewportSettings>()->SnapToSurface.bEnabled)
+		{
+			Rotation = FindActorAlignmentRotation(FQuat::Identity, FVector::UpVector, VertexNormal);
+		}
+		return FTransform(Rotation, Location);
+	}
+
+	UTypedElementSelectionSet* GetSelectionSetFromViewportClient(const TSharedPtr<FEditorViewportClient>& ViewportClient)
+	{
+		if (ViewportClient)
+		{
+			if (const FEditorModeTools* ModeTools = ViewportClient->GetModeTools())
+			{
+				return ModeTools->GetEditorSelectionSet();
+			}
+		}
+		return nullptr;
+	}
+}
+
+FViewportDropHandler::FViewportDropHandler(TSharedPtr<FEditorViewportClient> InViewportClient, FTargetRowGetter InGetDefaultTargetRow)
+	: FPreviewWidgetDropHandler(ViewportDropHandler_Private::GetSelectionSetFromViewportClient(InViewportClient))
+	, ViewportClient(MoveTemp(InViewportClient))
+	, GetDefaultTargetRow(MoveTemp(InGetDefaultTargetRow))
+	, LastCursorPos()
+	, LastTargetRow(DataStorage::InvalidRowHandle)
+	, LastTransform(FTransform::Identity)
+{
+}
+
+FViewportDropHandler::~FViewportDropHandler()
+{
+}
+
+FWidgetDropHandler::EUpdateParameters FViewportDropHandler::UpdateParameters(DataStorage::ICoreProvider& Storage, const FGeometry& Geometry,
+	const FDragDropEvent& InputEvent)
+{
+	using namespace UE::Editor::DataStorage;
+
+	FIntPoint CursorPos = ViewportClient ? ViewportDropHandler_Private::GetViewportCursorPosition(*ViewportClient, Geometry, InputEvent) : FIntPoint();	
+	if (LastCursorPos != CursorPos && ensure(ViewportClient))
+	{
+		FRowHandleArray IgnoreRows;
+		GetPreviewRows(IgnoreRows, Storage);
+		FTransform Transform = ViewportDropHandler_Private::TraceDropTransform(*ViewportClient, CursorPos, Storage, IgnoreRows.GetRows());
+
+		// By default, we just get the target row from our creator / viewport client.
+		// @todo: Introduce cases to get target by tracing.
+		RowHandle TargetRow = (GetDefaultTargetRow && ViewportClient) ? GetDefaultTargetRow(*ViewportClient) : InvalidRowHandle;
+		
+		EUpdateParameters Result = LastTargetRow != TargetRow ? EUpdateParameters::ResetOperations : EUpdateParameters::MinorChanges;
+		LastCursorPos = CursorPos;
+		LastTargetRow = TargetRow;
+		LastTransform = Transform;
+		return Result;		
+	}
+	return EUpdateParameters::NoChanges;
+}
+
+void FViewportDropHandler::SetRowsFromParameters(DataStorage::ICoreProvider& Storage, DataStorage::FRowHandleArrayView Rows) const
+{
+	using namespace UE::Editor::DataStorage;
+	using namespace UE::Editor::DataStorage::Operations;
+	
+	for (RowHandle InputRow : Rows)
+	{
+		Storage.AddColumn(InputRow, FDropTargetColumn{ .Value = LastTargetRow });
+		Storage.AddColumn(InputRow, FDropTransformColumn{ .Value = LastTransform });
+	}
+}
+
+void FViewportDropHandler::Stop(DataStorage::ICoreProvider& Storage, UDropOperationSystem& DropSystem)
+{
+	FPreviewWidgetDropHandler::Stop(Storage, DropSystem);
+	LastCursorPos = FIntPoint();
+	LastTargetRow = DataStorage::InvalidRowHandle;
+	LastTransform = FTransform::Identity;
+}
+
+void FViewportDropHandler::CreatePreviews(DataStorage::ICoreProvider& Storage, UDropOperationSystem& DropSystem)
+{
+	FPreviewWidgetDropHandler::CreatePreviews(Storage, DropSystem);
+	InvalidateViewport();
+}
+
+void FViewportDropHandler::UpdatePreviews(DataStorage::ICoreProvider& Storage, UDropOperationSystem& DropSystem)
+{
+	FPreviewWidgetDropHandler::UpdatePreviews(Storage, DropSystem);
+	InvalidateViewport();
+}
+
+void FViewportDropHandler::RemovePreviews(DataStorage::ICoreProvider& Storage, UDropOperationSystem& DropSystem)
+{
+	FPreviewWidgetDropHandler::RemovePreviews(Storage, DropSystem);
+	InvalidateViewport();
+}
+
+void FViewportDropHandler::InvalidateViewport()
+{
+	if (ViewportClient && ViewportClient->Viewport)
+	{
+		ViewportClient->Viewport->InvalidateDisplay();
+	}
+}
+
+}
+
+#undef LOCTEXT_NAMESPACE
diff --git a/Engine/Plugins/Experimental/EditorDataStorageFeatures/Source/TedsOperations/Public/DragAndDrop/Widgets/PreviewWidgetDropHandler.h b/Engine/Plugins/Experimental/EditorDataStorageFeatures/Source/TedsOperations/Public/DragAndDrop/Widgets/PreviewWidgetDropHandler.h
new file mode 100644
--- /dev/null
+++ b/Engine/Plugins/Experimental/EditorDataStorageFeatures/Source/TedsOperations/Public/DragAndDrop/Widgets/PreviewWidgetDropHandler.h
@@ -0,0 +1,43 @@
+// Copyright Epic Games, Inc. All Rights Reserved.
+
+#pragma once
+
+#include "DragAndDrop/Widgets/WidgetDropHandler.h"
+
+#define UE_API TEDSOPERATIONS_API
+
+namespace UE::Editor
+{
+
+/**
+ * Utility type for widgets that accepts drag&drop-events and preview creation on drag-over using teds operations.
+ */
+class FPreviewWidgetDropHandler : public FWidgetDropHandler
+{
+public:
+	UE_API FPreviewWidgetDropHandler(TWeakObjectPtr<UTypedElementSelectionSet> SelectionSet);
+	UE_API virtual ~FPreviewWidgetDropHandler() override;
+
+	bool HasPreviews() const { return !PreviewInputRows.IsEmpty(); }
+	
+protected:
+	UE_API virtual FDropResult ResetOperations(DataStorage::ICoreProvider& Storage, UDropOperationSystem& DropSystem) override;	
+	UE_API virtual TOptional<FDropResult> Update(DataStorage::ICoreProvider& Storage, UDropOperationSystem& DropSystem, const FGeometry& Geometry,
+		const FDragDropEvent& InputEvent) override;
+	UE_API virtual void Stop(DataStorage::ICoreProvider& Storage, UDropOperationSystem& DropSystem) override;
+	UE_API virtual void Drop(DataStorage::ICoreProvider& Storage, UDropOperationSystem& DropSystem, const FGeometry& Geometry,
+		const FDragDropEvent& InputEvent) override;
+	
+	UE_API virtual void CreatePreviews(DataStorage::ICoreProvider& Storage, UDropOperationSystem& DropSystem);
+	UE_API virtual void UpdatePreviews(DataStorage::ICoreProvider& Storage, UDropOperationSystem& DropSystem);
+	UE_API virtual void RemovePreviews(DataStorage::ICoreProvider& Storage, UDropOperationSystem& DropSystem);
+
+	UE_API void GetPreviewRows(DataStorage::FRowHandleArray& OutRows, DataStorage::ICoreProvider& Storage) const;
+
+protected:
+	DataStorage::FRowHandleArray PreviewInputRows;
+};
+
+}
+
+#undef UE_API
\ No newline at end of file
diff --git a/Engine/Plugins/Experimental/EditorDataStorageFeatures/Source/TedsOperations/Public/DragAndDrop/Widgets/ViewportDropHandler.h b/Engine/Plugins/Experimental/EditorDataStorageFeatures/Source/TedsOperations/Public/DragAndDrop/Widgets/ViewportDropHandler.h
new file mode 100644
--- /dev/null
+++ b/Engine/Plugins/Experimental/EditorDataStorageFeatures/Source/TedsOperations/Public/DragAndDrop/Widgets/ViewportDropHandler.h
@@ -0,0 +1,49 @@
+// Copyright Epic Games, Inc. All Rights Reserved.
+
+#pragma once
+
+#include "DragAndDrop/Widgets/PreviewWidgetDropHandler.h"
+
+class FEditorViewportClient;
+
+#define UE_API TEDSOPERATIONS_API
+
+namespace UE::Editor
+{
+
+class FViewportDropHandler : public FPreviewWidgetDropHandler
+{
+public:
+	using FTargetRowGetter = TFunction<DataStorage::RowHandle(const FEditorViewportClient&)>;
+
+	/**
+	 * Construct a viewport drop handler.
+	 * @param InViewportClient The client of the viewport the handler is representing.
+	 * @param InGetDefaultTargetRow A function to return a "default" target during drag&drop that is used when not tracing.
+	 */
+	UE_API FViewportDropHandler(TSharedPtr<FEditorViewportClient> InViewportClient, FTargetRowGetter InGetDefaultTargetRow);
+	UE_API virtual ~FViewportDropHandler() override;
+
+protected:
+	UE_API virtual EUpdateParameters UpdateParameters(DataStorage::ICoreProvider& Storage, const FGeometry& Geometry, const FDragDropEvent& InputEvent) override;
+	UE_API virtual void SetRowsFromParameters(DataStorage::ICoreProvider& Storage, DataStorage::FRowHandleArrayView Rows) const override;
+	UE_API virtual void Stop(DataStorage::ICoreProvider& Storage, UDropOperationSystem& DropSystem) override;
+	
+	UE_API virtual void CreatePreviews(DataStorage::ICoreProvider& Storage, UDropOperationSystem& DropSystem) override;
+	UE_API virtual void UpdatePreviews(DataStorage::ICoreProvider& Storage, UDropOperationSystem& DropSystem) override;
+	UE_API virtual void RemovePreviews(DataStorage::ICoreProvider& Storage, UDropOperationSystem& DropSystem) override;
+
+	UE_API void InvalidateViewport();
+	
+protected:
+	TSharedPtr<FEditorViewportClient> ViewportClient;
+	FTargetRowGetter GetDefaultTargetRow; // A function to return a "default" target during drag&drop that is used when not tracing.
+	
+	FIntPoint LastCursorPos;
+	DataStorage::RowHandle LastTargetRow;
+	FTransform LastTransform;
+};
+
+}
+
+#undef UE_API
\ No newline at end of file
diff --git a/Engine/Plugins/Experimental/EditorDataStorageFeatures/Source/TedsOperations/Public/DragAndDrop/Widgets/WidgetDropHandler.h b/Engine/Plugins/Experimental/EditorDataStorageFeatures/Source/TedsOperations/Public/DragAndDrop/Widgets/WidgetDropHandler.h
--- a/Engine/Plugins/Experimental/EditorDataStorageFeatures/Source/TedsOperations/Public/DragAndDrop/Widgets/WidgetDropHandler.h
+++ b/Engine/Plugins/Experimental/EditorDataStorageFeatures/Source/TedsOperations/Public/DragAndDrop/Widgets/WidgetDropHandler.h
@@ -2,14 +2,16 @@
 
 #pragma once
 
-#include "DataStorage/Handles.h"
-#include "DragAndDrop/IWidgetDropHandler.h"
 #include "Elements/Framework/TypedElementRowHandleArray.h"
 #include "Elements/Framework/TypedElementRowHandleArrayView.h"
+#include "Input/Reply.h"
 #include "Misc/Optional.h"
 
-class UTypedElementSelectionSet;
+class FDragDropEvent;
+class FDragDropOperation;
 class UDropOperationSystem;
+class UTypedElementSelectionSet;
+struct FGeometry;
 namespace UE::Editor::DataStorage
 {
 class ICoreProvider;
@@ -20,32 +22,35 @@ class ICoreProvider;
 namespace UE::Editor
 {
 
-class FWidgetDropHandler : public IWidgetDropHandler
+/**
+ * Utility type for widgets that accepts drag&drop-events using teds operations.
+ */ 
+class FWidgetDropHandler
 {
 public:
 	struct FDropResult
 	{
-		bool  bCanDrop = false;
-		FText Tooltip  = FText();
+		bool bCanDrop = false;
+		FText Tooltip = FText();
 	};
 
 protected:
 	enum class EUpdateParameters
 	{
-		NoChanges,       // Nothing has changed.
-		MinorChanges,    // Some changes to input, but no need to reset operations.
+		NoChanges, // Nothing has changed.
+		MinorChanges, // Some changes to input, but no need to reset operations.
 		ResetOperations, // Input has major changes, reset operations.
 	};
 
 public:
 	UE_API FWidgetDropHandler(TWeakObjectPtr<UTypedElementSelectionSet> SelectionSet);
-	UE_API virtual ~FWidgetDropHandler() override;
+	UE_API virtual ~FWidgetDropHandler();
 
-	UE_API virtual bool                OnDragEnter(const FGeometry& Geometry, const FDragDropEvent& InputEvent) override;
-	UE_API virtual FReply              OnDragOver(const FGeometry& Geometry, const FDragDropEvent& InputEvent) override;
-	UE_API virtual bool                OnDragLeave(const FDragDropEvent& InputEvent) override;
-	UE_API virtual FReply              OnDrop(const FGeometry& Geometry, const FDragDropEvent& InputEvent) override;
-	UE_API virtual FDragDropOperation* GetCurrentOperation() const override;
+	UE_API virtual bool OnDragEnter(const FGeometry& Geometry, const FDragDropEvent& InputEvent);
+	UE_API virtual FReply OnDragOver(const FGeometry& Geometry, const FDragDropEvent& InputEvent);
+	UE_API virtual bool OnDragLeave(const FDragDropEvent& InputEvent);
+	UE_API virtual FReply OnDrop(const FGeometry& Geometry, const FDragDropEvent& InputEvent);
+	UE_API virtual FDragDropOperation* GetCurrentOperation() const;
 
 protected:
 	/** Creates input rows for the given drag&drop input event. */
@@ -55,8 +60,9 @@ protected:
 	/** Finds suitable operations for the given input rows and sets 'bIsValidDrop'. */
 	UE_API virtual FDropResult ResetOperations(DataStorage::ICoreProvider& Storage, UDropOperationSystem& DropSystem);
 	
-	/** Updates all input rows for the updated input parameters. */
+	/** Method to update cached parameters for the given input data. */
 	UE_API virtual EUpdateParameters UpdateParameters(DataStorage::ICoreProvider& Storage, const FGeometry& Geometry, const FDragDropEvent& InputEvent);
+	/** Method to assign cached parameters to the given rows. */
 	UE_API virtual void SetRowsFromParameters(DataStorage::ICoreProvider& Storage, DataStorage::FRowHandleArrayView Rows) const;
 	
 	UE_API virtual FDropResult Start(DataStorage::ICoreProvider& Storage, UDropOperationSystem& DropSystem, const FGeometry& Geometry,
@@ -72,9 +78,9 @@ protected:
 protected:
 	TWeakObjectPtr<UTypedElementSelectionSet> SelectionSet;
 
-	FDragDropOperation*          CurrentOperation;
+	FDragDropOperation* CurrentOperation;
 	DataStorage::FRowHandleArray InputRows;
-	bool                         bIsValidDrop;
+	bool bIsValidDrop;
 };
 
 }
diff --git a/Engine/Plugins/Experimental/EditorDataStorageFeatures/Source/TedsOperations/TedsOperations.build.cs b/Engine/Plugins/Experimental/EditorDataStorageFeatures/Source/TedsOperations/TedsOperations.build.cs
--- a/Engine/Plugins/Experimental/EditorDataStorageFeatures/Source/TedsOperations/TedsOperations.build.cs
+++ b/Engine/Plugins/Experimental/EditorDataStorageFeatures/Source/TedsOperations/TedsOperations.build.cs
@@ -25,7 +25,6 @@ public class TedsOperations : ModuleRules
 					"Core",
 					"CoreUObject",
 					"SlateCore",
-					"ToolWidgets",
 					"TypedElementFramework",
 					"UnrealEd",
 				});
diff --git a/Engine/Source/Developer/ToolWidgets/Public/DragAndDrop/IWidgetDropHandler.h b/Engine/Source/Developer/ToolWidgets/Public/DragAndDrop/IWidgetDropHandler.h
deleted file mode 100644
diff --git a/Engine/Source/Editor/LevelEditor/Private/LevelViewportLayoutEntity.h b/Engine/Source/Editor/LevelEditor/Private/LevelViewportLayoutEntity.h
deleted file mode 100644
diff --git a/Engine/Source/Editor/LevelEditor/Public/LevelViewportLayoutEntity.h b/Engine/Source/Editor/LevelEditor/Public/LevelViewportLayoutEntity.h
new file mode 100644
--- /dev/null
+++ b/Engine/Source/Editor/LevelEditor/Public/LevelViewportLayoutEntity.h
@@ -0,0 +1,37 @@
+// Copyright Epic Games, Inc. All Rights Reserved.
+
+#pragma once
+
+#include "CoreMinimal.h"
+#include "Widgets/SWidget.h"
+#include "LevelViewportLayout.h"
+
+#define UE_API LEVELEDITOR_API
+
+class FLevelEditorViewportClient;
+class SLevelViewport;
+
+class FLevelViewportLayoutEntity : public ILevelViewportLayoutEntity
+{
+public:
+
+	UE_API FLevelViewportLayoutEntity(TSharedPtr<SAssetEditorViewport> InLevelViewport);
+	UE_API virtual TSharedRef<SWidget> AsWidget() const override;
+	UE_API virtual TSharedPtr<SLevelViewport> AsLevelViewport() const override;
+
+	UE_API bool IsPlayInEditorViewportActive() const;
+	UE_API void RegisterGameViewportIfPIE();
+	UE_API void SetKeyboardFocus();
+	UE_API void OnLayoutDestroyed();
+	UE_API void SaveConfig(const FString& ConfigSection);
+	UE_API FLevelEditorViewportClient& GetLevelViewportClient() const;
+	UE_API FName GetType() const;
+	UE_API void TakeHighResScreenShot() const;
+
+protected:
+
+	/** This entity's level viewport */
+	TSharedRef<SLevelViewport> LevelViewport;
+};
+
+#undef UE_API
