[TEDS] Made more data searchable and sortable.

#Synced-CL 48364666

diff --git a/Engine/Plugins/Experimental/EditorDataStorageFeatures/Source/TedsAssetData/Public/TedsAssetDataColumns.h b/Engine/Plugins/Experimental/EditorDataStorageFeatures/Source/TedsAssetData/Public/TedsAssetDataColumns.h
--- a/Engine/Plugins/Experimental/EditorDataStorageFeatures/Source/TedsAssetData/Public/TedsAssetDataColumns.h
+++ b/Engine/Plugins/Experimental/EditorDataStorageFeatures/Source/TedsAssetData/Public/TedsAssetDataColumns.h
@@ -27,7 +27,7 @@ struct FAssetPathColumn_Experimental : public FEditorDataStorageColumn
 {
 	GENERATED_BODY()
 	
-	UPROPERTY()
+	UPROPERTY(meta = (Searchable, Sortable))
 	FName Path;
 };
 
@@ -83,7 +83,7 @@ struct FAssetNameColumn final : public FEditorDataStorageColumn
 	GENERATED_BODY()
 
 	// The name of the asset
-	UPROPERTY()
+	UPROPERTY(meta = (Searchable, Sortable))
 	FName Name;
 };
 
@@ -138,7 +138,7 @@ struct FVirtualPathColumn_Experimental : public FEditorDataStorageColumn
 {
 	GENERATED_BODY()
 
-	UPROPERTY(meta = (Searchable))
+	UPROPERTY(meta = (Searchable, Sortable))
 	FName VirtualPath;
 };
 
@@ -155,7 +155,7 @@ struct FItemTextAttributeColumn_Experimental : public FItemAttributeBaseColumn_E
 {
 	GENERATED_BODY()
 
-	UPROPERTY()
+	UPROPERTY(meta = (Searchable, Sortable))
 	FText Value;
 };
 
@@ -164,6 +164,6 @@ struct FItemStringAttributeColumn_Experimental : public FItemAttributeBaseColumn
 {
 	GENERATED_BODY()
 
-	UPROPERTY()
+	UPROPERTY(meta = (Searchable, Sortable))
 	FString Value;
 };
\ No newline at end of file
