Fixed new drag&drop operations workflow not deleting preview actors in some cases.

#Synced-CL 48539628

diff --git a/Engine/Plugins/Experimental/EditorDataStorageFeatures/Source/TedsOperations/Private/DragAndDrop/Widgets/PreviewWidgetDropHandler.cpp b/Engine/Plugins/Experimental/EditorDataStorageFeatures/Source/TedsOperations/Private/DragAndDrop/Widgets/PreviewWidgetDropHandler.cpp
--- a/Engine/Plugins/Experimental/EditorDataStorageFeatures/Source/TedsOperations/Private/DragAndDrop/Widgets/PreviewWidgetDropHandler.cpp
+++ b/Engine/Plugins/Experimental/EditorDataStorageFeatures/Source/TedsOperations/Private/DragAndDrop/Widgets/PreviewWidgetDropHandler.cpp
@@ -2,6 +2,7 @@
 
 #include "DragAndDrop/Widgets/PreviewWidgetDropHandler.h"
 
+#include "Deletion/DeletionOperationInput.h"
 #include "Deletion/DeletionOperationSystem.h"
 #include "DragAndDrop/DropOperationInput.h"
 #include "DragAndDrop/DropOperationSystem.h"
@@ -193,6 +194,12 @@ void FPreviewWidgetDropHandler::RemovePreviews(DataStorage::ICoreProvider& Stora
 		return;
 	}
 
+	// Make sure the previews gets deleted.
+	for (RowHandle PreviewInputRow : PreviewInputRows.GetRows())
+	{
+		Storage.AddColumn<FDeletionForceTag>(PreviewInputRow);	
+	}
+	
 	// Delete our created previews.
 	(void)DeletionSystem->Apply(PreviewInputRows.GetRows());
 
diff --git a/Engine/Source/Editor/UnrealEd/Private/Factories/ActorFactory.cpp b/Engine/Source/Editor/UnrealEd/Private/Factories/ActorFactory.cpp
--- a/Engine/Source/Editor/UnrealEd/Private/Factories/ActorFactory.cpp
+++ b/Engine/Source/Editor/UnrealEd/Private/Factories/ActorFactory.cpp
@@ -487,8 +487,8 @@ AActor* UActorFactory::SpawnActor(UObject* InAsset, ULevel* InLevel, const FTran
 		{
 			bIsPlacementSystemCreatingPreviewElements = PlacementSubsystem->IsCreatingPreviewElements();
 		}
-		SpawnParamsCopy.bTemporaryEditorActor = bIsCreatingPreviewElements || bIsPlacementSystemCreatingPreviewElements;
-		SpawnParamsCopy.bHideFromSceneOutliner = bIsPlacementSystemCreatingPreviewElements;
+		SpawnParamsCopy.bTemporaryEditorActor |= bIsCreatingPreviewElements || bIsPlacementSystemCreatingPreviewElements;
+		SpawnParamsCopy.bHideFromSceneOutliner |= bIsPlacementSystemCreatingPreviewElements;
 
 		return LocalLevel->OwningWorld->SpawnActor(DefaultActor->GetClass(), &InTransform, SpawnParamsCopy);
 	}
