TEDS Alert: Add support for the new hierarchy interface + various fixes

AddAlert now has an argument for the hierarchy name to use when calculating child alerts
- The old FTableRowParent queries have been removed completely, so you need to use the new interface
Added an FTedsAlertHierarchyColumn to rows with the alert or child alert column to signify which hierarchy this particular alert operates on
- Added new subqueries to copy/swap FTedsAlertHierarchyColumn as alerts are moved
Certain alert queries are now registered "per hierarchy" lazily since queries can only operate on single hierarchies
- These Queries check the alert to only propagate to the parent to the relevant hierearchy for each alert.
Cleanup alerts on deletion properly by removing the action and hierarchy column
Instead of using a cached parent + sync tags, the alert hierarchies use the ParentChangedColumn to track changes

#Synced-CL 48607495

diff --git a/Engine/Plugins/Experimental/EditorDataStorageFeatures/Source/TedsAlerts/Private/TedsAlerts.cpp b/Engine/Plugins/Experimental/EditorDataStorageFeatures/Source/TedsAlerts/Private/TedsAlerts.cpp
--- a/Engine/Plugins/Experimental/EditorDataStorageFeatures/Source/TedsAlerts/Private/TedsAlerts.cpp
+++ b/Engine/Plugins/Experimental/EditorDataStorageFeatures/Source/TedsAlerts/Private/TedsAlerts.cpp
@@ -6,26 +6,49 @@
 #include "Elements/Columns/TypedElementMiscColumns.h"
 #include "Elements/Framework/TypedElementQueryBuilder.h"
 #include "Elements/Interfaces/TypedElementDataStorageInterface.h"
+#include "HAL/IConsoleManager.h"
 #include "TedsAlertsFactory.h"
 
 namespace UE::Editor::DataStorage::Alerts
 {
+	static TAutoConsoleVariable<bool> CVarEnableTedsAlerts(
+		TEXT("Teds.Feature.Alerts"),
+		true,
+		TEXT("Enable adding alerts to rows in TEDS"),
+		ECVF_Default
+	);
+
 	void AddAlert(ICoreProvider& DataStorage, RowHandle TargetRow, const FName& Name, FText Message, 
-		Columns::FAlertColumnType Type, uint8 Priority, Columns::FAlertActionCallback Action)
+		Columns::FAlertColumnType Type, uint8 Priority, Columns::FAlertActionCallback Action, const FName& HierarchyName)
 	{
 		using namespace UE::Editor::DataStorage::Columns;
 
+		if (!CVarEnableTedsAlerts.GetValueOnGameThread())
+		{
+			return;
+		}
+
 		// Add the new alert to the unsorted alert chain table and let TEDS sort is out at an appropriate time.
-		if (const UTedsAlertsFactory* AlertsFactory = DataStorage.FindFactory<UTedsAlertsFactory>())
+		if (UTedsAlertsFactory* AlertsFactory = DataStorage.FindFactory<UTedsAlertsFactory>())
 		{
 			// Add a placeholder alert to be filled in later.
 			if (!DataStorage.HasColumns<FAlertColumn>(TargetRow))
 			{
 				DataStorage.AddColumn<FAlertColumn>(TargetRow);
 			}
-			
+
+			FHierarchyHandle HierarchyHandle = DataStorage.FindHierarchyByName(HierarchyName);
+
+			if (DataStorage.IsValidHierarchyHandle(HierarchyHandle))
+			{
+				checkf(DataStorage.GetParentChangedColumnType(HierarchyHandle),
+						TEXT("Alerts with hierarchies require a parent changed column to properly detect changes. This can be specified on hierarchy registration."));
+				
+				AlertsFactory->RegisterAlertHierarchy(DataStorage, HierarchyName);
+			}
+
 			DataStorage.AddRow(AlertsFactory->GetUnsortedAlertChainTable(),
-				[&DataStorage, TargetRow, Name, Message = MoveTemp(Message), Type, Priority, Action = MoveTemp(Action)]
+				[&DataStorage, TargetRow, Name, Message = MoveTemp(Message), Type, Priority, Action = MoveTemp(Action), HierarchyHandle]
 				(RowHandle Row) mutable
 				{
 					DataStorage.AddColumn(Row, FAlertColumn
@@ -36,6 +59,9 @@ namespace UE::Editor::DataStorage::Alerts
 							.AlertType = Type,
 							.Priority = Priority
 						});
+					
+					DataStorage.AddColumn(Row, FAlertHierarchyColumn{.HierarchyHandle =  HierarchyHandle});
+
 					if (Action)
 					{
 						DataStorage.AddColumn(Row, FAlertActionColumn
@@ -48,15 +74,15 @@ namespace UE::Editor::DataStorage::Alerts
 	}
 
 	void AddAlert(IQueryContext& Context, RowHandle TargetRow, const FName& Name, FText Message,
-		Columns::FAlertColumnType Type, uint8 Priority, Columns::FAlertActionCallback Action)
+		Columns::FAlertColumnType Type, uint8 Priority, Columns::FAlertActionCallback Action, const FName& HierarchyName)
 	{
 		using namespace UE::Editor::DataStorage;
 		
-		Context.PushCommand([TargetRow, Name, Message = MoveTemp(Message), Type, Priority, Action = MoveTemp(Action)]() mutable
+		Context.PushCommand([TargetRow, Name, Message = MoveTemp(Message), Type, Priority, Action = MoveTemp(Action), HierarchyName]() mutable
 			{
 				if (ICoreProvider* DataStorage = GetMutableDataStorageFeature<ICoreProvider>(StorageFeatureName))
 				{
-					AddAlert(*DataStorage, TargetRow, Name, MoveTemp(Message), Type, Priority, MoveTemp(Action));
+					AddAlert(*DataStorage, TargetRow, Name, MoveTemp(Message), Type, Priority, MoveTemp(Action), HierarchyName);
 				}
 			});
 	}
diff --git a/Engine/Plugins/Experimental/EditorDataStorageFeatures/Source/TedsAlerts/Private/TedsAlertsFactory.cpp b/Engine/Plugins/Experimental/EditorDataStorageFeatures/Source/TedsAlerts/Private/TedsAlertsFactory.cpp
--- a/Engine/Plugins/Experimental/EditorDataStorageFeatures/Source/TedsAlerts/Private/TedsAlertsFactory.cpp
+++ b/Engine/Plugins/Experimental/EditorDataStorageFeatures/Source/TedsAlerts/Private/TedsAlertsFactory.cpp
@@ -2,7 +2,6 @@
 
 #include "TedsAlertsFactory.h"
 
-#include "Elements/Columns/TypedElementHiearchyColumns.h"
 #include "Elements/Columns/TypedElementMiscColumns.h"
 #include "Elements/Framework/TypedElementQueryBuilder.h"
 #include "Elements/Interfaces/TypedElementDataStorageInterface.h"
@@ -25,9 +24,7 @@
 {
 	RegisterSubQueries(DataStorage);
 	RegisterSortUnsortedAlertsQuery(DataStorage);
-	RegisterOnAddQueries(DataStorage);
 	RegisterOnRemoveQueries(DataStorage);
-	RegisterParentUpdatesQueries(DataStorage);
 	RegisterChildAlertUpdatesQueries(DataStorage);
 }
 
@@ -51,6 +48,23 @@
 	return UnsortedAlertsQuery;
 }
 
+void UTedsAlertsFactory::RegisterAlertHierarchy(UE::Editor::DataStorage::ICoreProvider& DataStorage, const FName& InHierarchyName)
+{
+	using namespace UE::Editor::DataStorage::Queries;
+
+	FHierarchyHandle HierarchyHandle = DataStorage.FindHierarchyByName(InHierarchyName);
+	
+	if (DataStorage.IsValidHierarchyHandle(HierarchyHandle) && !RegisteredHierarchies.Contains(InHierarchyName))
+	{
+		RegisterOnAddHierarchyQueries(DataStorage, HierarchyHandle);
+		RegisterOnRemoveHierarchyQueries(DataStorage, HierarchyHandle);
+		RegisterParentUpdatesHierarchyQueries(DataStorage, HierarchyHandle);
+		RegisterChildAlertHierarchyUpdatesQueries(DataStorage, InHierarchyName, HierarchyHandle);
+
+		RegisteredHierarchies.Add(InHierarchyName);
+	}
+}
+
 void UTedsAlertsFactory::RegisterSubQueries(UE::Editor::DataStorage::ICoreProvider& DataStorage)
 {
 	using namespace UE::Editor::DataStorage::Queries;
@@ -74,14 +88,14 @@
 			.ReadWrite<FTedsAlertActionColumn>()
 		.Compile());
 
+	AlertHierarchyQuery = DataStorage.RegisterQuery(
+		Select()
+			.ReadWrite<FTedsAlertHierarchyColumn>()
+		.Compile());
+
 	ChildAlertColumnReadWriteQuery = DataStorage.RegisterQuery(
 		Select()
 			.ReadWrite<FTedsChildAlertColumn>()
 		.Compile());
-
-	ParentReadOnlyQuery = DataStorage.RegisterQuery(
-		Select()
-			.ReadOnly<FTableRowParentColumn>()
-		.Compile());
 }
 
@@ -89,6 +103,9 @@
 {
 	using namespace UE::Editor::DataStorage::Queries;
 
+	constexpr int32 AlertActionQueryIndex = 1;
+	constexpr int32 AlertHierarchyQueryIndex = 2;
+
 	DataStorage.RegisterQuery(
 		Select("Sort unsorted alerts", FProcessor(EQueryTickPhase::PrePhysics, DataStorage.GetQueryTickGroupName(EQueryTickGroups::Update)),
 			[](IQueryContext& NewAlertContext, RowHandle NewAlertRow, FTedsAlertColumn& NewAlert)
@@ -110,12 +127,12 @@
 									if (TargetAlert.Message.IsEmpty())
 									{
 										// This is a placeholder so override it.
-										AssignAlert(Context, TargetRow, TargetAlert, NewAlertContext, NewAlertRow, NewAlert, 1);
+										AssignAlert(Context, TargetRow, TargetAlert, NewAlertContext, NewAlertRow, NewAlert, AlertActionQueryIndex, AlertHierarchyQueryIndex);
 									}
 									else
 									{
 										// The active alert needs to be replaced so swap the active alert with the new alert.
-										SwapAlerts(Context, TargetRow, TargetAlert, NewAlertContext, NewAlertRow, NewAlert, 1);
+										SwapAlerts(Context, TargetRow, TargetAlert, NewAlertContext, NewAlertRow, NewAlert, AlertActionQueryIndex, AlertHierarchyQueryIndex);
 									}
 									// (Re)calculate child alerts.
 									Context.ActivateQueries(AlertConditionName);
@@ -124,7 +141,7 @@
 								{
 									// Found a spot in the chain. To avoid having to track the previous entry, just swap the 
 									// alert at that spot with the new one and chain them up.
-									SwapAlerts(Context, TargetRow, TargetAlert, NewAlertContext, NewAlertRow, NewAlert, 1);
+									SwapAlerts(Context, TargetRow, TargetAlert, NewAlertContext, NewAlertRow, NewAlert, AlertActionQueryIndex, AlertHierarchyQueryIndex);
 								}
 								NextRow = InvalidRowHandle;
 							}
@@ -150,33 +167,7 @@
 		.DependsOn()
 			.SubQuery(SortedAlertsQuery)
 			.SubQuery(AlertActionQuery)
+			.SubQuery(AlertHierarchyQuery)
 		.Compile());
 }
 
-void UTedsAlertsFactory::RegisterOnAddQueries(UE::Editor::DataStorage::ICoreProvider& DataStorage)
-{
-	using namespace UE::Editor::DataStorage::Queries;
-
-	DataStorage.RegisterQuery(
-		Select("Register alert with parent on alert add", FObserver::OnAdd<FTedsAlertColumn>(),
-			[](IQueryContext& Context, RowHandle Row)
-			{
-				Context.ActivateQueries(AlertConditionName);
-			})
-		.Where()
-			.All<FTableRowParentColumn>() // Only need to do an update pass if there are parents.
-			.None<FTedsAlertChainTag>()
-		.Compile());
-
-	DataStorage.RegisterQuery(
-		Select("Register alert with parent on parent add", FObserver::OnAdd<FTableRowParentColumn>(),
-			[](IQueryContext& Context, RowHandle Row)
-			{
-				Context.ActivateQueries(AlertConditionName);
-			})
-		.Where()
-			.Any<FTedsAlertColumn, FTedsChildAlertColumn>()
-			.None<FTedsAlertChainTag>()
-		.Compile());
-}
-
@@ -208,6 +199,8 @@
 						}
 					}));
 
+				Context.RemoveColumns<FTedsAlertHierarchyColumn, FTedsAlertActionColumn>(Row);
+
 				// Update any alert parents.
 				Context.ActivateQueries(AlertConditionName);
 			})
@@ -217,82 +210,14 @@
 			.SubQuery(SortedAlertsQuery)
 			.SubQuery(UnsortedAlertsQuery)
 		.Compile());
-
-	DataStorage.RegisterQuery(
-		Select("Update alert upon parent removal", FObserver::OnRemove<FTableRowParentColumn>(),
-			[](IQueryContext& Context, RowHandle Row)
-			{
-				Context.ActivateQueries(AlertConditionName);
-			})
-		.Where()
-			.Any<FTedsAlertColumn, FTedsChildAlertColumn>()
-			.None<FTedsAlertChainTag>()
-		.Compile());
 }
 
-void UTedsAlertsFactory::RegisterParentUpdatesQueries(UE::Editor::DataStorage::ICoreProvider& DataStorage)
-{
-	using namespace UE::Editor::DataStorage::Queries;
-
-	DataStorage.RegisterQuery(
-		Select(
-			"Trigger alert update if alert's parent changed", 
-			FProcessor(EQueryTickPhase::PostPhysics, DataStorage.GetQueryTickGroupName(EQueryTickGroups::Default)),
-			[](IQueryContext& Context, FTedsAlertColumn& Alert, const FTableRowParentColumn& Parent)
-			{
-				if (Alert.CachedParent != Parent.Parent)
-				{
-					Alert.CachedParent = Parent.Parent;
-					Context.ActivateQueries(AlertConditionName);
-				}
-			})
-		.Where()
-			.Any<FTypedElementSyncBackToWorldTag, FTypedElementSyncFromWorldTag>()
-			.None<FTedsAlertChainTag>()
-		.Compile());
-
-	DataStorage.RegisterQuery(
-		Select(
-			"Trigger alert update if child alert's parent changed",
-			FProcessor(EQueryTickPhase::PostPhysics, DataStorage.GetQueryTickGroupName(EQueryTickGroups::Default)),
-			[](IQueryContext& Context, FTedsChildAlertColumn& ChildAlert, const FTableRowParentColumn& Parent)
-			{
-				if (ChildAlert.CachedParent != Parent.Parent)
-				{
-					ChildAlert.CachedParent = Parent.Parent;
-					Context.ActivateQueries(AlertConditionName);
-				}
-			})
-		.Where()
-			.Any<FTypedElementSyncBackToWorldTag, FTypedElementSyncFromWorldTag>()
-			.None<FTedsAlertChainTag>()
-		.Compile());
-}
-
 void UTedsAlertsFactory::RegisterChildAlertUpdatesQueries(UE::Editor::DataStorage::ICoreProvider& DataStorage)
 {
 	using namespace UE::Editor::DataStorage::Queries;
 
 	DataStorage.RegisterQuery(
 		Select(
-			"Add missing child alerts",
-			FPhaseAmble(FPhaseAmble::ELocation::Preamble, EQueryTickPhase::PostPhysics)
-				.MakeActivatable(AlertConditionName),
-			[](IQueryContext& Context, RowHandle Row, FTedsAlertColumn& Alert, const FTableRowParentColumn& Parent)
-			{
-				if (Context.IsRowAssigned(Parent.Parent))
-				{
-					AddChildAlertsToHierarchy(Context, Parent.Parent, 0);
-				}
-			})
-		.Where()
-			.None<FTedsAlertChainTag>()
-		.DependsOn()
-			.SubQuery(ParentReadOnlyQuery)
-		.Compile());
-
-	DataStorage.RegisterQuery(
-		Select(
 			"Clear child alerts",
 			FProcessor(EQueryTickPhase::PostPhysics, DataStorage.GetQueryTickGroupName(EQueryTickGroups::PreUpdate))
 				.MakeActivatable(AlertConditionName),
@@ -306,21 +231,6 @@
 
 	DataStorage.RegisterQuery(
 		Select(
-			"Increment child alerts",
-			FProcessor(EQueryTickPhase::PostPhysics, DataStorage.GetQueryTickGroupName(EQueryTickGroups::Update))
-				.MakeActivatable(AlertConditionName),
-			[](IQueryContext& Context, RowHandle Row, FTedsAlertColumn& Alert)
-			{
-				IncrementParents(Context, Alert.CachedParent, Alert.AlertType, 0);
-			})
-		.Where()
-			.None<FTedsAlertChainTag>()
-		.DependsOn()
-			.SubQuery(ChildAlertColumnReadWriteQuery)
-		.Compile());
-
-	DataStorage.RegisterQuery(
-		Select(
 			"Remove unused child alerts",
 			FProcessor(EQueryTickPhase::PostPhysics, DataStorage.GetQueryTickGroupName(EQueryTickGroups::PostUpdate))
 				.MakeActivatable(AlertConditionName),
@@ -333,20 +243,146 @@
 						return;
 					}
 				}
-				Context.RemoveColumns<FTedsChildAlertColumn>(Row);
+				Context.RemoveColumns<FTedsChildAlertColumn, FTedsAlertHierarchyColumn>(Row);
 				Context.AddColumns<FTypedElementSyncBackToWorldTag>(Row);
 			})
 		.Compile());
 }
 
+
+void UTedsAlertsFactory::RegisterOnAddHierarchyQueries(UE::Editor::DataStorage::ICoreProvider& DataStorage, UE::Editor::DataStorage::FHierarchyHandle InHierarchyHandle)
+{
+	using namespace UE::Editor::DataStorage::Queries;
+
+	DataStorage.RegisterQuery(
+		Select("Register alert with parent on alert hierarchy add", FObserver::OnAdd<FTedsAlertHierarchyColumn>(),
+			[InHierarchyHandle](IQueryContext& Context, RowHandle Row, const FTedsAlertHierarchyColumn& AlertHierarchyColumn)
+			{
+				if (AlertHierarchyColumn.HierarchyHandle == InHierarchyHandle)
+				{
+					Context.ActivateQueries(AlertConditionName);
+				}
+			})
+		.Where()
+			.All(DataStorage.GetChildTagType(InHierarchyHandle)) // Only need to do an update pass if there are parents.
+			.None<FTedsAlertChainTag>()
+		.Compile());
+
+	DataStorage.RegisterQuery(
+		Select("Register alert with parent on parent add (FHierarchyHandle)",
+			FObserver(FObserver::EEvent::Add, DataStorage.GetChildTagType(InHierarchyHandle)),
+			[InHierarchyHandle](IQueryContext& Context, RowHandle Row, const FTedsAlertHierarchyColumn& AlertHierarchyColumn)
+			{
+				if (AlertHierarchyColumn.HierarchyHandle == InHierarchyHandle)
+				{
+					Context.ActivateQueries(AlertConditionName);
+				}
+			})
+		.Where()
+			.Any<FTedsAlertColumn, FTedsChildAlertColumn>()
+			.None<FTedsAlertChainTag>()
+		.Compile());
+}
+
+void UTedsAlertsFactory::RegisterOnRemoveHierarchyQueries(UE::Editor::DataStorage::ICoreProvider& DataStorage,
+	UE::Editor::DataStorage::FHierarchyHandle InHierarchyHandle)
+{
+	using namespace UE::Editor::DataStorage::Queries;
+
+	DataStorage.RegisterQuery(
+		Select("Update alert upon parent removal", FObserver(FObserver::EEvent::Remove, DataStorage.GetChildTagType(InHierarchyHandle)),
+			[InHierarchyHandle](IQueryContext& Context, RowHandle Row, const FTedsAlertHierarchyColumn& AlertHierarchyColumn)
+			{
+				if (AlertHierarchyColumn.HierarchyHandle == InHierarchyHandle)
+				{
+					Context.ActivateQueries(AlertConditionName);
+				}
+			})
+		.Where()
+			.Any<FTedsAlertColumn, FTedsChildAlertColumn>()
+			.None<FTedsAlertChainTag>()
+		.Compile());
+}
+
+void UTedsAlertsFactory::RegisterParentUpdatesHierarchyQueries(UE::Editor::DataStorage::ICoreProvider& DataStorage, UE::Editor::DataStorage::FHierarchyHandle InHierarchyHandle)
+{
+	using namespace UE::Editor::DataStorage::Queries;
+	
+	if (const UScriptStruct* ParentChangedColumnType = DataStorage.GetParentChangedColumnType(InHierarchyHandle))
+	{
+		DataStorage.RegisterQuery(
+			Select(
+			"Trigger alert update if alert or child alert's parent changed", 
+			FProcessor(EQueryTickPhase::PostPhysics, DataStorage.GetQueryTickGroupName(EQueryTickGroups::Default)),
+			[InHierarchyHandle](IQueryContext& Context, RowHandle Row, const FTedsAlertHierarchyColumn& AlertHierarchyColumn)
+			{
+				if (AlertHierarchyColumn.HierarchyHandle == InHierarchyHandle)
+				{
+					Context.ActivateQueries(AlertConditionName);
+				}
+			})
+		.Where()
+			.All(ParentChangedColumnType)
+			.Any<FTedsAlertColumn, FTedsChildAlertColumn>()
+			.None<FTedsAlertChainTag>()
+		.Compile());
+	}
+}
+
+void UTedsAlertsFactory::RegisterChildAlertHierarchyUpdatesQueries(UE::Editor::DataStorage::ICoreProvider& DataStorage, const FName& InHierarchyName, UE::Editor::DataStorage::FHierarchyHandle InHierarchyHandle)
+{
+	using namespace UE::Editor::DataStorage::Queries;
+    
+	DataStorage.RegisterQuery(
+		Select(
+			"Add missing child alerts",
+			FPhaseAmble(FPhaseAmble::ELocation::Preamble, EQueryTickPhase::PostPhysics)
+				.MakeActivatable(AlertConditionName),
+			[InHierarchyHandle](IQueryContext& Context, RowHandle Row, const FTedsAlertHierarchyColumn& AlertHierarchyColumn)
+			{
+				if (AlertHierarchyColumn.HierarchyHandle == InHierarchyHandle)
+				{
+					RowHandle ParentRow = Context.GetParentRow(Row);
+					if (Context.IsRowAssigned(ParentRow))
+					{
+						AddChildAlertsToHierarchy(Context, ParentRow, AlertHierarchyColumn.HierarchyHandle);
+					}
+				}
+			})
+		.AccessesHierarchy(InHierarchyName)
+		.Where()
+			.None<FTedsAlertChainTag>()
+		.Compile());
+
+	DataStorage.RegisterQuery(
+		Select(
+			"Increment child alerts",
+			FProcessor(EQueryTickPhase::PostPhysics, DataStorage.GetQueryTickGroupName(EQueryTickGroups::Update))
+				.MakeActivatable(AlertConditionName),
+			[InHierarchyHandle](IQueryContext& Context, RowHandle Row, FTedsAlertColumn& Alert, const FTedsAlertHierarchyColumn& AlertHierarchyColumn)
+			{
+				if (AlertHierarchyColumn.HierarchyHandle == InHierarchyHandle)
+				{
+					IncrementParents(Context, Context.GetParentRow(Row), Alert.AlertType, 0);
+				}
+			})
+		.AccessesHierarchy(InHierarchyName)
+		.Where()
+			.None<FTedsAlertChainTag>()
+		.DependsOn()
+			.SubQuery(ChildAlertColumnReadWriteQuery)
+		.Compile());
+}
+
 void UTedsAlertsFactory::AssignAlert(
 	UE::Editor::DataStorage::ISubqueryContext& TargetContext, UE::Editor::DataStorage::RowHandle TargetRow, FTedsAlertColumn& Target,
 	UE::Editor::DataStorage::IQueryContext& SourceContext, UE::Editor::DataStorage::RowHandle SourceRow, FTedsAlertColumn& Source,
-	int32 AlertActionQueryIndex)
+	int32 AlertActionQueryIndex, int32 AlertHierarchyQueryIndex)
 {
 	Target = MoveTemp(Source);
 	Target.NextAlert = UE::Editor::DataStorage::InvalidRowHandle;
-	
+
+	// Assign Alert Action
 	FTedsAlertActionColumn* TargetAction = GetAction(SourceContext, TargetRow, AlertActionQueryIndex);
 	FTedsAlertActionColumn* SourceAction = GetAction(SourceContext, SourceRow, AlertActionQueryIndex);
 	if (SourceAction)
@@ -368,6 +404,28 @@
 		}
 	}
 
+	// Assign Alert Hierarchy
+	FTedsAlertHierarchyColumn* TargetHierarchy = GetHierarchy(SourceContext, TargetRow, AlertHierarchyQueryIndex);
+	FTedsAlertHierarchyColumn* SourceHierarchy = GetHierarchy(SourceContext, SourceRow, AlertHierarchyQueryIndex);
+	if (SourceHierarchy)
+	{
+		if (TargetHierarchy)
+		{
+			*TargetHierarchy = MoveTemp(*SourceHierarchy);
+		}
+		else
+		{
+			SourceContext.AddColumn(TargetRow, MoveTemp(*SourceHierarchy));
+		}
+	}
+	else
+	{
+		if (TargetHierarchy)
+		{
+			TargetContext.RemoveColumns<FTedsAlertHierarchyColumn>(TargetRow);
+		}
+	}
+	
 	// Notify UI.
 	TargetContext.AddColumns<FTypedElementSyncBackToWorldTag>(TargetRow);
 	// Fully absorbed into the placeholder so the addition to the chain is no longer needed.
@@ -377,7 +435,7 @@
 void UTedsAlertsFactory::SwapAlerts(
 	UE::Editor::DataStorage::ISubqueryContext& OriginalContext, UE::Editor::DataStorage::RowHandle OriginalRow, FTedsAlertColumn& Original,
 	UE::Editor::DataStorage::IQueryContext& NewContext, UE::Editor::DataStorage::RowHandle NewRow, FTedsAlertColumn& New,
-	int32 AlertActionQueryIndex)
+	int32 AlertActionQueryIndex, int32 AlertHierarchyQueryIndex)
 {
 	Swap(Original, New);
 	
@@ -407,6 +465,32 @@
 		}
 	}
 
+	// Swap the alert hierarchy columns if they exist.
+	{
+		FTedsAlertHierarchyColumn* OriginalHierarchy = GetHierarchy(NewContext, OriginalRow, AlertHierarchyQueryIndex);
+		FTedsAlertHierarchyColumn* NewHierarchy = GetHierarchy(NewContext, NewRow, AlertHierarchyQueryIndex);
+		if (NewHierarchy)
+		{
+			if (OriginalHierarchy)
+			{
+				Swap(*OriginalHierarchy, *NewHierarchy);
+			}
+			else
+			{
+				OriginalContext.AddColumn(OriginalRow, MoveTemp(*NewHierarchy));
+				NewContext.RemoveColumns<FTedsAlertHierarchyColumn>(NewRow);
+			}
+		}
+		else
+		{
+			if (OriginalHierarchy)
+			{
+				NewContext.AddColumn(NewRow, MoveTemp(*OriginalHierarchy));
+				OriginalContext.RemoveColumns<FTedsAlertHierarchyColumn>(OriginalRow);
+			}
+		}
+	}
+
 	Original.NextAlert = NewRow;
 	// Notify UI. Most of the time this is not needed as only the active alert should be used in the UI,
 	// but the TEDS Debugger might be showing the other alerts, so make sure they're updated to prevent
@@ -441,8 +525,22 @@
 	return Result;
 }
 
+FTedsAlertHierarchyColumn* UTedsAlertsFactory::GetHierarchy(UE::Editor::DataStorage::IQueryContext& Context, UE::Editor::DataStorage::RowHandle Row,
+	int32 ChildAlertQueryIndex)
+{
+	using namespace UE::Editor::DataStorage::Queries;
+	
+	FTedsAlertHierarchyColumn* Result = nullptr;
+	Context.RunSubquery(ChildAlertQueryIndex, Row, CreateSubqueryCallbackBinding(
+		[&Result](ISubqueryContext& Context, FTedsAlertHierarchyColumn& HierarchyColumn)
+		{
+			Result = &HierarchyColumn;
+		}));
+	return Result;
+}
+
 void UTedsAlertsFactory::AddChildAlertsToHierarchy(
-	UE::Editor::DataStorage::IQueryContext& Context, UE::Editor::DataStorage::RowHandle Parent, int32 ParentQueryIndex)
+	UE::Editor::DataStorage::IQueryContext& Context, UE::Editor::DataStorage::RowHandle Parent, UE::Editor::DataStorage::FHierarchyHandle HierarchyHandle)
 {
 	using namespace UE::Editor::DataStorage;
 
@@ -450,15 +548,17 @@
 	do
 	{
 		RowHandle NextParent = Parent;
-		bHasParent = MoveToNextParent(NextParent, Context, ParentQueryIndex);
+		bHasParent = MoveToNextParent(NextParent, Context);
 
 		// Check if a child alert column exists and add one if not.
 		if (!Context.HasColumn<FTedsChildAlertColumn>(Parent))
 		{
 			FTedsChildAlertColumn ChildAlert;
 			ResetChildAlertCounters(ChildAlert);
-			ChildAlert.CachedParent = NextParent;
 			Context.AddColumn(Parent, MoveTemp(ChildAlert));
+
+			// Child Alerts also need the hierarchy column to make sure they can walk up the parent chain
+			Context.AddColumn(Parent, FTedsAlertHierarchyColumn{.HierarchyHandle = HierarchyHandle});
 		}
 		Parent = NextParent;
 	} while (bHasParent);
@@ -478,24 +578,18 @@
 			{
 				ChildAlert.Counts[static_cast<size_t>(AlertType)]++;
 				Context.AddColumns<FTypedElementSyncBackToWorldTag>(Row);
-				ParentRow = ChildAlert.CachedParent;
+				ParentRow = Context.GetParentRow(ParentRow);
 			}));
 		checkf(Result.Count > 0, TEXT("Expected to be able to setup the child alert, but it was missing on the parent column."));
 	}
 }
 
 bool UTedsAlertsFactory::MoveToNextParent(
-	UE::Editor::DataStorage::RowHandle& Parent, UE::Editor::DataStorage::IQueryContext& Context, int32 SubQueryIndex)
+	UE::Editor::DataStorage::RowHandle& Parent, UE::Editor::DataStorage::IQueryContext& Context)
 {
 	using namespace UE::Editor::DataStorage::Queries;
-
-	FQueryResult Result = Context.RunSubquery(SubQueryIndex, Parent, CreateSubqueryCallbackBinding(
-		[&Parent](const FTableRowParentColumn& NextParent)
-		{
-			Parent = NextParent.Parent;
-		}));
-	Parent = Result.Count != 0 ? Parent : InvalidRowHandle;
-	return Result.Count != 0;
+	Parent = Context.GetParentRow(Parent);
+	return Context.IsRowAvailable(Parent);
 }
 
 void UTedsAlertsFactory::ResetChildAlertCounters(FTedsChildAlertColumn& ChildAlert)
diff --git a/Engine/Plugins/Experimental/EditorDataStorageFeatures/Source/TedsAlerts/Private/TedsAlertsFactory.h b/Engine/Plugins/Experimental/EditorDataStorageFeatures/Source/TedsAlerts/Private/TedsAlertsFactory.h
--- a/Engine/Plugins/Experimental/EditorDataStorageFeatures/Source/TedsAlerts/Private/TedsAlertsFactory.h
+++ b/Engine/Plugins/Experimental/EditorDataStorageFeatures/Source/TedsAlerts/Private/TedsAlertsFactory.h
@@ -32,6 +32,10 @@ public:
 	UE::Editor::DataStorage::QueryHandle GetSortedAlertsQuery() const;
 	UE::Editor::DataStorage::QueryHandle GetUnsortedAlertsQuery() const;
 
+	// Register a new hierarchy as being part of alerts. This adds observers to track changes for this specific hierarchy so rows with alerts
+	// beloging to this hierarchy have can update their parent
+	void RegisterAlertHierarchy(UE::Editor::DataStorage::ICoreProvider& DataStorage, const FName& InHierarchyName);
+
 private:
 	static const FName AlertConditionName;
 
@@ -41,24 +45,32 @@ private:
 	UE::Editor::DataStorage::QueryHandle SortedAlertsQuery = UE::Editor::DataStorage::InvalidQueryHandle;
 	UE::Editor::DataStorage::QueryHandle UnsortedAlertsQuery = UE::Editor::DataStorage::InvalidQueryHandle;
 	UE::Editor::DataStorage::QueryHandle AlertActionQuery = UE::Editor::DataStorage::InvalidQueryHandle;
-	UE::Editor::DataStorage::QueryHandle ChildAlertColumnReadWriteQuery = UE::Editor::DataStorage::InvalidTableHandle;
-	UE::Editor::DataStorage::QueryHandle ParentReadOnlyQuery = UE::Editor::DataStorage::InvalidTableHandle;
+	UE::Editor::DataStorage::QueryHandle AlertHierarchyQuery = UE::Editor::DataStorage::InvalidQueryHandle;
+	UE::Editor::DataStorage::QueryHandle ChildAlertColumnReadWriteQuery = UE::Editor::DataStorage::InvalidQueryHandle;
+
+	// Keep track of hierarchies we have already registered observers for to avoid duplicates
+	TSet<FName> RegisteredHierarchies;
 
 	void RegisterSubQueries(UE::Editor::DataStorage::ICoreProvider& DataStorage);
 	void RegisterSortUnsortedAlertsQuery(UE::Editor::DataStorage::ICoreProvider& DataStorage);
-	void RegisterOnAddQueries(UE::Editor::DataStorage::ICoreProvider& DataStorage);
 	void RegisterOnRemoveQueries(UE::Editor::DataStorage::ICoreProvider& DataStorage);
-	void RegisterParentUpdatesQueries(UE::Editor::DataStorage::ICoreProvider& DataStorage);
 	void RegisterChildAlertUpdatesQueries(UE::Editor::DataStorage::ICoreProvider& DataStorage);
 
+	// Hierarchy update queries, these currently have to be registered per hierarchy since a query can only access one hierarchy and it isn't possible
+	// to right a query that can dynamically pick which hierarchy to use.
+	void RegisterOnAddHierarchyQueries(UE::Editor::DataStorage::ICoreProvider& DataStorage, UE::Editor::DataStorage::FHierarchyHandle InHierarchyHandle);
+	void RegisterOnRemoveHierarchyQueries(UE::Editor::DataStorage::ICoreProvider& DataStorage, UE::Editor::DataStorage::FHierarchyHandle InHierarchyHandle);
+	void RegisterParentUpdatesHierarchyQueries(UE::Editor::DataStorage::ICoreProvider& DataStorage, UE::Editor::DataStorage::FHierarchyHandle InHierarchyHandle);
+	void RegisterChildAlertHierarchyUpdatesQueries(UE::Editor::DataStorage::ICoreProvider& DataStorage, const FName& InHierarchyName, UE::Editor::DataStorage::FHierarchyHandle InHierarchyHandle);
+
 	static void AssignAlert(
 		UE::Editor::DataStorage::ISubqueryContext& TargetContext, UE::Editor::DataStorage::RowHandle TargetRow, FTedsAlertColumn& Target,
 		UE::Editor::DataStorage::IQueryContext& SourceContext, UE::Editor::DataStorage::RowHandle SourceRow, FTedsAlertColumn& Source,
-		int32 AlertActionQueryIndex);
+		int32 AlertActionQueryIndex, int32 AlertHierarchyQueryIndex);
 	static void SwapAlerts(
 		UE::Editor::DataStorage::ISubqueryContext& OriginalContext, UE::Editor::DataStorage::RowHandle OriginalRow, FTedsAlertColumn& Original,
 		UE::Editor::DataStorage::IQueryContext& NewContext, UE::Editor::DataStorage::RowHandle NewRow, FTedsAlertColumn& New,
-		int32 AlertActionQueryIndex);
+		int32 AlertActionQueryIndex, int32 AlertHierarchyQueryIndex);
 	static void AppendAlert(
 		FTedsAlertColumn& LastAlert, 
 		UE::Editor::DataStorage::IQueryContext& AdditionalAlertContext, 
@@ -66,13 +78,15 @@ private:
 		FTedsAlertColumn& AdditionalAlert);
 	static FTedsAlertActionColumn* GetAction(UE::Editor::DataStorage::IQueryContext& Context,
 		UE::Editor::DataStorage::RowHandle Row, int32 ChildAlertQueryIndex);
+	static FTedsAlertHierarchyColumn* GetHierarchy(UE::Editor::DataStorage::IQueryContext& Context,
+		UE::Editor::DataStorage::RowHandle Row, int32 ChildAlertQueryIndex);
 
 	static void AddChildAlertsToHierarchy(
-		UE::Editor::DataStorage::IQueryContext& Context, UE::Editor::DataStorage::RowHandle Parent, int32 ParentQueryIndex);
+		UE::Editor::DataStorage::IQueryContext& Context, UE::Editor::DataStorage::RowHandle Parent, UE::Editor::DataStorage::FHierarchyHandle HierarchyHandle);
 	static void IncrementParents(
 		UE::Editor::DataStorage::IQueryContext& Context, UE::Editor::DataStorage::RowHandle Row, FTedsAlertColumnType AlertType,
 		int32 ChildAlertQueryIndex);
 	static bool MoveToNextParent(
-		UE::Editor::DataStorage::RowHandle& Parent, UE::Editor::DataStorage::IQueryContext& Context, int32 SubQueryIndex);
+		UE::Editor::DataStorage::RowHandle& Parent, UE::Editor::DataStorage::IQueryContext& Context);
 	static void ResetChildAlertCounters(FTedsChildAlertColumn& ChildAlert);
 };
diff --git a/Engine/Plugins/Experimental/EditorDataStorageFeatures/Source/TedsAlerts/Public/TedsAlertColumns.h b/Engine/Plugins/Experimental/EditorDataStorageFeatures/Source/TedsAlerts/Public/TedsAlertColumns.h
--- a/Engine/Plugins/Experimental/EditorDataStorageFeatures/Source/TedsAlerts/Public/TedsAlertColumns.h
+++ b/Engine/Plugins/Experimental/EditorDataStorageFeatures/Source/TedsAlerts/Public/TedsAlertColumns.h
@@ -31,9 +31,6 @@ struct FTedsAlertColumn final : public FEditorDataStorageColumn
 	UPROPERTY()
 	FText Message;
 
-	// Store a copy of the parent row so it's possible to detect if a row has been reparented.
-	UE::Editor::DataStorage::RowHandle CachedParent;
-
 	/**
 	 * If valid, points to the next alert in the chain.The next alert will take the place of the current
 	 * alert if this alert is cleared.
@@ -63,13 +60,21 @@ USTRUCT(meta = (DisplayName = "Child alert"))
 struct FTedsChildAlertColumn final : public FEditorDataStorageColumn
 {
 	GENERATED_BODY()
-
-	// Store a copy of the parent row so it's possible to detect if a row has been reparented.
-	UE::Editor::DataStorage::RowHandle CachedParent;
-
+	
 	uint16 Counts[static_cast<size_t>(FTedsAlertColumnType::MAX)];
 };
 
+/**
+ * Column that holdes info about the hierarchy an alert is participating in
+ */
+USTRUCT(meta = (DisplayName = "Alert Hierarchy Column"))
+struct FTedsAlertHierarchyColumn final : public FEditorDataStorageColumn
+{
+	GENERATED_BODY()
+	
+	UE::Editor::DataStorage::FHierarchyHandle HierarchyHandle;
+};
+
 /**
  * Column that can be added to an alert column to have it trigger an action when the alert is clicked.
  */
@@ -107,6 +112,7 @@ namespace UE::Editor::DataStorage::Columns
 	
 	using FAlertColumn = FTedsAlertColumn;
 	using FChildAlertColumn = FTedsChildAlertColumn;
+	using FAlertHierarchyColumn = FTedsAlertHierarchyColumn;
 	using FAlertActionColumn = FTedsAlertActionColumn;
 	using FAlertChainTag = FTedsAlertChainTag;
 	using FUnsortedAlertChainTag = FTedsUnsortedAlertChainTag;
diff --git a/Engine/Plugins/Experimental/EditorDataStorageFeatures/Source/TedsAlerts/Public/TedsAlerts.h b/Engine/Plugins/Experimental/EditorDataStorageFeatures/Source/TedsAlerts/Public/TedsAlerts.h
--- a/Engine/Plugins/Experimental/EditorDataStorageFeatures/Source/TedsAlerts/Public/TedsAlerts.h
+++ b/Engine/Plugins/Experimental/EditorDataStorageFeatures/Source/TedsAlerts/Public/TedsAlerts.h
@@ -28,9 +28,11 @@ namespace UE::Editor::DataStorage::Alerts
 	 * @Priority If there are multiple alerts for the same type, the priority is used to order the alerts so the most important
 	 *		one is shown when the active alert is removed. Higher values indicate higher priority.
 	 * @Action An optional callback that gets triggered when the user presses the alert icon.
+	 * @HierarchyName Optional hierarchy name to use when resolving parent alerts. The alert queries will walk up the parent chain in the provided
+	 *		hierarchy to add the child alert column
 	 */
 	TEDSALERTS_API void AddAlert(ICoreProvider& DataStorage, RowHandle TargetRow, const FName& Name, FText Message, 
-		Columns::FAlertColumnType Type, uint8 Priority = 127, Columns::FAlertActionCallback Action = {});
+		Columns::FAlertColumnType Type, uint8 Priority = 127, Columns::FAlertActionCallback Action = {}, const FName& HierarchyName = NAME_None);
 	/**
 	 * Add an alert to the provided row. Various tools like the Outliner monitor for alerts and can show them for instance as
 	 * a colored triangle with a exclamation count. If a row also has a parent row handle, the alerts will propagate their count
@@ -44,9 +46,11 @@ namespace UE::Editor::DataStorage::Alerts
 	 * @Priority If there are multiple alerts for the same type, the priority is used to order the alerts so the most important
 	 *		one is shown when the active alert is removed. Higher values indicate higher priority.
 	 * @Action An optional callback that gets triggered when the user presses the alert icon.
+	 * @HierarchyName Optional hierarchy name to use when resolving parent alerts. The alert queries will walk up the parent chain in the provided
+	 *		hierarchy to add the child alert column
 	 */
 	TEDSALERTS_API void AddAlert(IQueryContext& Context, RowHandle TargetRow, const FName& Name, FText Message,
-		Columns::FAlertColumnType Type, uint8 Priority = 127, Columns::FAlertActionCallback Action = {});
+		Columns::FAlertColumnType Type, uint8 Priority = 127, Columns::FAlertActionCallback Action = {}, const FName& HierarchyName = NAME_None);
 	
 	/**
 	 * Locates the first alerts with the provided name in the alert chain found at the provided row and updates its text.
diff --git a/Engine/Source/Runtime/TypedElementFramework/Public/DataStorage/Handles.h b/Engine/Source/Runtime/TypedElementFramework/Public/DataStorage/Handles.h
--- a/Engine/Source/Runtime/TypedElementFramework/Public/DataStorage/Handles.h
+++ b/Engine/Source/Runtime/TypedElementFramework/Public/DataStorage/Handles.h
@@ -21,6 +21,10 @@ namespace UE::Editor::DataStorage
 	struct FHierarchyHandle
 	{
 		uint64 Reserved = 0;
+		bool operator==(FHierarchyHandle Other) const
+		{
+			return Reserved == Other.Reserved;
+		}
 	};
 } // namespace UE::Editor::DataStorage
 
