[TEDS] Use the new query topology hash to speed up the query result node.

The Query Result Node in for TEDS' Query Stack required to run a query each frame and see if there were any changes between the current and previous run to see and optional only report if changes were found. At scale this was a very expensive operation, which crippled the editor in some situations. One example would be a simple query in the TEDS Debugger to list all asset names. This produced millions of results which brought the editor's framerate down to 8fps. With this change there's no longer a dependency on the number of rows, just whether or not there was a topological change in one of the tables or table variants associated with the query. As a result the framerate is no longer effected until there's an actual change.

This change was applied at all callsites such that the behavior doesn't change, but since the original recommendation was to only use `ESyncFlags::IncrementWhenDifferent` when needed, most locations use then the cheaper version without it that always updates everything. It's however now recommended to see if `ESyncActions::ForceRefreshOnUpdate` can be replaced with the much cheaper `ESyncActions::RefreshOnUpdate`.

Another small change was added as well that passed on the allotted time for time sliced work on to derived implementations.

#Synced-CL 48540487

diff --git a/Engine/Plugins/Experimental/EditorDataStorageFeatures/Source/TedsActorCompatibility/Private/ActorComponentDebugHierarchyWidget/ActorComponentTree.cpp b/Engine/Plugins/Experimental/EditorDataStorageFeatures/Source/TedsActorCompatibility/Private/ActorComponentDebugHierarchyWidget/ActorComponentTree.cpp
--- a/Engine/Plugins/Experimental/EditorDataStorageFeatures/Source/TedsActorCompatibility/Private/ActorComponentDebugHierarchyWidget/ActorComponentTree.cpp
+++ b/Engine/Plugins/Experimental/EditorDataStorageFeatures/Source/TedsActorCompatibility/Private/ActorComponentDebugHierarchyWidget/ActorComponentTree.cpp
@@ -73,7 +73,7 @@ namespace UE::Editor::DataStorage::ActorCompatibility
 					TSharedPtr<QueryStack::IRowNode> TopLevelRowNode = MakeShared<QueryStack::FRowQueryCallbackResultsNode>(
 						*DataStorageInterface,
 						QueryNode,
-						QueryStack::FRowQueryResultsNode::ESyncFlags::IncrementWhenDifferent | QueryStack::FRowQueryResultsNode::ESyncFlags::RefreshOnUpdate,
+						QueryStack::FRowQueryResultsNode::ESyncActions::RefreshOnUpdate,
 						FilterRowsBySelectionSet
 						);
 
diff --git a/Engine/Plugins/Experimental/EditorDataStorageFeatures/Source/TedsActorCompatibility/Private/ActorComponentDebugHierarchyWidget/QueryStackNodes/RowQueryCallbackResultsNode.cpp b/Engine/Plugins/Experimental/EditorDataStorageFeatures/Source/TedsActorCompatibility/Private/ActorComponentDebugHierarchyWidget/QueryStackNodes/RowQueryCallbackResultsNode.cpp
--- a/Engine/Plugins/Experimental/EditorDataStorageFeatures/Source/TedsActorCompatibility/Private/ActorComponentDebugHierarchyWidget/QueryStackNodes/RowQueryCallbackResultsNode.cpp
+++ b/Engine/Plugins/Experimental/EditorDataStorageFeatures/Source/TedsActorCompatibility/Private/ActorComponentDebugHierarchyWidget/QueryStackNodes/RowQueryCallbackResultsNode.cpp
@@ -10,10 +10,19 @@ namespace UE::Editor::DataStorage::QueryStack
 	FRowQueryCallbackResultsNode::FRowQueryCallbackResultsNode(
 		ICoreProvider& Storage,
 		TSharedPtr<IQueryNode> InQueryNode,
-		ESyncFlags SyncFlags,
 		CallbackFn InCallbackFn)
-		: FRowQueryResultsNode(Storage, InQueryNode, SyncFlags)
-		, Callback(InCallbackFn)
+		: FRowQueryResultsNode(Storage, InQueryNode)
+		, Callback(MoveTemp(InCallbackFn))
+	{
+	}
+
+	FRowQueryCallbackResultsNode::FRowQueryCallbackResultsNode(
+		ICoreProvider& Storage,
+		TSharedPtr<IQueryNode> InQueryNode,
+		ESyncActions SyncAction,
+		CallbackFn InCallbackFn)
+		: FRowQueryResultsNode(Storage, InQueryNode, SyncAction)
+		, Callback(MoveTemp(InCallbackFn))
 	{
 	}
 
diff --git a/Engine/Plugins/Experimental/EditorDataStorageFeatures/Source/TedsActorCompatibility/Private/ActorComponentDebugHierarchyWidget/QueryStackNodes/RowQueryCallbackResultsNode.h b/Engine/Plugins/Experimental/EditorDataStorageFeatures/Source/TedsActorCompatibility/Private/ActorComponentDebugHierarchyWidget/QueryStackNodes/RowQueryCallbackResultsNode.h
--- a/Engine/Plugins/Experimental/EditorDataStorageFeatures/Source/TedsActorCompatibility/Private/ActorComponentDebugHierarchyWidget/QueryStackNodes/RowQueryCallbackResultsNode.h
+++ b/Engine/Plugins/Experimental/EditorDataStorageFeatures/Source/TedsActorCompatibility/Private/ActorComponentDebugHierarchyWidget/QueryStackNodes/RowQueryCallbackResultsNode.h
@@ -21,10 +21,15 @@ namespace UE::Editor::DataStorage::QueryStack
 	public:
 		using EmitRowFn = TFunctionRef<void(TArrayView<const RowHandle>)>;
 		using CallbackFn = TFunction<void(IDirectQueryContext&, TArrayView<const RowHandle>, EmitRowFn EmitRows)>;
+		
 		FRowQueryCallbackResultsNode(
 			ICoreProvider& Storage,
 			TSharedPtr<IQueryNode> InQueryNode,
-			ESyncFlags SyncFlags,
+			CallbackFn InCallback);
+		FRowQueryCallbackResultsNode(
+			ICoreProvider& Storage,
+			TSharedPtr<IQueryNode> InQueryNode,
+			ESyncActions SyncAction,
 			CallbackFn InCallback);
 	
 	protected:
diff --git a/Engine/Plugins/Experimental/EditorDataStorageFeatures/Source/TedsContentBrowser/Private/Widgets/TedsContentBrowserAssetViewWidget.cpp b/Engine/Plugins/Experimental/EditorDataStorageFeatures/Source/TedsContentBrowser/Private/Widgets/TedsContentBrowserAssetViewWidget.cpp
--- a/Engine/Plugins/Experimental/EditorDataStorageFeatures/Source/TedsContentBrowser/Private/Widgets/TedsContentBrowserAssetViewWidget.cpp
+++ b/Engine/Plugins/Experimental/EditorDataStorageFeatures/Source/TedsContentBrowser/Private/Widgets/TedsContentBrowserAssetViewWidget.cpp
@@ -35,7 +35,7 @@ public:
 		TSharedRef<QueryStack::FQueryNode> QueryNode = MakeShared<QueryStack::FQueryNode>(*Storage, Arguments._InitParams.QueryDescription);
 
 		// For now we just refresh the query on update every frame since this is a prototype, in the future we can use a monitor node if perf is a concern
-		RowView = MakeShared<QueryStack::FRowQueryResultsNode>(*Storage, QueryNode, QueryStack::FRowQueryResultsNode::ESyncFlags::RefreshOnUpdate);
+		RowView = MakeShared<QueryStack::FRowQueryResultsNode>(*Storage, QueryNode, QueryStack::FRowQueryResultsNode::ESyncActions::ForceRefreshOnUpdate);
 		
 		TSharedPtr<SWidget> ChildWidget;
 
diff --git a/Engine/Plugins/Experimental/EditorDataStorageFeatures/Source/TedsDebugger/Private/Widgets/QueryEditor/TedsQueryEditorResultsView.cpp b/Engine/Plugins/Experimental/EditorDataStorageFeatures/Source/TedsDebugger/Private/Widgets/QueryEditor/TedsQueryEditorResultsView.cpp
--- a/Engine/Plugins/Experimental/EditorDataStorageFeatures/Source/TedsDebugger/Private/Widgets/QueryEditor/TedsQueryEditorResultsView.cpp
+++ b/Engine/Plugins/Experimental/EditorDataStorageFeatures/Source/TedsDebugger/Private/Widgets/QueryEditor/TedsQueryEditorResultsView.cpp
@@ -32,6 +32,13 @@ namespace UE::Editor::DataStorage::Debug::QueryEditor
 	{
 		Model->GetModelChangedDelegate().Remove(ModelChangedDelegateHandle);
 		Model->OnHierarchyChanged().Remove(HierarchyChangedDelegateHandle);
+
+		if (CountQueryHandle != InvalidQueryHandle)
+		{
+			ICoreProvider& TedsInterface = Model->GetTedsInterface();
+			TedsInterface.UnregisterQuery(CountQueryHandle);
+			CountQueryHandle = InvalidQueryHandle;
+		}
 	}
 
 	void SResultsView::Construct(const FArguments& InArgs, FTedsQueryEditorModel& InModel)
@@ -50,10 +57,10 @@ namespace UE::Editor::DataStorage::Debug::QueryEditor
 	
 		QueryNode = MakeShared<FQueryNode>(Model->GetTedsInterface());
 
-		// Use RefreshOnUpdate to re-run the query every frame since we don't have a good change detection mechanism for hierarchies right now
+		// Use ForceRefreshOnUpdate to re-run the query every frame since we don't have a good change detection mechanism for hierarchies right now
 		// and just using observers can cause the visual hierarchy to go out of sync
 		RowQueryStack = MakeShared<FRowQueryResultsNode>(Model->GetTedsInterface(), QueryNode,
-			FRowQueryResultsNode::ESyncFlags::RefreshOnUpdate);
+			FRowQueryResultsNode::ESyncActions::ForceRefreshOnUpdate);
 		//RowQueryStack = MakeShared<FRowMonitorNode>(Model->GetTedsInterface(), RowQueryStack, QueryNode);
 		
 		ChildSlot
diff --git a/Engine/Plugins/Experimental/EditorDataStorageFeatures/Source/TedsEverythingPicker/Private/Context/TedsPickerContextUtil.cpp b/Engine/Plugins/Experimental/EditorDataStorageFeatures/Source/TedsEverythingPicker/Private/Context/TedsPickerContextUtil.cpp
--- a/Engine/Plugins/Experimental/EditorDataStorageFeatures/Source/TedsEverythingPicker/Private/Context/TedsPickerContextUtil.cpp
+++ b/Engine/Plugins/Experimental/EditorDataStorageFeatures/Source/TedsEverythingPicker/Private/Context/TedsPickerContextUtil.cpp
@@ -53,7 +53,7 @@ namespace UE::Editor::DataStorage::Picker
 			{
 				FinalNode =
 					MakeShared<QueryStack::FRowQueryResultsNode>(*Storage,
-						ReferenceQueryNode, QueryStack::FRowQueryResultsNode::ESyncFlags::RefreshOnUpdate);
+						ReferenceQueryNode, QueryStack::FRowQueryResultsNode::ESyncActions::ForceRefreshOnUpdate);
 			}
 			else
 			{
@@ -139,7 +139,7 @@ namespace UE::Editor::DataStorage::Picker
 			{
 				FinalNode =
 					MakeShared<QueryStack::FRowQueryResultsNode>(*Storage,
-						TypeQueryNode, QueryStack::FRowQueryResultsNode::ESyncFlags::RefreshOnUpdate);
+						TypeQueryNode, QueryStack::FRowQueryResultsNode::ESyncActions::ForceRefreshOnUpdate);
 			}
 			else
 			{
diff --git a/Engine/Plugins/Experimental/EditorDataStorageFeatures/Source/TedsOperations/Private/TedsOperationSystem.cpp b/Engine/Plugins/Experimental/EditorDataStorageFeatures/Source/TedsOperations/Private/TedsOperationSystem.cpp
--- a/Engine/Plugins/Experimental/EditorDataStorageFeatures/Source/TedsOperations/Private/TedsOperationSystem.cpp
+++ b/Engine/Plugins/Experimental/EditorDataStorageFeatures/Source/TedsOperations/Private/TedsOperationSystem.cpp
@@ -76,7 +76,7 @@ void UOperationSystem::RegisterQueries(UE::Editor::DataStorage::ICoreProvider& S
 	GetQueryDescription(QueryDescription);
 
 	QueryNode   = MakeShared<FQueryNode>(Storage, MoveTemp(QueryDescription));
-	ResultsNode = MakeShared<FRowQueryResultsNode>(Storage, QueryNode, FRowQueryResultsNode::ESyncFlags::RefreshOnUpdate);
+	ResultsNode = MakeShared<FRowQueryResultsNode>(Storage, QueryNode, FRowQueryResultsNode::ESyncActions::ForceRefreshOnUpdate);
 	Executor    = MakeShared<FExplicitUpdateExecutor>("Operations", ResultsNode);
 }
 
diff --git a/Engine/Plugins/Experimental/EditorDataStorageFeatures/Source/TedsOutliner/Private/TedsOutlinerImpl.cpp b/Engine/Plugins/Experimental/EditorDataStorageFeatures/Source/TedsOutliner/Private/TedsOutlinerImpl.cpp
--- a/Engine/Plugins/Experimental/EditorDataStorageFeatures/Source/TedsOutliner/Private/TedsOutlinerImpl.cpp
+++ b/Engine/Plugins/Experimental/EditorDataStorageFeatures/Source/TedsOutliner/Private/TedsOutlinerImpl.cpp
@@ -813,14 +813,14 @@ void FTedsOutlinerImpl::RecompileQueries()
 	TSharedPtr<IQueryNode> CombinedQueryNode = MakeShared<FQueryMergeNode>(*Storage, QueryNodes);
 
 	// If we have a monitor node, it handles updating the collector node automatically. Otherwise, we need it to update on tick
-	FRowQueryResultsNode::ESyncFlags SyncFlags =
-		CreationParams.bUseDefaultObservers ? FRowQueryResultsNode::ESyncFlags::None : FRowQueryResultsNode::ESyncFlags::RefreshOnUpdate;
+	FRowQueryResultsNode::ESyncActions SyncActions =
+		CreationParams.bUseDefaultObservers ? FRowQueryResultsNode::ESyncActions::None : FRowQueryResultsNode::ESyncActions::ForceRefreshOnUpdate;
 
 	// Node to collect all rows matching the initial query (stored so OR filters can filter on the original query)
-	InitialRowCollectorNode = MakeShared<FRowQueryResultsNode>(*Storage, InitialQueryNode, SyncFlags);
+	InitialRowCollectorNode = MakeShared<FRowQueryResultsNode>(*Storage, InitialQueryNode, SyncActions);
 
 	// Node to collect all rows matching the composite query
-	CombinedRowCollectorNode = MakeShared<FRowQueryResultsNode>(*Storage, CombinedQueryNode, SyncFlags);
+	CombinedRowCollectorNode = MakeShared<FRowQueryResultsNode>(*Storage, CombinedQueryNode, SyncActions);
 
 	FilterNodes.Empty();
 	TSharedPtr<IRowNode> PrimaryFilterNode = CombinedRowCollectorNode;
diff --git a/Engine/Plugins/Experimental/EditorDataStorageFeatures/Source/TedsQueryStack/Private/TedsRowQueryResultsNode.cpp b/Engine/Plugins/Experimental/EditorDataStorageFeatures/Source/TedsQueryStack/Private/TedsRowQueryResultsNode.cpp
--- a/Engine/Plugins/Experimental/EditorDataStorageFeatures/Source/TedsQueryStack/Private/TedsRowQueryResultsNode.cpp
+++ b/Engine/Plugins/Experimental/EditorDataStorageFeatures/Source/TedsQueryStack/Private/TedsRowQueryResultsNode.cpp
@@ -4,16 +4,49 @@
 
 #include "Elements/Framework/TypedElementQueryBuilder.h"
 #include "Elements/Interfaces/TypedElementDataStorageInterface.h"
+#include "Hash/xxhash.h"
 
 namespace UE::Editor::DataStorage::QueryStack
 {
-	FRowQueryResultsNode::FRowQueryResultsNode(ICoreProvider& Storage, TSharedPtr<IQueryNode> InQueryNode, ESyncFlags SyncFlags)
+	FRowQueryResultsNode::FRowQueryResultsNode(ICoreProvider& Storage, TSharedPtr<IQueryNode> InQueryNode)
 		: QueryNode(MoveTemp(InQueryNode))
 		, Storage(Storage)
-		, SyncFlags(SyncFlags)
 	{
-		Refresh();
-		QueryRevision = QueryNode->GetRevision();
+	}
+
+PRAGMA_DISABLE_DEPRECATION_WARNINGS
+	// Remaps the only valid combination of ESyncFlags to the new sync actions for backwards compatibility.
+	static FRowQueryResultsNode::ESyncActions RemapSyncOptions(FRowQueryResultsNode::EDeprecatedSyncFlags SyncFlags)
+	{
+		bool bRequestIncrementWhenDifferent = EnumHasAllFlags(SyncFlags, FRowQueryResultsNode::EDeprecatedSyncFlags::IncrementWhenDifferent);
+		EnumRemoveFlags(SyncFlags, FRowQueryResultsNode::ESyncFlags::IncrementWhenDifferent);
+
+		switch (SyncFlags)
+		{
+		default:
+			[[fallthrough]];
+		case FRowQueryResultsNode::ESyncFlags::None:
+			return FRowQueryResultsNode::ESyncActions::None;
+		case FRowQueryResultsNode::ESyncFlags::RefreshOnUpdate:
+			return bRequestIncrementWhenDifferent
+				? FRowQueryResultsNode::ESyncActions::RefreshOnUpdate
+				: FRowQueryResultsNode::ESyncActions::ForceRefreshOnUpdate;
+		case FRowQueryResultsNode::ESyncFlags::RefreshOnQueryChange:
+			return FRowQueryResultsNode::ESyncActions::RefreshOnQueryChange;
+		}
+	}
+
+	FRowQueryResultsNode::FRowQueryResultsNode(ICoreProvider& Storage, TSharedPtr<IQueryNode> InQueryNode, EDeprecatedSyncFlags InSyncFlags)
+		: FRowQueryResultsNode(Storage, MoveTemp(InQueryNode), RemapSyncOptions(InSyncFlags))
+	{
+	}
+PRAGMA_ENABLE_DEPRECATION_WARNINGS
+
+	FRowQueryResultsNode::FRowQueryResultsNode(ICoreProvider& Storage, TSharedPtr<IQueryNode> InQueryNode, ESyncActions InSyncAction)
+		: QueryNode(MoveTemp(InQueryNode))
+		, Storage(Storage)
+		, SyncAction(InSyncAction)
+	{
 	}
 	
 	INode::RevisionId FRowQueryResultsNode::GetRevision() const
@@ -23,42 +56,49 @@ namespace UE::Editor::DataStorage::QueryStack
 
 	void FRowQueryResultsNode::Refresh()
 	{
-		if (EnumHasAllFlags(SyncFlags, ESyncFlags::IncrementWhenDifferent))
-		{
-			FRowHandleArray NewRows;
-			RefreshInternal(NewRows);
-			if (NewRows.Num() != Rows.Num())
-			{
-				Rows = MoveTemp(NewRows);
-				Revision++;
-			}
-			else
-			{
-				Rows.Sort();
-				NewRows.Sort();
-				FRowHandleArrayView NewRowView = NewRows.GetRows();
-				if (FMemory::Memcmp(NewRowView.GetData(), Rows.GetRows().GetData(), NewRowView.NumBytes()) != 0)
-				{
-					Rows = MoveTemp(NewRows);
-					Revision++;
-				}
-			}
-		}
-		else
-		{
-			Rows.Empty();
-			RefreshInternal(Rows);
-			Revision++;
-		}
+		Refresh(FTimespan::MaxValue());
+	}
+
+	void FRowQueryResultsNode::Refresh(FTimespan AllottedTime)
+	{
+		Rows.Empty();
+		RefreshInternal(AllottedTime, Rows);
+		Revision++;
 	}
 
 	void FRowQueryResultsNode::Update(FTimespan AllottedTime)
 	{
-		if (EnumHasAllFlags(SyncFlags, ESyncFlags::RefreshOnUpdate) ||
-			(EnumHasAllFlags(SyncFlags, ESyncFlags::RefreshOnQueryChange) && QueryRevision != QueryNode->GetRevision()))
+		bool bUpdateQuery = QueryRevision != QueryNode->GetRevision();
+		if (bUpdateQuery)
 		{
-			Refresh();
 			QueryRevision = QueryNode->GetRevision();
+			TopologyHash = 0;
+			QueryUpdated();
+		}
+
+		switch (SyncAction)
+		{
+		case ESyncActions::None:
+			break;
+		case ESyncActions::RefreshOnUpdate:
+		{
+			uint64 NewTopologyHash = Storage.CalculateQueryTablesTopologyHash(QueryNode->GetQuery());
+			if (TopologyHash != NewTopologyHash)
+			{
+				Refresh(AllottedTime);
+				TopologyHash = NewTopologyHash;
+			}
+			break;
+		}
+		case ESyncActions::RefreshOnQueryChange:
+			if (bUpdateQuery)
+			{
+				Refresh(AllottedTime);
+			}
+			break;
+		case ESyncActions::ForceRefreshOnUpdate:
+			Refresh(AllottedTime);
+			break;
 		}
 	}
 
@@ -77,6 +117,11 @@ namespace UE::Editor::DataStorage::QueryStack
 		return Rows;
 	}
 
+	void FRowQueryResultsNode::RefreshInternal(FTimespan AllottedTime, FRowHandleArray& TargetRows)
+	{
+		RefreshInternal(TargetRows);
+	}
+	
 	void FRowQueryResultsNode::RefreshInternal(FRowHandleArray& TargetRows)
 	{
 		using namespace UE::Editor::DataStorage::Queries;
diff --git a/Engine/Plugins/Experimental/EditorDataStorageFeatures/Source/TedsQueryStack/Public/TedsRowQueryResultsNode.h b/Engine/Plugins/Experimental/EditorDataStorageFeatures/Source/TedsQueryStack/Public/TedsRowQueryResultsNode.h
--- a/Engine/Plugins/Experimental/EditorDataStorageFeatures/Source/TedsQueryStack/Public/TedsRowQueryResultsNode.h
+++ b/Engine/Plugins/Experimental/EditorDataStorageFeatures/Source/TedsQueryStack/Public/TedsRowQueryResultsNode.h
@@ -23,44 +23,67 @@ namespace UE::Editor::DataStorage::QueryStack
 	class FRowQueryResultsNode : public IRowNode
 	{
 	public:
-		enum class ESyncFlags
+		enum class EDeprecatedSyncFlags
 		{
 			None = 0,
-			/** 
-			 * Update the row list whenever "Update" is called. This can take several milliseconds for large lists. It's
-			 * recommended to use a monitoring node to detect changes for large lists.
-			 */
 			RefreshOnUpdate = 1 << 0,
-			/** Update the row list whenever the parent query changes. */
 			RefreshOnQueryChange = 1 << 1,
-			/** 
-			 * Compares the previous row list with the current and only updates the revision if there are differences.
-			 * This requires additional sorting and comparing, making the node substantially more expensive.
-			 */
 			IncrementWhenDifferent = 1 << 2
 		};
+		using ESyncFlags UE_DEPRECATED(5.8, "'ESyncFlags' has been deprecated, please use `ESyncActions`.") = EDeprecatedSyncFlags;
 
-		TEDSQUERYSTACK_API FRowQueryResultsNode(ICoreProvider& Storage, TSharedPtr<IQueryNode> InQueryNode, ESyncFlags SyncFlags);
+		enum class ESyncActions
+		{
+			/**
+			 * The node is never updated and requires an explicit call to "Refresh".
+			 */
+			None,
+			/**
+			 * Each time "Update" is called the state of the query is checked and compared with the cached state. If any
+			 * row additions/removals are detected to any of the tables associated with the query it is rerun.
+			 */
+			RefreshOnUpdate,
+			/** Only update the row list whenever the parent query changes. */
+			RefreshOnQueryChange,
+			/**
+			 * Each time "Update" is called the query is run again and replaces the current set of rows. Unlike "RefreshOnUpdate", 
+			 * there is no caching done to mitigate the cost of running a query. This options is only recommended for debugging
+			 * purposes as it can significantly impact performance.
+			 */
+			ForceRefreshOnUpdate,
+		};
+
+		TEDSQUERYSTACK_API FRowQueryResultsNode(ICoreProvider& Storage, TSharedPtr<IQueryNode> InQueryNode);
+		TEDSQUERYSTACK_API FRowQueryResultsNode(ICoreProvider& Storage, TSharedPtr<IQueryNode> InQueryNode, ESyncActions InSyncAction);
 		virtual ~FRowQueryResultsNode() override = default;
 
 		TEDSQUERYSTACK_API void Refresh();
+		TEDSQUERYSTACK_API void Refresh(FTimespan AllottedTime);
 
 		TEDSQUERYSTACK_API virtual RevisionId GetRevision() const override;
 		TEDSQUERYSTACK_API virtual void VisitParents(ParentListCallback Callback) override;
 		TEDSQUERYSTACK_API virtual FRowHandleArrayView GetRows() const override;
 		TEDSQUERYSTACK_API virtual FRowHandleArray& GetMutableRows() override;
 
+		UE_DEPRECATED(5.8, "'ESyncFlags' has been deprecated, please use `ESyncActions`.")
+		TEDSQUERYSTACK_API FRowQueryResultsNode(ICoreProvider& Storage, TSharedPtr<IQueryNode> InQueryNode, EDeprecatedSyncFlags InSyncFlags);
+
 	protected:
 		TEDSQUERYSTACK_API virtual void Update(FTimespan AllottedTime) override;
+		TEDSQUERYSTACK_API virtual void RefreshInternal(FTimespan AllottedTime, FRowHandleArray& TargetRows);
 		TEDSQUERYSTACK_API virtual void RefreshInternal(FRowHandleArray& TargetRows);
+		virtual void QueryUpdated() {};
 
 		FRowHandleArray Rows;
 		TSharedPtr<IQueryNode> QueryNode;
 		ICoreProvider& Storage;
+		uint64 TopologyHash = 0;
 		RevisionId QueryRevision = UninitializedRevisionId;
 		RevisionId Revision = UninitializedRevisionId;
-		ESyncFlags SyncFlags = ESyncFlags::None;
+		ESyncActions SyncAction = ESyncActions::RefreshOnUpdate;
 	};
-
-	ENUM_CLASS_FLAGS(FRowQueryResultsNode::ESyncFlags);
 } // namespace UE::Editor::DataStorage::QueryStack
+
+PRAGMA_DISABLE_DEPRECATION_WARNINGS
+	ENUM_CLASS_FLAGS(UE::Editor::DataStorage::QueryStack::FRowQueryResultsNode::EDeprecatedSyncFlags)
+PRAGMA_ENABLE_DEPRECATION_WARNINGS
diff --git a/Engine/Plugins/Experimental/EditorDataStorageFeatures/Source/TedsTableViewer/Private/Widgets/STedsFilterBar.cpp b/Engine/Plugins/Experimental/EditorDataStorageFeatures/Source/TedsTableViewer/Private/Widgets/STedsFilterBar.cpp
--- a/Engine/Plugins/Experimental/EditorDataStorageFeatures/Source/TedsTableViewer/Private/Widgets/STedsFilterBar.cpp
+++ b/Engine/Plugins/Experimental/EditorDataStorageFeatures/Source/TedsTableViewer/Private/Widgets/STedsFilterBar.cpp
@@ -272,7 +272,7 @@ namespace UE::Editor::DataStorage
 				QueryNodes.Add(ExternalQuery.Value);
 			}
 			const TSharedPtr<IQueryNode> CombinedQueryNode = MakeShared<FQueryMergeNode>(*Storage, QueryNodes);
-			CombinedRowCollectorNode = MakeShared<FRowQueryResultsNode>(*Storage, CombinedQueryNode, FRowQueryResultsNode::ESyncFlags::RefreshOnUpdate);
+			CombinedRowCollectorNode = MakeShared<FRowQueryResultsNode>(*Storage, CombinedQueryNode, FRowQueryResultsNode::ESyncActions::ForceRefreshOnUpdate);
 			const TSharedPtr<QueryStack::IRowNode> TemporaryFilterNodesArray[2] = {CombinedRowCollectorNode, InFilterableRowNode};
 			
 			CombinedQueryFilterRowNode = MakeShared<FRowMergeNode>(TemporaryFilterNodesArray, FRowMergeNode::EMergeApproach::Repeating);
diff --git a/Engine/Plugins/Experimental/InstanceDataObjectFixupTool/Source/Private/DataRecoveryTool.cpp b/Engine/Plugins/Experimental/InstanceDataObjectFixupTool/Source/Private/DataRecoveryTool.cpp
--- a/Engine/Plugins/Experimental/InstanceDataObjectFixupTool/Source/Private/DataRecoveryTool.cpp
+++ b/Engine/Plugins/Experimental/InstanceDataObjectFixupTool/Source/Private/DataRecoveryTool.cpp
@@ -513,7 +513,7 @@ TSharedRef<SWidget> SDataRecoveryTool::MakeTEDSTableViewer()
 
 	ReferenceQueryResultsNode = MakeShared<QueryStack::FRowQueryResultsNode>(*Storage,
 		ReferenceQueryNode,
-		QueryStack::FRowQueryResultsNode::ESyncFlags::RefreshOnUpdate);
+		QueryStack::FRowQueryResultsNode::ESyncActions::ForceRefreshOnUpdate);
 
 	TSharedPtr<QueryStack::IRowNode> SearchNode;
 	TSharedPtr<SBox> SearchBoxContainer = SNew(SBox);
