TEDS Components/Alerts:
Change the actor component hierarchy name to be generic
Alerts about components are now added to the component rows instead of their parent row (e.g actor)
Added sync tags for component hierarchy changes

#Synced-CL 48607495

diff --git a/Engine/Plugins/Experimental/EditorDataStorageFeatures/Source/TedsActorCompatibility/Private/ActorComponentDebugHierarchyWidget/ActorComponentTree.cpp b/Engine/Plugins/Experimental/EditorDataStorageFeatures/Source/TedsActorCompatibility/Private/ActorComponentDebugHierarchyWidget/ActorComponentTree.cpp
--- a/Engine/Plugins/Experimental/EditorDataStorageFeatures/Source/TedsActorCompatibility/Private/ActorComponentDebugHierarchyWidget/ActorComponentTree.cpp
+++ b/Engine/Plugins/Experimental/EditorDataStorageFeatures/Source/TedsActorCompatibility/Private/ActorComponentDebugHierarchyWidget/ActorComponentTree.cpp
@@ -39,7 +39,7 @@ namespace UE::Editor::DataStorage::ActorCompatibility
 
 				ICoreProvider* DataStorageInterface = GetMutableDataStorageFeature<ICoreProvider>(StorageFeatureName);
 				
-				FHierarchyHandle HierarchyHandle = DataStorageInterface->FindHierarchyByName(TEXT("ActorComponent"));
+				FHierarchyHandle HierarchyHandle = DataStorageInterface->FindHierarchyByName(TEXT("ObjectComponent"));
 
 				if (DataStorageInterface->IsValidHierarchyHandle(HierarchyHandle))
 				{
diff --git a/Engine/Plugins/Experimental/EditorDataStorageFeatures/Source/TedsActorCompatibility/Private/TedsActorCompatibilityModule.cpp b/Engine/Plugins/Experimental/EditorDataStorageFeatures/Source/TedsActorCompatibility/Private/TedsActorCompatibilityModule.cpp
--- a/Engine/Plugins/Experimental/EditorDataStorageFeatures/Source/TedsActorCompatibility/Private/TedsActorCompatibilityModule.cpp
+++ b/Engine/Plugins/Experimental/EditorDataStorageFeatures/Source/TedsActorCompatibility/Private/TedsActorCompatibilityModule.cpp
@@ -54,7 +54,8 @@ namespace UE::Editor::DataStorage
 		
 				const FHierarchyRegistrationParams HierarchyParams
 				{
-					.Name = TEXT("ActorComponent"),
+					.Name = TEXT("ObjectComponent"),
+					.ParentChangedColumn = FComponentHierarchySyncTag::StaticStruct()
 				};
 
 				CoreProvider->RegisterHierarchy(HierarchyParams);
@@ -170,7 +171,7 @@ namespace UE::Editor::DataStorage
 								}
 							}
 						})
-						.AccessesHierarchy(TEXT("ActorComponent"))
+						.AccessesHierarchy(TEXT("ObjectComponent"))
 						.Where()
 							.All<FActorComponentTypeTag, FTypedElementSyncFromWorldTag>()
 						.Compile());
diff --git a/Engine/Plugins/Experimental/EditorDataStorageFeatures/Source/TedsActorCompatibility/Public/Columns/TedsActorComponentCompatibilityColumns.h b/Engine/Plugins/Experimental/EditorDataStorageFeatures/Source/TedsActorCompatibility/Public/Columns/TedsActorComponentCompatibilityColumns.h
--- a/Engine/Plugins/Experimental/EditorDataStorageFeatures/Source/TedsActorCompatibility/Public/Columns/TedsActorComponentCompatibilityColumns.h
+++ b/Engine/Plugins/Experimental/EditorDataStorageFeatures/Source/TedsActorCompatibility/Public/Columns/TedsActorComponentCompatibilityColumns.h
@@ -7,6 +7,13 @@
 
 USTRUCT()
 struct FActorComponentTypeTag : public FEditorDataStorageTag
+{
+	GENERATED_BODY()
+};
+
+// Tag added onto component rows when their parent changes. Removed at the end of the frame
+USTRUCT(meta = (DisplayName = "Component Hierarchy Change"))
+struct FComponentHierarchySyncTag final : public FEditorDataStorageTag
 {
 	GENERATED_BODY()
 };
\ No newline at end of file
