TEDS Hierarchies : Add an optional ETraversalOrder to ICoreProvider::WalkDepthFirst to support post-order traversal.

The default value is ETraversalOrder::PreOrder, so the behavior remains unchanged for all existing usages of this method.

#Synced-CL 48620407

diff --git a/Engine/Plugins/Experimental/EditorDataStorage/Source/TedsCore/Private/Hierarchy/EditorDataStorageHierarchyImplementation.cpp b/Engine/Plugins/Experimental/EditorDataStorage/Source/TedsCore/Private/Hierarchy/EditorDataStorageHierarchyImplementation.cpp
--- a/Engine/Plugins/Experimental/EditorDataStorage/Source/TedsCore/Private/Hierarchy/EditorDataStorageHierarchyImplementation.cpp
+++ b/Engine/Plugins/Experimental/EditorDataStorage/Source/TedsCore/Private/Hierarchy/EditorDataStorageHierarchyImplementation.cpp
@@ -379,31 +379,49 @@ namespace UE::Editor::DataStorage
 	void FTedsHierarchyAccessInterface::WalkDepthFirst(
 		const ICoreProvider& Context,
 		RowHandle Row,
-		TFunction<void(const ICoreProvider& Context, RowHandle Owner,RowHandle Target)> OnVisitedFn) const
+		TFunction<void(const ICoreProvider& Context, RowHandle Owner,RowHandle Target)> OnVisitedFn,
+		ICoreProvider::ETraversalOrder TraversalOrder) const
 	{
 		if (!Context.IsRowAvailable(Row))
 		{
 			return;
 		}
 		// Recursively walk depth first to visit each child
-		auto WalkDepthFirstImpl = [this](const ICoreProvider& Context_, RowHandle Row, decltype(OnVisitedFn)& OnVisitedFn_, auto& WalkDepthFirstImplRef) -> void
+		auto WalkDepthFirstImpl = [this, TraversalOrder](const ICoreProvider& Context_, RowHandle Row, decltype(OnVisitedFn)& OnVisitedFn_, auto& WalkDepthFirstImplRef) -> void
 		{
 			const FEditorDataHierarchyData_Template* HierarchyData = static_cast<const FEditorDataHierarchyData_Template*>(Context_.GetColumnData(Row, HierarchyDataColumnType));
 			if (HierarchyData)
 			{
 				for (RowHandle ChildRow : HierarchyData->Children)
-				{					
-					OnVisitedFn_(Context_, Row, ChildRow);
+				{
+					if (TraversalOrder == ICoreProvider::ETraversalOrder::PreOrder)
+					{
+						OnVisitedFn_(Context_, Row, ChildRow);
+					}
 					
 					WalkDepthFirstImplRef(Context_, ChildRow, OnVisitedFn_, WalkDepthFirstImplRef);
+
+					if (TraversalOrder == ICoreProvider::ETraversalOrder::PostOrder)
+					{
+						OnVisitedFn_(Context_, Row, ChildRow);
+					}
 				}
 			}
 		};
 
-		// Call the top level object
-		OnVisitedFn(Context, InvalidRowHandle /*no parent*/, Row);
+		if (TraversalOrder == ICoreProvider::ETraversalOrder::PreOrder)
+		{
+			// Call the top level object
+			OnVisitedFn(Context, InvalidRowHandle /*no parent*/, Row);
+		}
 
 		WalkDepthFirstImpl(Context, Row, OnVisitedFn, WalkDepthFirstImpl);
+
+		if (TraversalOrder == ICoreProvider::ETraversalOrder::PostOrder)
+		{
+			// Call the top level object
+			OnVisitedFn(Context, InvalidRowHandle /*no parent*/, Row);
+		}
 	}
 
 	void FTedsHierarchyAccessInterface::SetParentRow(ICoreProvider* CoreProvider, RowHandle Target, RowHandle Parent) const
diff --git a/Engine/Plugins/Experimental/EditorDataStorage/Source/TedsCore/Private/Hierarchy/EditorDataStorageHierarchyImplementation.h b/Engine/Plugins/Experimental/EditorDataStorage/Source/TedsCore/Private/Hierarchy/EditorDataStorageHierarchyImplementation.h
--- a/Engine/Plugins/Experimental/EditorDataStorage/Source/TedsCore/Private/Hierarchy/EditorDataStorageHierarchyImplementation.h
+++ b/Engine/Plugins/Experimental/EditorDataStorage/Source/TedsCore/Private/Hierarchy/EditorDataStorageHierarchyImplementation.h
@@ -81,7 +81,8 @@ namespace UE::Editor::DataStorage
 		void WalkDepthFirst(
 			const ICoreProvider& Context,
 			RowHandle Row,
-			TFunction<void(const ICoreProvider& Context, RowHandle Owner, RowHandle Target)> OnVisitedFn) const;
+			TFunction<void(const ICoreProvider& Context, RowHandle Owner, RowHandle Target)> OnVisitedFn,
+			ICoreProvider::ETraversalOrder TraversalOrder = ICoreProvider::ETraversalOrder::PreOrder) const;
 		
 		TFunction<RowHandle(const void*, const UScriptStruct*)> CreateParentExtractionFunction() const;
 		
diff --git a/Engine/Plugins/Experimental/EditorDataStorage/Source/TedsCore/Private/TypedElementDatabase.cpp b/Engine/Plugins/Experimental/EditorDataStorage/Source/TedsCore/Private/TypedElementDatabase.cpp
--- a/Engine/Plugins/Experimental/EditorDataStorage/Source/TedsCore/Private/TypedElementDatabase.cpp
+++ b/Engine/Plugins/Experimental/EditorDataStorage/Source/TedsCore/Private/TypedElementDatabase.cpp
@@ -1192,11 +1192,12 @@ bool UEditorDataStorage::HasChildren(UE::Editor::DataStorage::FHierarchyHandle I
 void UEditorDataStorage::WalkDepthFirst(
 	UE::Editor::DataStorage::FHierarchyHandle InHierarchyHandle,
 	UE::Editor::DataStorage::RowHandle Row,
-	FHierarchyIterationCallback VisitFn) const
+	FHierarchyIterationCallback VisitFn,
+	ETraversalOrder TraversalOrder) const
 {
 	if (const UE::Editor::DataStorage::FTedsHierarchyAccessInterface* AccessInterface = Environment->GetHierarchyRegistrar().GetAccessInterface(InHierarchyHandle))
 	{
-		return AccessInterface->WalkDepthFirst(*this, Row, MoveTemp(VisitFn));
+		return AccessInterface->WalkDepthFirst(*this, Row, MoveTemp(VisitFn), TraversalOrder);
 	}
 }
 
diff --git a/Engine/Plugins/Experimental/EditorDataStorage/Source/TedsCore/Private/TypedElementDatabase.h b/Engine/Plugins/Experimental/EditorDataStorage/Source/TedsCore/Private/TypedElementDatabase.h
--- a/Engine/Plugins/Experimental/EditorDataStorage/Source/TedsCore/Private/TypedElementDatabase.h
+++ b/Engine/Plugins/Experimental/EditorDataStorage/Source/TedsCore/Private/TypedElementDatabase.h
@@ -169,7 +169,8 @@ public:
 	UE_API virtual void WalkDepthFirst(
 		UE::Editor::DataStorage::FHierarchyHandle,
 		UE::Editor::DataStorage::RowHandle Row, 
-		TFunction<void(const ICoreProvider& Context, UE::Editor::DataStorage::RowHandle Owner, UE::Editor::DataStorage::RowHandle Target)> VisitFn) const override;
+		TFunction<void(const ICoreProvider& Context, UE::Editor::DataStorage::RowHandle Owner, UE::Editor::DataStorage::RowHandle Target)> VisitFn,
+		ETraversalOrder TraversalOrder = ETraversalOrder::PreOrder) const override;
 
 	UE_API void RegisterTickGroup(FName GroupName, UE::Editor::DataStorage::EQueryTickPhase Phase, FName BeforeGroup, FName AfterGroup, UE::Editor::DataStorage::EExecutionMode ExecutionMode);
 	UE_API void UnregisterTickGroup(FName GroupName, UE::Editor::DataStorage::EQueryTickPhase Phase);
diff --git a/Engine/Plugins/Experimental/EditorDataStorageFeatures/Source/TedsTableViewer/Private/HierarchyViewerIntefaces.cpp b/Engine/Plugins/Experimental/EditorDataStorageFeatures/Source/TedsTableViewer/Private/HierarchyViewerIntefaces.cpp
--- a/Engine/Plugins/Experimental/EditorDataStorageFeatures/Source/TedsTableViewer/Private/HierarchyViewerIntefaces.cpp
+++ b/Engine/Plugins/Experimental/EditorDataStorageFeatures/Source/TedsTableViewer/Private/HierarchyViewerIntefaces.cpp
@@ -26,9 +26,10 @@ namespace UE::Editor::DataStorage
 		return Storage.GetParentRow(HierarchyHandle, InRow);
 	}
 
-	void FHierarchyViewerData::WalkDepthFirst(const ICoreProvider& Storage, RowHandle InRow, ICoreProvider::FHierarchyIterationCallback VisitFn) const
+	void FHierarchyViewerData::WalkDepthFirst(const ICoreProvider& Storage, RowHandle InRow, ICoreProvider::FHierarchyIterationCallback VisitFn,
+		ICoreProvider::ETraversalOrder TraversalOrder) const
 	{
-		Storage.WalkDepthFirst(HierarchyHandle, InRow, VisitFn);
+		Storage.WalkDepthFirst(HierarchyHandle, InRow, VisitFn, TraversalOrder);
 	}
 
 	void FHierarchyViewerData::RegisterQueries(ICoreProvider& Storage)
@@ -97,7 +98,7 @@ namespace UE::Editor::DataStorage
 	}
 
 	void FHierarchyViewerMultiData::WalkDepthFirst(const ICoreProvider& Storage, RowHandle InRow,
-		ICoreProvider::FHierarchyIterationCallback VisitFn) const
+		ICoreProvider::FHierarchyIterationCallback VisitFn, ICoreProvider::ETraversalOrder TraversalOrder) const
 	{
 		TSet<RowHandle> Visited;
 		
@@ -110,7 +111,7 @@ namespace UE::Editor::DataStorage
 					VisitFn(Context, Owner, Target);
 					Visited.Add(Target);
 				}
-			});
+			}, TraversalOrder);
 		}
 	}
 }
diff --git a/Engine/Plugins/Experimental/EditorDataStorageFeatures/Source/TedsTableViewer/Public/HierarchyViewerIntefaces.h b/Engine/Plugins/Experimental/EditorDataStorageFeatures/Source/TedsTableViewer/Public/HierarchyViewerIntefaces.h
--- a/Engine/Plugins/Experimental/EditorDataStorageFeatures/Source/TedsTableViewer/Public/HierarchyViewerIntefaces.h
+++ b/Engine/Plugins/Experimental/EditorDataStorageFeatures/Source/TedsTableViewer/Public/HierarchyViewerIntefaces.h
@@ -21,7 +21,9 @@ namespace UE::Editor::DataStorage
 		
 		UE_API virtual ~IHierarchyViewerDataInterface() = default;
 		UE_API virtual RowHandle GetParent(const ICoreProvider& Storage, RowHandle InRow) const = 0;
-		UE_API virtual void WalkDepthFirst(const ICoreProvider& Storage, RowHandle InRow, ICoreProvider::FHierarchyIterationCallback VisitFn) const = 0;
+		UE_API virtual void WalkDepthFirst(const ICoreProvider& Storage, RowHandle InRow, ICoreProvider::FHierarchyIterationCallback VisitFn,
+			ICoreProvider::ETraversalOrder TraversalOrder = ICoreProvider::ETraversalOrder::PreOrder) const = 0;
+		
 		UE_API FOnParentChanged& OnHierarchyChangedEvent()
 		{
 			return ParentChangedEvent;
@@ -38,7 +40,8 @@ namespace UE::Editor::DataStorage
 		UE_API explicit FHierarchyViewerData(ICoreProvider& Storage, FHierarchyHandle InHierarchyHandle);
 		UE_API virtual ~FHierarchyViewerData() override;
 		UE_API virtual RowHandle GetParent(const ICoreProvider& Storage, RowHandle InRow) const override;
-		UE_API virtual void WalkDepthFirst(const ICoreProvider& Storage, RowHandle InRow, ICoreProvider::FHierarchyIterationCallback VisitFn) const override;
+		UE_API virtual void WalkDepthFirst(const ICoreProvider& Storage, RowHandle InRow, ICoreProvider::FHierarchyIterationCallback VisitFn,
+			ICoreProvider::ETraversalOrder TraversalOrder = ICoreProvider::ETraversalOrder::PreOrder) const override;
 
 	protected:
 		void RegisterQueries(ICoreProvider& Storage);
@@ -58,7 +61,8 @@ namespace UE::Editor::DataStorage
 		// Returns the first parent it finds among the list of hierarchies
 		UE_API virtual RowHandle GetParent(const ICoreProvider& Storage, RowHandle InRow) const override;
 		// Does a DFS of each hierarchy one-by-one. If a child is encountered multiple times in different hierarchies, it is only reported in the first instance
-		UE_API virtual void WalkDepthFirst(const ICoreProvider& Storage, RowHandle InRow, ICoreProvider::FHierarchyIterationCallback VisitFn) const override;
+		UE_API virtual void WalkDepthFirst(const ICoreProvider& Storage, RowHandle InRow, ICoreProvider::FHierarchyIterationCallback VisitFn,
+			ICoreProvider::ETraversalOrder TraversalOrder = ICoreProvider::ETraversalOrder::PreOrder) const override;
 		
 	protected:
 		TArray<TSharedPtr<FHierarchyViewerData>> Hierarchies;
diff --git a/Engine/Source/Runtime/TypedElementFramework/Public/Elements/Interfaces/TypedElementDataStorageInterface.h b/Engine/Source/Runtime/TypedElementFramework/Public/Elements/Interfaces/TypedElementDataStorageInterface.h
--- a/Engine/Source/Runtime/TypedElementFramework/Public/Elements/Interfaces/TypedElementDataStorageInterface.h
+++ b/Engine/Source/Runtime/TypedElementFramework/Public/Elements/Interfaces/TypedElementDataStorageInterface.h
@@ -425,12 +425,21 @@ public:
 
 	using FHierarchyIterationCallback = TFunction<void(const ICoreProvider& Context, RowHandle Owner, RowHandle Target)>;
 	
+	enum class ETraversalOrder : uint8
+	{
+		/** Specifies that a pre-order traversal should be performed. In a pre-order traversal, rows are visited before their children. */
+		PreOrder,
+		/** Specifies that a post-order traversal should be performed. In a post-order traversal, rows are visited after their children. */
+		PostOrder,
+	};
+
 	// Iterates depth first from the passed in Row, calling the VisitFn on each row
 	// The Row passed in will have VisitFn called on it
 	virtual void WalkDepthFirst(
 		FHierarchyHandle InHierarchyHandle,
 		RowHandle Row,
-		FHierarchyIterationCallback VisitFn) const = 0;
+		FHierarchyIterationCallback VisitFn,
+		ETraversalOrder TraversalOrder = ETraversalOrder::PreOrder) const = 0;
 	
 	/**
 	 * Outputs the registered query callbacks to the given output device for debugging purposes.
