[TEDS Outliner] Create TEDS Rows for the Table Viewer Rows and have expansion state be persisted

#Synced-CL 48615712

diff --git a/Engine/Plugins/Experimental/EditorDataStorageFeatures/Source/TedsContentBrowser/Private/Widgets/ContentBrowserTileViewWidget.cpp b/Engine/Plugins/Experimental/EditorDataStorageFeatures/Source/TedsContentBrowser/Private/Widgets/ContentBrowserTileViewWidget.cpp
--- a/Engine/Plugins/Experimental/EditorDataStorageFeatures/Source/TedsContentBrowser/Private/Widgets/ContentBrowserTileViewWidget.cpp
+++ b/Engine/Plugins/Experimental/EditorDataStorageFeatures/Source/TedsContentBrowser/Private/Widgets/ContentBrowserTileViewWidget.cpp
@@ -80,8 +80,7 @@ TSharedPtr<SWidget> FContentBrowserTileViewWidgetConstructor::CreateWidget(UE::E
 		FAttributeBinder WidgetRowBinder(WidgetRow, DataStorage);
 		FAttributeBinder RowBinder(TargetRow, DataStorage);
 
-		FHierarchyHandle WidgetHierarchy = DataStorage->FindHierarchyByName(TableViewerUtils::GetWidgetHierarchyName());
-		RowHandle ParentWidgetRowHandle = DataStorage->GetParentRow(WidgetHierarchy, WidgetRow);
+		const RowHandle ParentWidgetRowHandle = TableViewerUtils::GetTableViewerUiRow(DataStorage, WidgetRow);
 		
 		FAttributeBinder ParentWidgetRowBinder(ParentWidgetRowHandle, DataStorage);
 
diff --git a/Engine/Plugins/Experimental/EditorDataStorageFeatures/Source/TedsOutliner/Private/Widgets/OutlinerTextWidget.cpp b/Engine/Plugins/Experimental/EditorDataStorageFeatures/Source/TedsOutliner/Private/Widgets/OutlinerTextWidget.cpp
--- a/Engine/Plugins/Experimental/EditorDataStorageFeatures/Source/TedsOutliner/Private/Widgets/OutlinerTextWidget.cpp
+++ b/Engine/Plugins/Experimental/EditorDataStorageFeatures/Source/TedsOutliner/Private/Widgets/OutlinerTextWidget.cpp
@@ -78,10 +78,8 @@ namespace UE::Editor::Outliner::TextWidget::Private
 				})
 			});
 
-			// The table viewer, if present, is the parent widget of our widget row
-
-			DataStorage::FHierarchyHandle WidgetHierarchy = DataStorage->FindHierarchyByName(DataStorage::TableViewerUtils::GetWidgetHierarchyName());
-			DataStorage::RowHandle ParentRow = DataStorage->GetParentRow(WidgetHierarchy, WidgetRow);
+			// The table viewer, if present, is the grandparent widget of our widget row (Table -> Row -> Cell)
+			const DataStorage::RowHandle ParentRow = DataStorage::TableViewerUtils::GetTableViewerUiRow(DataStorage, WidgetRow);
 
 			if (DataStorage->IsRowAvailable(ParentRow))
 			{
diff --git a/Engine/Plugins/Experimental/EditorDataStorageFeatures/Source/TedsTableViewer/Private/TedsTableViewerUtils.cpp b/Engine/Plugins/Experimental/EditorDataStorageFeatures/Source/TedsTableViewer/Private/TedsTableViewerUtils.cpp
--- a/Engine/Plugins/Experimental/EditorDataStorageFeatures/Source/TedsTableViewer/Private/TedsTableViewerUtils.cpp
+++ b/Engine/Plugins/Experimental/EditorDataStorageFeatures/Source/TedsTableViewer/Private/TedsTableViewerUtils.cpp
@@ -10,6 +10,7 @@
 #include "Elements/Columns/TypedElementTypeInfoColumns.h"
 #include "Elements/Interfaces/TypedElementDataStorageUiInterface.h"
 #include "Styling/SlateIconFinder.h"
+#include "TedsTableViewerWidgetColumns.h"
 
 #include UE_INLINE_GENERATED_CPP_BY_NAME(TedsTableViewerUtils)
 
@@ -17,6 +18,7 @@ namespace UE::Editor::DataStorage::TableViewerUtils
 {
 	static FName TableViewerWidgetTableName("Editor_TableViewerWidgetTable");
 	static FName TableViewerWidgetHierarchyName("TableViewerWidgetHierarchy");
+	static FName TableViewerRowWidgetTableName("Editor_TableViewerRowWidgetTable");
 
 	FName GetWidgetTableName()
 	{
@@ -27,6 +29,11 @@ namespace UE::Editor::DataStorage::TableViewerUtils
 	{
 		return TableViewerWidgetHierarchyName;
 	}
+
+	FName GetRowWidgetTableName()
+	{
+		return TableViewerRowWidgetTableName;
+	}
 	
 	// TEDS UI TODO: Maybe the widget can specify a user facing name derived from the matched columns instead of trying to find the longest matching name
 	FName FindLongestMatchingName(TConstArrayView<TWeakObjectPtr<const UScriptStruct>> ColumnTypes, int32 DefaultNameIndex)
@@ -174,8 +181,28 @@ namespace UE::Editor::DataStorage::TableViewerUtils
 		}
 		
 		// Fallback to the regular actor icon if we haven't found any specific icon
-		return FSlateIconFinder::FindIconForClass(AActor::StaticClass()).GetOptionalIcon();;
+		return FSlateIconFinder::FindIconForClass(AActor::StaticClass()).GetOptionalIcon();
+
+	}
 
+	FString CreateRowMappingKey(const FName& HierarchyIdentifier, const TableViewerItemPtr Item)
+	{
+		return FString::Printf(TEXT("%s.%llu"), *HierarchyIdentifier.ToString(), Item.RowHandle);
+	}
+	
+	RowHandle GetTableViewerUiRow(const ICoreProvider* DataStorage, const RowHandle ChildWidgetRow)
+	{
+		const FHierarchyHandle WidgetHierarchy = DataStorage->FindHierarchyByName(GetWidgetHierarchyName());
+		RowHandle ParentRow = ChildWidgetRow;
+		while (ParentRow != InvalidRowHandle)
+		{
+			ParentRow = DataStorage->GetParentRow(WidgetHierarchy, ParentRow);
+			if (DataStorage->HasColumns<FTableViewerTag>(ParentRow))
+			{
+				break;
+			}
+		}
+		return ParentRow;
 	}
 }
 
@@ -191,6 +218,11 @@ void UTedsTableViewerFactory::RegisterTables(UE::Editor::DataStorage::ICoreProvi
 				FTypedElementRowReferenceColumn::StaticStruct()
 			}, TableViewerUtils::TableViewerWidgetTableName);
 	}
+	
+	DataStorage.RegisterTable(
+		{
+			FTableViewerRowTag::StaticStruct()
+		}, TableViewerUtils::TableViewerRowWidgetTableName);
 }
 
 void UTedsTableViewerFactory::RegisterHierarchies(UE::Editor::DataStorage::ICoreProvider& DataStorage)
diff --git a/Engine/Plugins/Experimental/EditorDataStorageFeatures/Source/TedsTableViewer/Private/Widgets/Composite/STedsCompositeHierarchyViewer.cpp b/Engine/Plugins/Experimental/EditorDataStorageFeatures/Source/TedsTableViewer/Private/Widgets/Composite/STedsCompositeHierarchyViewer.cpp
--- a/Engine/Plugins/Experimental/EditorDataStorageFeatures/Source/TedsTableViewer/Private/Widgets/Composite/STedsCompositeHierarchyViewer.cpp
+++ b/Engine/Plugins/Experimental/EditorDataStorageFeatures/Source/TedsTableViewer/Private/Widgets/Composite/STedsCompositeHierarchyViewer.cpp
@@ -106,6 +106,7 @@ namespace UE::Editor::DataStorage
 		HierarchyViewerArgs.AllNodeProvider(QueryStackNode);
 		SAssignNew(HierarchyViewer, SHierarchyViewer, MoveTemp(InHierarchyInterface))
 			.AllNodeProvider(HierarchyViewerArgs._AllNodeProvider)
+			.TableViewerIdentifier(HierarchyViewerArgs._TableViewerIdentifier)
 			.Columns(HierarchyViewerArgs._Columns)
 			.EmptyRowsMessage(HierarchyViewerArgs._EmptyRowsMessage)
 			.CellWidgetPurpose(HierarchyViewerArgs._CellWidgetPurpose)
@@ -135,6 +136,11 @@ namespace UE::Editor::DataStorage
 	{
 		HierarchyViewer->SetColumns(InColumns);
 	}
+	
+	void STedsCompositeHierarchyViewer::SetTableViewerIdentifier(const FName& NewIdentifier)
+    {
+		HierarchyViewer->SetTableViewerIdentifier(NewIdentifier);
+    }
 
 	void STedsCompositeHierarchyViewer::AddCustomRowWidget(const TSharedRef<FTedsTableViewerColumn>& InColumn)
 	{
diff --git a/Engine/Plugins/Experimental/EditorDataStorageFeatures/Source/TedsTableViewer/Private/Widgets/STedsHierarchyViewer.cpp b/Engine/Plugins/Experimental/EditorDataStorageFeatures/Source/TedsTableViewer/Private/Widgets/STedsHierarchyViewer.cpp
--- a/Engine/Plugins/Experimental/EditorDataStorageFeatures/Source/TedsTableViewer/Private/Widgets/STedsHierarchyViewer.cpp
+++ b/Engine/Plugins/Experimental/EditorDataStorageFeatures/Source/TedsTableViewer/Private/Widgets/STedsHierarchyViewer.cpp
@@ -4,10 +4,12 @@
 
 #include "TableViewerWidgetDropHandler.h"
 #include "TedsTableViewerColumn.h"
+#include "TedsTableViewerUtils.h"
 #include "TedsTableViewerWidgetColumns.h"
 #include "Columns/SlateDelegateColumns.h"
 #include "DataStorage/Features.h"
 #include "Elements/Columns/TypedElementHiearchyColumns.h"
+#include "Elements/Columns/TypedElementSlateWidgetColumns.h"
 #include "Elements/Columns/TypedElementUIColumns.h"
 #include "Elements/Framework/TypedElementDataStorageWidget.h"
 #include "QueryStack/TopLevelRowsNode.h"
@@ -20,6 +22,11 @@
 
 namespace UE::Editor::DataStorage
 {
+	SHierarchyViewer::~SHierarchyViewer()
+	{
+		RemoveAllTableViewerUIRows();
+	}
+	
 	void SHierarchyViewer::Construct(
 		const FArguments& InArgs,
 		TSharedPtr<IHierarchyViewerDataInterface> InHierarchyInterface)
@@ -63,6 +70,8 @@ namespace UE::Editor::DataStorage
 
 		ExecutorInstanceName = FName("TableViewerModel", InstanceCounter++);
 		QueryStackExecutor = QueryStack::FExplicitUpdateExecutor(ExecutorInstanceName, TopLevelRowsNode);
+
+		SetTableViewerIdentifier(InArgs._TableViewerIdentifier);
 		
 		HeaderRowWidget = SNew( SHeaderRow )
 							.CanSelectGeneratedColumn(InArgs._CanSelectGeneratedColumn);
@@ -81,6 +90,7 @@ namespace UE::Editor::DataStorage
 			.TopLevelRowsSource(&TopLevelRows)
 			.RowsSource(&Model->GetItems())
 			.OnGenerateRow(this, &SHierarchyViewer::MakeTableRowWidget)
+			.OnExpansionChanged(this, &SHierarchyViewer::OnExpansionChanged)
 			.OnSelectionChanged(this, &SHierarchyViewer::OnListSelectionChanged)
 			.SelectionMode(InArgs._ListSelectionMode);
 
@@ -191,7 +201,7 @@ namespace UE::Editor::DataStorage
 			{
 				// FHideRowFromUITag - The table viewer should not show up as a row in a table viewer because that will cause all sorts of recursion issues
 				// The others are Columns we are going to bind to attributes on STreeView
-				DataStorage->AddColumns<FHideRowFromUITag, FWidgetContextMenuColumn, FWidgetRowScrolledIntoView,
+				DataStorage->AddColumns<FTableViewerTag, FHideRowFromUITag, FWidgetContextMenuColumn, FWidgetRowScrolledIntoView,
 					FWidgetDoubleClickedColumn, FTableViewerSelectionChangedColumn, FTableViewerRenameSelectionColumn>(WidgetRowHandle);
 				DataStorage->AddColumn(WidgetRowHandle, FTableViewerReferenceColumn{.TableViewerWeak = SharedThis(this)});
 				DataStorage->AddColumn(WidgetRowHandle, FHierarchyViewerDataColumn{.HierarchyData = HierarchyInterface});
@@ -262,6 +272,25 @@ namespace UE::Editor::DataStorage
 		CreateInternalWidget();
 	}
 
+	void SHierarchyViewer::OnExpansionChanged(TableViewerItemPtr Item, bool bExpanded) const
+	{
+		if (ICoreProvider* Storage = Model->GetDataStorageInterface())
+		{
+			const FString TableRowId = TableViewerUtils::CreateRowMappingKey(TableViewerIdentifier, Item);
+			if (const RowHandle UIRow = Storage->LookupMappedRow(TableViewerUIRowMappingDomain, FMapKey(TableRowId)); Storage->IsRowAvailable(UIRow))
+			{
+				if (bExpanded)
+				{
+					Storage->AddColumns<FExpandedInUITag>(UIRow);
+				}
+				else
+				{
+					Storage->RemoveColumns<FExpandedInUITag>(UIRow);
+				}
+			}
+		}
+	}
+
 	void SHierarchyViewer::OnListSelectionChanged(TableViewerItemPtr Item, ESelectInfo::Type SelectInfo)
 	{
 		if(OnSelectionChanged.IsBound())
@@ -292,6 +321,15 @@ namespace UE::Editor::DataStorage
 		RefreshColumnWidgets();
 	}
 
+	void SHierarchyViewer::SetTableViewerIdentifier(const FName& NewIdentifier)
+	{
+		// This will remove any pre-existing row info with the previous name
+		RemoveAllTableViewerUIRows();
+		TableViewerIdentifier = NewIdentifier.IsNone() ?
+			FName("TEDSHierarchyViewerID", InstanceCounter) :
+			NewIdentifier;
+	}
+
 	void SHierarchyViewer::AddCustomRowWidget(const TSharedRef<FTedsTableViewerColumn>& InColumn)
 	{
 		Model->AddCustomRowWidget(InColumn);
@@ -385,11 +423,74 @@ namespace UE::Editor::DataStorage
 		return TreeView->IsItemVisible(InItem);
 	}
 
+	RowHandle SHierarchyViewer::AddOrFindTableViewerUIRow(TableViewerItemPtr InItem) const
+	{
+		if (ICoreProvider* Storage = Model->GetDataStorageInterface())
+		{
+			const FString MappingKey = TableViewerUtils::CreateRowMappingKey(TableViewerIdentifier, InItem);
+			// If the row already existed but was removed from TreeView's SparseItemInfos (Searching, Filters, etc.), re-expand it
+			if (const RowHandle ExistingRow = Storage->LookupMappedRow(TableViewerUIRowMappingDomain, FMapKey(MappingKey)); Storage->IsRowAvailable(ExistingRow))
+			{
+				if(Storage->HasColumns<FExpandedInUITag>(ExistingRow))
+				{
+					TreeView->SetItemExpansion(InItem, true);
+				}
+				return ExistingRow;
+			}
+			else
+			{
+				const RowHandle AddedRow = Storage->AddRow(Storage->FindTable(TableViewerUtils::GetRowWidgetTableName()));
+				
+				Storage->SetParentRow(Storage->FindHierarchyByName(TableViewerUtils::GetWidgetHierarchyName()), AddedRow, GetWidgetRowHandle());
+				Storage->MapRow(TableViewerUIRowMappingDomain, FMapKey(MappingKey), AddedRow);
+
+				if (!TreeView->IsItemExpanded(InItem))
+				{
+					// Default to expanded if not already
+					TreeView->SetItemExpansion(InItem, true);
+				}
+				else
+				{
+					Storage->AddColumns<FExpandedInUITag>(AddedRow);
+				}
+
+				return AddedRow;
+			}
+		}
+
+		return InvalidRowHandle;
+	}
+
+	void SHierarchyViewer::RemoveAllTableViewerUIRows() const
+	{
+		if(TedsWidget.IsValid() && Model.IsValid())
+		{
+			// Remove any rows attached to this table and the dynamic table
+			if (ICoreProvider* Storage = Model->GetDataStorageInterface())
+			{
+				TArray<RowHandle> ChildUIRows;
+				
+				Storage->WalkDepthFirst(Storage->FindHierarchyByName(TableViewerUtils::GetWidgetHierarchyName()), TedsWidget->GetRowHandle(),
+				[&ChildUIRows](const ICoreProvider& Context, RowHandle, const RowHandle Target)
+				{
+					if (Context.HasColumns<FTedsTableViewerRowTag>(Target))
+					{
+						ChildUIRows.Add(Target);
+					}
+				});
+				Storage->BatchRemoveRows(ChildUIRows);
+			}
+		}
+	}
+
 	TSharedRef<ITableRow> SHierarchyViewer::MakeTableRowWidget(TableViewerItemPtr InItem, const TSharedRef<STableViewBase>& OwnerTable) const
 	{
+		const RowHandle TableViewerRowUIRow = AddOrFindTableViewerUIRow(InItem);
+		
 		return SNew(SHierarchyViewerRow, OwnerTable, Model.ToSharedRef(), SharedThis(this))
 				.Padding(FMargin(0))
-				.ParentWidgetRowHandle(GetWidgetRowHandle())
+				.WidgetRowHandle(TableViewerRowUIRow == InvalidRowHandle ? GetWidgetRowHandle() : TableViewerRowUIRow)
+				.TableViewerWidgetRowHandle(GetWidgetRowHandle())
 				.Item(InItem);
 	}
 } // namespace UE::Editor::DataStorage
diff --git a/Engine/Plugins/Experimental/EditorDataStorageFeatures/Source/TedsTableViewer/Private/Widgets/STedsTableViewer.cpp b/Engine/Plugins/Experimental/EditorDataStorageFeatures/Source/TedsTableViewer/Private/Widgets/STedsTableViewer.cpp
--- a/Engine/Plugins/Experimental/EditorDataStorageFeatures/Source/TedsTableViewer/Private/Widgets/STedsTableViewer.cpp
+++ b/Engine/Plugins/Experimental/EditorDataStorageFeatures/Source/TedsTableViewer/Private/Widgets/STedsTableViewer.cpp
@@ -4,9 +4,11 @@
 
 #include "TableViewerWidgetDropHandler.h"
 #include "TedsTableViewerColumn.h"
+#include "TedsTableViewerUtils.h"
 #include "TedsTableViewerWidgetColumns.h"
 #include "Columns/SlateDelegateColumns.h"
 #include "DataStorage/Features.h"
+#include "Elements/Columns/TypedElementSlateWidgetColumns.h"
 #include "Elements/Columns/TypedElementUIColumns.h"
 #include "Elements/Framework/TypedElementAttributeBinding.h"
 #include "Elements/Framework/TypedElementDataStorageWidget.h"
@@ -18,6 +20,11 @@
 
 namespace UE::Editor::DataStorage
 {
+	STedsTableViewer::~STedsTableViewer()
+	{
+		RemoveAllTableViewerUIRows();
+	}
+	
 	void STedsTableViewer::Construct(const FArguments& InArgs)
 	{
 		OnSelectionChanged = InArgs._OnSelectionChanged;
@@ -54,6 +61,9 @@ namespace UE::Editor::DataStorage
 			.HeaderWidgetPurpose = HeaderWidgetPurpose,
 			.IsItemVisibleDelegate = FTedsTableViewerModel::FIsItemVisible::CreateSP(this, &STedsTableViewer::IsItemVisible)
 		});
+
+		InstanceCounter++;
+		SetTableViewerIdentifier(InArgs._TableViewerIdentifier);
 		
 		HeaderRowWidget = SNew( SHeaderRow )
 							.CanSelectGeneratedColumn(InArgs._CanSelectGeneratedColumn);
@@ -136,7 +146,7 @@ namespace UE::Editor::DataStorage
 				DataStorage->AddColumn(WidgetRowHandle, FHideRowFromUITag::StaticStruct());
 
 				// Columns we are going to bind to attributes on SListView
-				DataStorage->AddColumns<FWidgetContextMenuColumn, FWidgetRowScrolledIntoView,
+				DataStorage->AddColumns<FTableViewerTag, FWidgetContextMenuColumn, FWidgetRowScrolledIntoView,
 					FWidgetDoubleClickedColumn, FTableViewerSelectionChangedColumn, FTableViewerRenameSelectionColumn>(WidgetRowHandle);
 				DataStorage->AddColumn(WidgetRowHandle, FTableViewerReferenceColumn{.TableViewerWeak = SharedThis(this)});
 			}
@@ -205,6 +215,15 @@ namespace UE::Editor::DataStorage
 		}
 	}
 
+	void STedsTableViewer::SetTableViewerIdentifier(const FName& NewIdentifier)
+	{
+		// This will remove any pre-existing row info with the previous name
+		RemoveAllTableViewerUIRows();
+		TableViewerIdentifier = NewIdentifier.IsNone() ?
+			FName("TEDSTableViewerID", InstanceCounter) :
+			NewIdentifier;
+	}
+
 	void STedsTableViewer::SetQueryStack(TSharedPtr<QueryStack::IRowNode> InRowQueryStack)
 	{
 		Model->SetQueryStack(InRowQueryStack);
@@ -276,12 +295,61 @@ namespace UE::Editor::DataStorage
 		return ListView->IsItemVisible(InItem);
 	}
 
+	RowHandle STedsTableViewer::AddOrFindTableViewerUIRow(TableViewerItemPtr InItem) const
+	{
+		if (ICoreProvider* Storage = Model->GetDataStorageInterface())
+		{
+			const FString MappingKey = TableViewerUtils::CreateRowMappingKey(TableViewerIdentifier, InItem);
+			// If the row already existed but was removed from TreeView's SparseItemInfos (Searching, Filters, etc.), re-expand it
+			if (const RowHandle ExistingRow = Storage->LookupMappedRow(TableViewerUIRowMappingDomain, FMapKey(MappingKey)); Storage->IsRowAvailable(ExistingRow))
+			{
+				return ExistingRow;
+			}
+			else
+			{
+				const RowHandle AddedRow = Storage->AddRow(Storage->FindTable(TableViewerUtils::GetRowWidgetTableName()));
+				
+				Storage->SetParentRow(Storage->FindHierarchyByName(TableViewerUtils::GetWidgetHierarchyName()), AddedRow, GetWidgetRowHandle());
+				Storage->MapRow(TableViewerUIRowMappingDomain, FMapKey(MappingKey), AddedRow);
+				
+				return AddedRow;
+			}
+		}
+
+		return InvalidRowHandle;
+	}
+
+	void STedsTableViewer::RemoveAllTableViewerUIRows() const
+	{
+		if(TedsWidget.IsValid() && Model.IsValid())
+		{
+			// Remove any rows attached to this table and the dynamic table
+			if (ICoreProvider* Storage = Model->GetDataStorageInterface())
+			{
+				TArray<RowHandle> ChildUIRows;
+				
+				Storage->WalkDepthFirst(Storage->FindHierarchyByName(TableViewerUtils::GetWidgetHierarchyName()), TedsWidget->GetRowHandle(),
+				[&ChildUIRows](const ICoreProvider& Context, RowHandle, const RowHandle Target)
+				{
+					if (Context.HasColumns<FTedsTableViewerRowTag>(Target))
+					{
+						ChildUIRows.Add(Target);
+					}
+				});
+				Storage->BatchRemoveRows(ChildUIRows);
+			}
+		}
+	}
+
 	TSharedRef<ITableRow> STedsTableViewer::MakeTableRowWidget(TableViewerItemPtr InItem, const TSharedRef<STableViewBase>& OwnerTable) const
 	{
+		const RowHandle TableViewerRowUIRow = AddOrFindTableViewerUIRow(InItem);
+		
 		return SNew(STedsTableViewerRow, OwnerTable, Model.ToSharedRef(), SharedThis(this))
 				.ItemHeight(ItemHeight)
 				.Padding(ItemPadding)
-				.ParentWidgetRowHandle(GetWidgetRowHandle())
+				.WidgetRowHandle(TableViewerRowUIRow)
+				.TableViewerWidgetRowHandle(GetWidgetRowHandle())
 				.Item(InItem);
 	}
 } // namespace UE::Editor::DataStorage
diff --git a/Engine/Plugins/Experimental/EditorDataStorageFeatures/Source/TedsTableViewer/Private/Widgets/STedsTableViewerRow.cpp b/Engine/Plugins/Experimental/EditorDataStorageFeatures/Source/TedsTableViewer/Private/Widgets/STedsTableViewerRow.cpp
--- a/Engine/Plugins/Experimental/EditorDataStorageFeatures/Source/TedsTableViewer/Private/Widgets/STedsTableViewerRow.cpp
+++ b/Engine/Plugins/Experimental/EditorDataStorageFeatures/Source/TedsTableViewer/Private/Widgets/STedsTableViewerRow.cpp
@@ -21,7 +21,8 @@ namespace UE::Editor::DataStorage
 		Item = InArgs._Item;
 		TableViewerModel = MoveTemp(InTableViewerModel);
 		WeakTableViewer = MoveTemp(InTableViewer);
-		ParentWidgetRowHandle = InArgs._ParentWidgetRowHandle;
+		WidgetRowHandle = InArgs._WidgetRowHandle;
+		TableViewerWidgetRowHandle = InArgs._TableViewerWidgetRowHandle;
 		ItemHeight = InArgs._ItemHeight;
 
 		const auto Args = FSuperRowType::FArguments()
@@ -62,7 +63,7 @@ namespace UE::Editor::DataStorage
 	void STedsTableViewerRow::SetupWidgetColumns(ICoreProvider& InStorage, const RowHandle& InUIRowHandle)
 	{
 		const FHierarchyHandle WidgetHierarchy = InStorage.FindHierarchyByName(TableViewerUtils::GetWidgetHierarchyName());
-		InStorage.SetParentRow(WidgetHierarchy, InUIRowHandle, ParentWidgetRowHandle);
+		InStorage.SetParentRow(WidgetHierarchy, InUIRowHandle, WidgetRowHandle);
 
 		// Currently only this column is needed (not FExternalWidgetExclusiveSelectionColumn) since the double-click action will automatically
 		// deselect all and select only the desired row. 
@@ -83,7 +84,7 @@ namespace UE::Editor::DataStorage
 				// Add a column noting that a drop event is currently in progress
 				if (ICoreProvider* Storage = TableViewerModel->GetDataStorageInterface())
 				{
-					Storage->AddColumn(ParentWidgetRowHandle, FTableViewerDropInfoColumn{.DropTarget = InvalidRowHandle});
+					Storage->AddColumn(TableViewerWidgetRowHandle, FTableViewerDropInfoColumn{.DropTarget = InvalidRowHandle});
 				}
 
 				FRowHandleArray Rows;
@@ -107,7 +108,7 @@ namespace UE::Editor::DataStorage
 		// The actual drop logic is handled in the table viewer widget itself
 		if (ICoreProvider* Storage = TableViewerModel->GetDataStorageInterface())
 		{
-			Storage->AddColumn(ParentWidgetRowHandle, FTableViewerDropInfoColumn{.DropTarget = Item});
+			Storage->AddColumn(TableViewerWidgetRowHandle, FTableViewerDropInfoColumn{.DropTarget = Item});
 		}
 	}
 
@@ -117,7 +118,7 @@ namespace UE::Editor::DataStorage
 		// The actual drop logic is handled in the table viewer widget itself
 		if (ICoreProvider* Storage = TableViewerModel->GetDataStorageInterface())
 		{
-			if (FTableViewerDropInfoColumn* DropInfoColumn = Storage->GetColumn<FTableViewerDropInfoColumn>(ParentWidgetRowHandle))
+			if (FTableViewerDropInfoColumn* DropInfoColumn = Storage->GetColumn<FTableViewerDropInfoColumn>(TableViewerWidgetRowHandle))
 			{
 				DropInfoColumn->DropTarget = InvalidRowHandle;
 			}
@@ -160,8 +161,8 @@ namespace UE::Editor::DataStorage
 
 	void SHierarchyViewerRow::SetupWidgetColumns(ICoreProvider& InStorage, const RowHandle& InUIRowHandle)
 	{
-		// Forward hierarchy data from parent widget (hierarchy viewer) to cell widget for convinient access
-		if (const FHierarchyViewerDataColumn* Data = InStorage.GetColumn<FHierarchyViewerDataColumn>(ParentWidgetRowHandle))
+		// Forward hierarchy data from parent widget (row) to cell widget for convenient access
+		if (const FHierarchyViewerDataColumn* Data = InStorage.GetColumn<FHierarchyViewerDataColumn>(WidgetRowHandle))
 		{
 			InStorage.AddColumn(InUIRowHandle, FHierarchyViewerDataColumn{.HierarchyData = Data->HierarchyData });
 		}
diff --git a/Engine/Plugins/Experimental/EditorDataStorageFeatures/Source/TedsTableViewer/Private/Widgets/STedsTileViewer.cpp b/Engine/Plugins/Experimental/EditorDataStorageFeatures/Source/TedsTableViewer/Private/Widgets/STedsTileViewer.cpp
--- a/Engine/Plugins/Experimental/EditorDataStorageFeatures/Source/TedsTableViewer/Private/Widgets/STedsTileViewer.cpp
+++ b/Engine/Plugins/Experimental/EditorDataStorageFeatures/Source/TedsTableViewer/Private/Widgets/STedsTileViewer.cpp
@@ -258,7 +258,7 @@ namespace UE::Editor::DataStorage
 				DataStorage->AddColumn(WidgetRowHandle, FHideRowFromUITag::StaticStruct());
 
 				// Columns we are going to bind to attributes on SListView
-				DataStorage->AddColumns<FWidgetContextMenuColumn, FWidgetRowScrolledIntoView,
+				DataStorage->AddColumns<FTableViewerTag, FWidgetContextMenuColumn, FWidgetRowScrolledIntoView,
 					FWidgetDoubleClickedColumn, FTableViewerSelectionChangedColumn, FTableViewerRenameSelectionColumn>(WidgetRowHandle);
 				DataStorage->AddColumn(WidgetRowHandle, FTableViewerReferenceColumn{.TableViewerWeak = SharedThis(this)});
 
diff --git a/Engine/Plugins/Experimental/EditorDataStorageFeatures/Source/TedsTableViewer/Private/Widgets/STedsTreeView.cpp b/Engine/Plugins/Experimental/EditorDataStorageFeatures/Source/TedsTableViewer/Private/Widgets/STedsTreeView.cpp
--- a/Engine/Plugins/Experimental/EditorDataStorageFeatures/Source/TedsTableViewer/Private/Widgets/STedsTreeView.cpp
+++ b/Engine/Plugins/Experimental/EditorDataStorageFeatures/Source/TedsTableViewer/Private/Widgets/STedsTreeView.cpp
@@ -23,6 +23,7 @@ namespace UE::Editor::DataStorage
 			.HeaderRow(InArgs._HeaderRow)
 			.TreeItemsSource( InArgs._TopLevelRowsSource)
 			.OnGenerateRow(InArgs._OnGenerateRow)
+			.OnExpansionChanged(InArgs._OnExpansionChanged)
 			.OnSelectionChanged(InArgs._OnSelectionChanged)
 			.SelectionMode(InArgs._SelectionMode)
 			.OnContextMenuOpening(Binder.BindEvent(&FWidgetContextMenuColumn::OnContextMenuOpening))
diff --git a/Engine/Plugins/Experimental/EditorDataStorageFeatures/Source/TedsTableViewer/Private/Widgets/STedsTreeView.h b/Engine/Plugins/Experimental/EditorDataStorageFeatures/Source/TedsTableViewer/Private/Widgets/STedsTreeView.h
--- a/Engine/Plugins/Experimental/EditorDataStorageFeatures/Source/TedsTableViewer/Private/Widgets/STedsTreeView.h
+++ b/Engine/Plugins/Experimental/EditorDataStorageFeatures/Source/TedsTableViewer/Private/Widgets/STedsTreeView.h
@@ -30,6 +30,9 @@ namespace UE::Editor::DataStorage
 
 			/** Delegate fired on selection change */
 			SLATE_EVENT( FOnSelectionChanged, OnSelectionChanged )
+			
+			/** Delegate fired on expansion change */
+			SLATE_EVENT( FOnExpansionChanged, OnExpansionChanged )
 
 			/** Delegate that determines the selection mode */
 			SLATE_ATTRIBUTE( ESelectionMode::Type, SelectionMode )
diff --git a/Engine/Plugins/Experimental/EditorDataStorageFeatures/Source/TedsTableViewer/Public/ITedsTableViewer.h b/Engine/Plugins/Experimental/EditorDataStorageFeatures/Source/TedsTableViewer/Public/ITedsTableViewer.h
--- a/Engine/Plugins/Experimental/EditorDataStorageFeatures/Source/TedsTableViewer/Public/ITedsTableViewer.h
+++ b/Engine/Plugins/Experimental/EditorDataStorageFeatures/Source/TedsTableViewer/Public/ITedsTableViewer.h
@@ -49,6 +49,9 @@ namespace UE::Editor::DataStorage
 
     	// Clear the current list of columns being displayed and set it to the given list
     	virtual void SetColumns(const TArray<TWeakObjectPtr<const UScriptStruct>>& InColumns) = 0;
+    	
+    	// Sets the Table Identifier (Can be used in case data is dynamic and created post widget init)
+        virtual void SetTableViewerIdentifier(const FName& NewIdentifier) = 0;
    
     	virtual ~ITableViewer() = default;
     };
diff --git a/Engine/Plugins/Experimental/EditorDataStorageFeatures/Source/TedsTableViewer/Public/TedsTableViewerUtils.h b/Engine/Plugins/Experimental/EditorDataStorageFeatures/Source/TedsTableViewer/Public/TedsTableViewerUtils.h
--- a/Engine/Plugins/Experimental/EditorDataStorageFeatures/Source/TedsTableViewer/Public/TedsTableViewerUtils.h
+++ b/Engine/Plugins/Experimental/EditorDataStorageFeatures/Source/TedsTableViewer/Public/TedsTableViewerUtils.h
@@ -5,6 +5,7 @@
 #include "DataStorage/Handles.h"
 #include "Containers/ArrayView.h"
 #include "Elements/Interfaces/TypedElementDataStorageFactory.h"
+#include "TedsTableViewerModel.h"
 #include "UObject/Class.h"
 
 #include "TedsTableViewerUtils.generated.h"
@@ -27,6 +28,8 @@ namespace UE::Editor::DataStorage::TableViewerUtils
 	// Get the name of the widget hierarchy for table viewer widgets
 	// This hierarchy currently contains the relationship between the row containing the table viewer itself and individual cell widgets in that table viewer
 	TEDSTABLEVIEWER_API FName GetWidgetHierarchyName();
+
+	TEDSTABLEVIEWER_API FName GetRowWidgetTableName();
 	
 	// Find the longest matching common column name given a list of columns
 	TEDSTABLEVIEWER_API FName FindLongestMatchingName(TConstArrayView<TWeakObjectPtr<const UScriptStruct>> ColumnTypes, int32 DefaultNameIndex);
@@ -40,6 +43,12 @@ namespace UE::Editor::DataStorage::TableViewerUtils
 	
 	// Scans the row for columns and tries to find the best matching icon.
 	TEDSTABLEVIEWER_API const FSlateBrush* GetIconForRow(ICoreProvider* DataStorage, RowHandle Row);
+
+	// Creates the Row Map Key for the Table Row UI Rows
+	TEDSTABLEVIEWER_API FString CreateRowMappingKey(const FName& HierarchyIdentifier, const TableViewerItemPtr Item);
+
+	// Walks up the UI Widget Hierarchy until finding the row with the FTableViewerTag
+	TEDSTABLEVIEWER_API RowHandle GetTableViewerUiRow(const ICoreProvider* DataStorage, const RowHandle ChildWidgetRow);
 } // namespace UE::Editor::DataStorage::TableViewerUtils
 
 UCLASS()
diff --git a/Engine/Plugins/Experimental/EditorDataStorageFeatures/Source/TedsTableViewer/Public/TedsTableViewerWidgetColumns.h b/Engine/Plugins/Experimental/EditorDataStorageFeatures/Source/TedsTableViewer/Public/TedsTableViewerWidgetColumns.h
--- a/Engine/Plugins/Experimental/EditorDataStorageFeatures/Source/TedsTableViewer/Public/TedsTableViewerWidgetColumns.h
+++ b/Engine/Plugins/Experimental/EditorDataStorageFeatures/Source/TedsTableViewer/Public/TedsTableViewerWidgetColumns.h
@@ -9,6 +9,11 @@
 
 #include "TedsTableViewerWidgetColumns.generated.h"
 
+namespace UE::Editor::DataStorage
+{
+	inline static const FName TableViewerUIRowMappingDomain = TEXT("TableViewerUIRow");
+}
+
 /**
 * Column added to a widget row when an external widget manages exclusive selection for the widget referenced by the row,
  * such as an owning SListView or STreeView
@@ -27,6 +32,15 @@ struct FExternalWidgetExclusiveSelectionColumn : public FEditorDataStorageColumn
 	FIsSelected IsSelectedExclusively;
 };
 
+/**
+ * Tag on a UI row representing a table viewer (or hierarchy viewer) 
+ */
+USTRUCT(meta = (DisplayName = "Table Viewer Widget Tag"))
+struct FTedsTableViewerTag : public FEditorDataStorageTag
+{
+	GENERATED_BODY()
+};
+
 /**
  * Column on table viewer row that contains a reference to the ITableViewer
  */
@@ -99,12 +113,23 @@ struct FTedsHierarchyViewerRootColumn : public FEditorDataStorageColumn
 	UE::Editor::DataStorage::RowHandle Root;
 };
 
+/**
+ * Tag on a UI row representing a row in a table viewer (or hierarchy viewer) 
+ */
+USTRUCT(meta = (DisplayName = "Table Viewer Row Widget Tag"))
+struct FTedsTableViewerRowTag : public FEditorDataStorageTag
+{
+	GENERATED_BODY()
+};
+
 namespace UE::Editor::DataStorage
 {
+	using FTableViewerTag = FTedsTableViewerTag;
 	using FTableViewerReferenceColumn = FTedsTableViewerReferenceColumn;
 	using FTableViewerDropInfoColumn = FTedsTableViewerDropInfoColumn;
 	using FTableViewerSelectionChangedColumn = FTedsTableViewerSelectionChangedColumn;
 	using FTableViewerRenameSelectionColumn = FTedsTableViewerRenameSelectionColumn;
 	using FHierarchyViewerDataColumn = FTedsHierarchyViewerDataColumn;
 	using FHierarchyViewerRootColumn = FTedsHierarchyViewerRootColumn;
+	using FTableViewerRowTag = FTedsTableViewerRowTag;
 }
\ No newline at end of file
diff --git a/Engine/Plugins/Experimental/EditorDataStorageFeatures/Source/TedsTableViewer/Public/Widgets/Composite/STedsCompositeHierarchyViewer.h b/Engine/Plugins/Experimental/EditorDataStorageFeatures/Source/TedsTableViewer/Public/Widgets/Composite/STedsCompositeHierarchyViewer.h
--- a/Engine/Plugins/Experimental/EditorDataStorageFeatures/Source/TedsTableViewer/Public/Widgets/Composite/STedsCompositeHierarchyViewer.h
+++ b/Engine/Plugins/Experimental/EditorDataStorageFeatures/Source/TedsTableViewer/Public/Widgets/Composite/STedsCompositeHierarchyViewer.h
@@ -74,6 +74,9 @@ namespace UE::Editor::DataStorage
 
 		// Clear the current list of columns being displayed and set it to the given list
 		TEDSTABLEVIEWER_API virtual void SetColumns(const TArray<TWeakObjectPtr<const UScriptStruct>>& InColumns) override;
+		
+		// Sets the TableViewer/Hierarchy Identifier (Can be used in case data is dynamic and created post widget init)
+		TEDSTABLEVIEWER_API void SetTableViewerIdentifier(const FName& NewIdentifier);
 
 		// Add a custom per-row widget to the table viewer (that doesn't necessarily map to a TEDS column)
 		// This means a new column for the table viewer
diff --git a/Engine/Plugins/Experimental/EditorDataStorageFeatures/Source/TedsTableViewer/Public/Widgets/STedsHierarchyViewer.h b/Engine/Plugins/Experimental/EditorDataStorageFeatures/Source/TedsTableViewer/Public/Widgets/STedsHierarchyViewer.h
--- a/Engine/Plugins/Experimental/EditorDataStorageFeatures/Source/TedsTableViewer/Public/Widgets/STedsHierarchyViewer.h
+++ b/Engine/Plugins/Experimental/EditorDataStorageFeatures/Source/TedsTableViewer/Public/Widgets/STedsHierarchyViewer.h
@@ -61,6 +61,8 @@ namespace UE::Editor::DataStorage
 		{
 			
 		}
+		// Unique Name of this TableViewer (Hierarchy)
+		SLATE_ARGUMENT(FName, TableViewerIdentifier)
 
 		// Query Stack that will supply the rows to be displayed
 		SLATE_ARGUMENT(TSharedPtr<QueryStack::IRowNode>, AllNodeProvider)
@@ -77,7 +79,7 @@ namespace UE::Editor::DataStorage
 		// The widget purpose to use to create the header widgets
 		SLATE_ARGUMENT(IUiProvider::FPurposeID, HeaderWidgetPurpose)
 		
-		// Whether or not the user can select the columns generated by right clicking on the header menu.
+		// Whether the user can select the columns generated by right-clicking on the header menu.
 		SLATE_ARGUMENT(bool, CanSelectGeneratedColumn)
 
 		// Delegate called when the selection changes
@@ -96,6 +98,7 @@ namespace UE::Editor::DataStorage
 		SLATE_END_ARGS()
 
 	public:
+		TEDSTABLEVIEWER_API virtual ~SHierarchyViewer() override;
 		
 		TEDSTABLEVIEWER_API void Construct(
 			const FArguments& InArgs,
@@ -107,6 +110,9 @@ namespace UE::Editor::DataStorage
 		// Clear the current list of columns being displayed and set it to the given list
 		TEDSTABLEVIEWER_API virtual void SetColumns(const TArray<TWeakObjectPtr<const UScriptStruct>>& InColumns) override;
 
+		// Sets the Table Identifier (Can be used in case data is dynamic and created post widget init)
+		TEDSTABLEVIEWER_API virtual void SetTableViewerIdentifier(const FName& NewIdentifier) override;
+
     	// Add a custom per-row widget to the table viewer (that doesn't necessarily map to a TEDS column)
 		// This means a new column for the hierarchy viewer
 		TEDSTABLEVIEWER_API virtual void AddCustomRowWidget(const TSharedRef<FTedsTableViewerColumn>& InColumn) override;
@@ -141,6 +147,10 @@ namespace UE::Editor::DataStorage
 
 	protected:
 		
+		RowHandle AddOrFindTableViewerUIRow(TableViewerItemPtr InItem) const;
+
+		void RemoveAllTableViewerUIRows() const;
+		
 		TSharedRef<ITableRow> MakeTableRowWidget(TableViewerItemPtr InItem, const TSharedRef<STableViewBase>& OwnerTable) const;
 
 		bool IsItemVisible(TableViewerItemPtr InItem) const;
@@ -149,6 +159,8 @@ namespace UE::Editor::DataStorage
 
 		void RefreshColumnWidgets();
 
+		void OnExpansionChanged(TableViewerItemPtr Item, bool bExpanded) const;
+
 		void OnListSelectionChanged(TableViewerItemPtr Item, ESelectInfo::Type SelectInfo);
 
 		void AddWidgetColumns();
@@ -167,6 +179,9 @@ namespace UE::Editor::DataStorage
 
 	private:
 
+		// Unique identifier for this Hierarchy
+		FName TableViewerIdentifier;
+		
 		// The actual ListView widget that displays the rows
 		TSharedPtr<STedsTreeView> TreeView;
 
@@ -191,7 +206,7 @@ namespace UE::Editor::DataStorage
 		// The executor used to update the hierarchy stack.
 		QueryStack::FExplicitUpdateExecutor QueryStackExecutor;
 		
-		// Store the Instance count and the created name for use by the executor (so it can avoid recreation)
+		// Store the Instance count and the created name for use by the executor and identifier (so it can avoid recreation)
 		inline static int InstanceCounter = 0;
 		FName ExecutorInstanceName;
 		
diff --git a/Engine/Plugins/Experimental/EditorDataStorageFeatures/Source/TedsTableViewer/Public/Widgets/STedsTableViewer.h b/Engine/Plugins/Experimental/EditorDataStorageFeatures/Source/TedsTableViewer/Public/Widgets/STedsTableViewer.h
--- a/Engine/Plugins/Experimental/EditorDataStorageFeatures/Source/TedsTableViewer/Public/Widgets/STedsTableViewer.h
+++ b/Engine/Plugins/Experimental/EditorDataStorageFeatures/Source/TedsTableViewer/Public/Widgets/STedsTableViewer.h
@@ -51,6 +51,8 @@ namespace UE::Editor::DataStorage
 			
 		}
 		
+		// Unique Name of this Table
+		SLATE_ARGUMENT(FName, TableViewerIdentifier)
 
 		// Query Stack that will supply the rows to be displayed
 		SLATE_ARGUMENT(TSharedPtr<QueryStack::IRowNode>, QueryStack)
@@ -89,9 +91,13 @@ namespace UE::Editor::DataStorage
 		SLATE_END_ARGS()
 
 	public:
+		TEDSTABLEVIEWER_API virtual ~STedsTableViewer() override;
 		
 		TEDSTABLEVIEWER_API void Construct(const FArguments& InArgs);
 
+		// Sets the Table Identifier (Can be used in case data is dynamic and created post widget init)
+		TEDSTABLEVIEWER_API virtual void SetTableViewerIdentifier(const FName& NewIdentifier) override;
+
 		// Clear the current QueryStack being displayed, set it to the given node, and recreate the sorting nodes 
 		TEDSTABLEVIEWER_API virtual void SetQueryStack(TSharedPtr<QueryStack::IRowNode> InRowQueryStack) override;
 		
@@ -129,6 +135,10 @@ namespace UE::Editor::DataStorage
 		TEDSTABLEVIEWER_API virtual FReply OnDrop( const FGeometry& MyGeometry, const FDragDropEvent& DragDropEvent ) override;
 
 	protected:
+
+		RowHandle AddOrFindTableViewerUIRow(TableViewerItemPtr InItem) const;
+
+		void RemoveAllTableViewerUIRows() const;
 		
 		TSharedRef<ITableRow> MakeTableRowWidget(TableViewerItemPtr InItem, const TSharedRef<STableViewBase>& OwnerTable) const;
 
@@ -143,6 +153,9 @@ namespace UE::Editor::DataStorage
 		void AddWidgetColumns();
 
 	private:
+		// Unique identifier for this TableViewer
+		FName TableViewerIdentifier;
+		
 		// The actual height of each item
 		TAttribute<float> ItemHeight;
 
@@ -158,6 +171,9 @@ namespace UE::Editor::DataStorage
 		// Our model class
 		TSharedPtr<FTedsTableViewerModel> Model;
 
+		// Store the Instance count and the created name for use by the identifier
+		inline static int InstanceCounter = 0;
+
 		// Delegate fired when the selection changes
 		FOnSelectionChanged OnSelectionChanged;
 
diff --git a/Engine/Plugins/Experimental/EditorDataStorageFeatures/Source/TedsTableViewer/Public/Widgets/STedsTableViewerRow.h b/Engine/Plugins/Experimental/EditorDataStorageFeatures/Source/TedsTableViewer/Public/Widgets/STedsTableViewerRow.h
--- a/Engine/Plugins/Experimental/EditorDataStorageFeatures/Source/TedsTableViewer/Public/Widgets/STedsTableViewerRow.h
+++ b/Engine/Plugins/Experimental/EditorDataStorageFeatures/Source/TedsTableViewer/Public/Widgets/STedsTableViewerRow.h
@@ -23,7 +23,8 @@ namespace UE::Editor::DataStorage
 	public:
 
 		SLATE_BEGIN_ARGS( STedsTableViewerRow )
-			: _ParentWidgetRowHandle(InvalidRowHandle)
+			: _WidgetRowHandle(InvalidRowHandle)
+			, _TableViewerWidgetRowHandle(InvalidRowHandle)
 			, _Padding(FMargin(0.f, 4.f))
 			, _ItemHeight()
 		{}
@@ -31,7 +32,10 @@ namespace UE::Editor::DataStorage
 		/** The list item for this row */
 		SLATE_ARGUMENT( TableViewerItemPtr, Item )
 
-		SLATE_ARGUMENT( RowHandle, ParentWidgetRowHandle )
+			/** Row Handle of the parent of this */
+		SLATE_ARGUMENT( RowHandle, WidgetRowHandle )
+
+		SLATE_ARGUMENT( RowHandle, TableViewerWidgetRowHandle )
 
 		SLATE_ATTRIBUTE( FMargin, Padding )
 			
@@ -62,7 +66,8 @@ namespace UE::Editor::DataStorage
 
 	protected:
 		TAttribute<float> ItemHeight;
-		RowHandle ParentWidgetRowHandle = InvalidRowHandle;
+		RowHandle WidgetRowHandle = InvalidRowHandle;
+		RowHandle TableViewerWidgetRowHandle = InvalidRowHandle;
 		TSharedPtr<FTedsTableViewerModel> TableViewerModel;
 		TWeakPtr<const ITableViewer> WeakTableViewer;
 		TableViewerItemPtr Item;
diff --git a/Engine/Plugins/Experimental/EditorDataStorageFeatures/Source/TedsTableViewer/Public/Widgets/STedsTileViewer.h b/Engine/Plugins/Experimental/EditorDataStorageFeatures/Source/TedsTableViewer/Public/Widgets/STedsTileViewer.h
--- a/Engine/Plugins/Experimental/EditorDataStorageFeatures/Source/TedsTableViewer/Public/Widgets/STedsTileViewer.h
+++ b/Engine/Plugins/Experimental/EditorDataStorageFeatures/Source/TedsTableViewer/Public/Widgets/STedsTileViewer.h
@@ -101,6 +101,9 @@ namespace UE::Editor::DataStorage
 
 		TEDSTABLEVIEWER_API virtual bool IsSelectedExclusively(RowHandle InRow) const override;
 		
+		// Sets the Table Identifier (Can be used in case data is dynamic and created post widget init)
+        TEDSTABLEVIEWER_API virtual void SetTableViewerIdentifier(const FName& NewIdentifier) override {};
+		
 		// Clear the current QueryStack being displayed, set it to the given node, and recreate the sorting nodes 
         TEDSTABLEVIEWER_API virtual void SetQueryStack(TSharedPtr<QueryStack::IRowNode> InRowQueryStack) override;
 
diff --git a/Engine/Source/Runtime/TypedElementFramework/Public/Elements/Columns/TypedElementUIColumns.h b/Engine/Source/Runtime/TypedElementFramework/Public/Elements/Columns/TypedElementUIColumns.h
--- a/Engine/Source/Runtime/TypedElementFramework/Public/Elements/Columns/TypedElementUIColumns.h
+++ b/Engine/Source/Runtime/TypedElementFramework/Public/Elements/Columns/TypedElementUIColumns.h
@@ -12,6 +12,15 @@
  */
 USTRUCT(meta = (DisplayName = "Hide row from UI"))
 struct FHideRowFromUITag final : public FEditorDataStorageTag
+{
+	GENERATED_BODY()
+};
+
+/**
+ * Tag on to indicate that this row should be expanded in the UI
+ */
+USTRUCT(meta = (DisplayName = "Expanded item"))
+struct FExpandedInUITag : public FEditorDataStorageTag
 {
 	GENERATED_BODY()
 };
\ No newline at end of file
