[Teds Debugger] Fix crash in TEDS CoreProvider being accessed from the Model after deleted on Debugger closed

#Synced-CL 48615712

diff --git a/Engine/Plugins/Experimental/EditorDataStorageFeatures/Source/TedsDebugger/Private/Widgets/QueryEditor/TedsQueryEditorResultsView.cpp b/Engine/Plugins/Experimental/EditorDataStorageFeatures/Source/TedsDebugger/Private/Widgets/QueryEditor/TedsQueryEditorResultsView.cpp
--- a/Engine/Plugins/Experimental/EditorDataStorageFeatures/Source/TedsDebugger/Private/Widgets/QueryEditor/TedsQueryEditorResultsView.cpp
+++ b/Engine/Plugins/Experimental/EditorDataStorageFeatures/Source/TedsDebugger/Private/Widgets/QueryEditor/TedsQueryEditorResultsView.cpp
@@ -35,8 +35,10 @@ namespace UE::Editor::DataStorage::Debug::QueryEditor
 
 		if (CountQueryHandle != InvalidQueryHandle)
 		{
-			ICoreProvider& TedsInterface = Model->GetTedsInterface();
-			TedsInterface.UnregisterQuery(CountQueryHandle);
+			if (ICoreProvider* TedsInterface = GetMutableDataStorageFeature<ICoreProvider>(StorageFeatureName))
+			{
+				TedsInterface->UnregisterQuery(CountQueryHandle);
+			}
 			CountQueryHandle = InvalidQueryHandle;
 		}
 	}
