STedsTreeView : fix a race condition between list fresh and layout invalidation.

STreeView<TableViewerItemPtr>::RequestListRefresh() calls STableViewBase::RequestLayoutRefresh() which invalidates the layout and this was occuring prior to the call to UpdateTreeMap().  This change defers that until after the call to UpdateTreeMap() but it remains before the next STreeView::Tick.

#Synced-CL 48685028

diff --git a/Engine/Plugins/Experimental/EditorDataStorageFeatures/Source/TedsTableViewer/Private/Widgets/STedsTreeView.cpp b/Engine/Plugins/Experimental/EditorDataStorageFeatures/Source/TedsTableViewer/Private/Widgets/STedsTreeView.cpp
--- a/Engine/Plugins/Experimental/EditorDataStorageFeatures/Source/TedsTableViewer/Private/Widgets/STedsTreeView.cpp
+++ b/Engine/Plugins/Experimental/EditorDataStorageFeatures/Source/TedsTableViewer/Private/Widgets/STedsTreeView.cpp
@@ -37,6 +37,10 @@ namespace UE::Editor::DataStorage
 		if (bDirty)
 		{
 			UpdateTreeMap();
+			// The list refresh will invalidate the layout in STableViewBase::RequestLayoutRefresh()
+			// and we want to defer that until after we update the tree map.
+			// We are still marking the base class for refresh here prior to it's next tick which we call below.
+			STreeView<TableViewerItemPtr>::RequestListRefresh();
 			bDirty = false;
 		}
 
@@ -47,7 +51,6 @@ namespace UE::Editor::DataStorage
 	void STedsTreeView::RequestListRefresh()
 	{
 		bDirty = true;
-		STreeView<TableViewerItemPtr>::RequestListRefresh();
 	}
 	
 	void STedsTreeView::UpdateTreeMap()
