TEDS Alerts: remove any existing alerts using the requested name before adding a new alert, this fixes a bug where duplicate alerts could exist together instead of the second one overwriting the first

#Synced-CL 48607495

diff --git a/Engine/Plugins/Experimental/EditorDataStorageFeatures/Source/TedsAlerts/Private/TedsAlerts.cpp b/Engine/Plugins/Experimental/EditorDataStorageFeatures/Source/TedsAlerts/Private/TedsAlerts.cpp
--- a/Engine/Plugins/Experimental/EditorDataStorageFeatures/Source/TedsAlerts/Private/TedsAlerts.cpp
+++ b/Engine/Plugins/Experimental/EditorDataStorageFeatures/Source/TedsAlerts/Private/TedsAlerts.cpp
@@ -31,6 +31,9 @@ namespace UE::Editor::DataStorage::Alerts
 		// Add the new alert to the unsorted alert chain table and let TEDS sort is out at an appropriate time.
 		if (UTedsAlertsFactory* AlertsFactory = DataStorage.FindFactory<UTedsAlertsFactory>())
 		{
+			// To prevent duplicates, if there is already an alert with this name on the row remove and re-add it, which acts like an overwrite.
+			RemoveAlert(DataStorage, TargetRow, Name);
+			
 			// Add a placeholder alert to be filled in later.
 			if (!DataStorage.HasColumns<FAlertColumn>(TargetRow))
 			{
@@ -164,7 +167,7 @@ namespace UE::Editor::DataStorage::Alerts
 	{
 		using namespace UE::Editor::DataStorage::Columns;
 		using namespace UE::Editor::DataStorage::Queries;
-
+		
 		FAlertColumn* PreviousAlert = DataStorage.GetColumn<FAlertColumn>(TargetRow);
 		FAlertColumn* Alert = nullptr;
 		
