ViewportDropHandler now focuses viewport after successful drop.

#Synced-CL 48839584

diff --git a/Engine/Plugins/Experimental/EditorDataStorageFeatures/Source/TedsOperations/Private/DragAndDrop/Widgets/ViewportDropHandler.cpp b/Engine/Plugins/Experimental/EditorDataStorageFeatures/Source/TedsOperations/Private/DragAndDrop/Widgets/ViewportDropHandler.cpp
--- a/Engine/Plugins/Experimental/EditorDataStorageFeatures/Source/TedsOperations/Private/DragAndDrop/Widgets/ViewportDropHandler.cpp
+++ b/Engine/Plugins/Experimental/EditorDataStorageFeatures/Source/TedsOperations/Private/DragAndDrop/Widgets/ViewportDropHandler.cpp
@@ -12,8 +12,10 @@
 #include "Elements/Framework/TypedElementRegistry.h"
 #include "Elements/Interfaces/TypedElementDataStorageInterface.h"
 #include "Elements/Interfaces/TypedElementWorldInterface.h"
+#include "Framework/Application/SlateApplication.h"
 #include "Input/DragAndDrop.h"
 #include "Modules/ModuleManager.h"
+#include "SEditorViewport.h"
 #include "SceneView.h"
 #include "Settings/LevelEditorViewportSettings.h"
 #include "SnappingUtils.h"
@@ -186,6 +188,19 @@ void FViewportDropHandler::RemovePreviews(DataStorage::ICoreProvider& Storage, U
 	InvalidateViewport();
 }
 
+FReply FViewportDropHandler::OnDrop(const FGeometry& Geometry, const FDragDropEvent& InputEvent)
+{
+	FReply Reply = FPreviewWidgetDropHandler::OnDrop(Geometry, InputEvent);
+	if (Reply.IsEventHandled() && ViewportClient)
+	{
+		if(TSharedPtr<SEditorViewport> EditorViewport = ViewportClient->GetEditorViewportWidget())
+		{
+			FSlateApplication::Get().SetKeyboardFocus(EditorViewport);
+		}
+	}
+	return Reply;
+}
+
 void FViewportDropHandler::InvalidateViewport()
 {
 	if (ViewportClient && ViewportClient->Viewport)
diff --git a/Engine/Plugins/Experimental/EditorDataStorageFeatures/Source/TedsOperations/Public/DragAndDrop/Widgets/ViewportDropHandler.h b/Engine/Plugins/Experimental/EditorDataStorageFeatures/Source/TedsOperations/Public/DragAndDrop/Widgets/ViewportDropHandler.h
--- a/Engine/Plugins/Experimental/EditorDataStorageFeatures/Source/TedsOperations/Public/DragAndDrop/Widgets/ViewportDropHandler.h
+++ b/Engine/Plugins/Experimental/EditorDataStorageFeatures/Source/TedsOperations/Public/DragAndDrop/Widgets/ViewportDropHandler.h
@@ -23,6 +23,8 @@ public:
 	 */
 	UE_API FViewportDropHandler(TSharedPtr<FEditorViewportClient> InViewportClient, FTargetRowGetter InGetDefaultTargetRow);
 	UE_API virtual ~FViewportDropHandler() override;
+	
+	UE_API virtual FReply OnDrop(const FGeometry& Geometry, const FDragDropEvent& InputEvent) override;
 
 protected:
 	UE_API virtual EUpdateParameters UpdateParameters(DataStorage::ICoreProvider& Storage, const FGeometry& Geometry, const FDragDropEvent& InputEvent) override;
