[TEDS] Added 2 new nodes to the Query Stack for searching.

The new search node takes a parent query and a set of column variables/properties and searches the rows in the query for the provided string, resulting in only the rows that pass to be returned. If the query changes the search is automatically rerun. When the string changes the search is also rerun, unless the string is empty in which case the search is disabled.

This also includes a minor change to the query builder to prevent some needless copying of query conditions.

#Synced-CL 48813241

diff --git a/Engine/Plugins/Experimental/EditorDataStorageFeatures/Source/TedsQueryStack/Private/Searching/SearchUtils.cpp b/Engine/Plugins/Experimental/EditorDataStorageFeatures/Source/TedsQueryStack/Private/Searching/SearchUtils.cpp
new file mode 100644
--- /dev/null
+++ b/Engine/Plugins/Experimental/EditorDataStorageFeatures/Source/TedsQueryStack/Private/Searching/SearchUtils.cpp
@@ -0,0 +1,97 @@
+// Copyright Epic Games, Inc. All Rights Reserved.
+
+#include "Searching/SearchUtils.h"
+
+#include <algorithm>
+#include "Templates/FunctionFwd.h"
+#include "UObject/Class.h"
+#include "UObject/StrProperty.h"
+#include "UObject/TextProperty.h"
+#include "UObject/UnrealType.h"
+
+
+namespace UE::Editor::DataStorage::QueryStack::Searching
+{
+	size_t FCaseInsensitiveHash::operator()(FString::ElementType Value) const
+	{
+		return TChar<FString::ElementType>::ToLower(Value);
+	}
+	bool FCaseInsensitiveCompare::operator()(FString::ElementType Lhs, FString::ElementType Rhs) const
+	{
+		return  Lhs == Rhs ? true : TChar<FString::ElementType>::ToLower(Lhs) == TChar<FString::ElementType>::ToLower(Rhs);
+	}
+
+	FSearchContext::FSearchContext(FString InSearchString)
+		: SearchString(MoveTemp(InSearchString))
+		, StringSearcher(FStringSearcher(
+			*SearchString, *SearchString + SearchString.Len()))
+	{
+	}
+
+	StringSearchFunction CreateSearchFunction(const FProperty* Property)
+	{
+		if (const FStrProperty* StringProperty = CastField<FStrProperty>(Property))
+		{
+			return [](const Searching::FSearchContext& Context, const FProperty* Property, const void* Column, FString& TempString)
+				{
+					// Assume that the property has been checked at a higher call.
+					const FString* String = Property->ContainerPtrToValuePtr<FString>(Column);
+					return String ? Searching::Search(*String, Context, TempString) : false;
+				};
+		}
+		else if (const FTextProperty* TextProperty = CastField<FTextProperty>(Property))
+		{
+			return [](const Searching::FSearchContext& Context, const FProperty* Property, const void* Column, FString& TempString)
+				{
+					// Assume that the property has been checked at a higher call.
+					const FText* String = Property->ContainerPtrToValuePtr<FText>(Column);
+					return String ? Searching::Search(*String, Context, TempString) : false;
+				};
+		}
+		else if (const FNameProperty* NameProperty = CastField<FNameProperty>(Property))
+		{
+			return [](const Searching::FSearchContext& Context, const FProperty* Property, const void* Column, FString& TempString)
+				{
+					// Assume that the property has been checked at a higher call.
+					const FName* String = Property->ContainerPtrToValuePtr<FName>(Column);
+					return String ? Searching::Search(*String, Context, TempString) : false;
+				};
+		}
+		else
+		{
+			return nullptr;
+		}
+	}
+
+	bool Search(const FString& String, const Searching::FSearchContext& Context, FString& TempString)
+	{
+		FStringView View(String);
+		return !String.IsEmpty() && std::search(View.begin(), View.end(), Context.StringSearcher) != View.end();
+	}
+
+	bool Search(const FName& String, const Searching::FSearchContext& Context, FString& TempString)
+	{
+		TempString.Reset();
+		String.ToString(TempString);
+		return Search(TempString, Context, TempString);
+	}
+
+	bool Search(const FText& String, const Searching::FSearchContext& Context, FString& TempString)
+	{
+		return Search(String.ToString(), Context, TempString);
+	}
+
+	void ListSearchableProperties(const UScriptStruct* ColumnType, TFunctionRef<void(const FProperty*)> Callback)
+	{
+		static const FName SearchableMetaDataTag = FName("Searchable");
+
+		for (TFieldIterator<FProperty> PropertyIt(ColumnType); PropertyIt; ++PropertyIt)
+		{
+			const FProperty* Property = *PropertyIt;
+			if (Property->HasMetaData(SearchableMetaDataTag))
+			{
+				Callback(Property);
+			}
+		}
+	}
+} // namespace UE::Editor::DataStorage::QueryStack::Searching
diff --git a/Engine/Plugins/Experimental/EditorDataStorageFeatures/Source/TedsQueryStack/Private/TedsColumnsSearchNode.cpp b/Engine/Plugins/Experimental/EditorDataStorageFeatures/Source/TedsQueryStack/Private/TedsColumnsSearchNode.cpp
new file mode 100644
--- /dev/null
+++ b/Engine/Plugins/Experimental/EditorDataStorageFeatures/Source/TedsQueryStack/Private/TedsColumnsSearchNode.cpp
@@ -0,0 +1,245 @@
+// Copyright Epic Games, Inc. All Rights Reserved.
+
+#include "TedsColumnsSearchNode.h"
+
+#include "Elements/Framework/TypedElementConcurrentRowHandleCollector.h"
+#include "Elements/Framework/TypedElementQueryBuilder.h"
+#include "Elements/Framework/TypedElementQueryContext.h"
+#include "Internationalization/Text.h"
+#include "ProfilingDebugging/CpuProfilerTrace.h"
+#include "Searching/SearchUtils.h"
+#include "UObject/Class.h"
+#include "UObject/UnrealType.h"
+
+namespace UE::Editor::DataStorage::QueryStack
+{	
+	FColumnsSearchNode::FColumnsSearchNode(ICoreProvider& Storage, ESyncActions SyncAction)
+		: Storage(Storage)
+		, SyncAction(SyncAction)
+	{
+	}
+
+	FColumnsSearchNode::~FColumnsSearchNode()
+	{
+		UnregisterAllColumns();
+	}
+
+	void FColumnsSearchNode::StartSearch(FString SearchString)
+	{
+		if (!SearchString.IsEmpty())
+		{
+			SearchContext = MakeUnique<Searching::FSearchContext>(MoveTemp(SearchString));
+			Rows.Empty();
+			Revision++;
+		}
+		else
+		{
+			ClearSearch();
+		}
+	}
+
+	void FColumnsSearchNode::ClearSearch()
+	{
+		ResetSearch();
+		SearchContext.Reset();
+	}
+
+	bool FColumnsSearchNode::IsSearching() const
+	{
+		return SearchContext.IsValid() && SearchContext->CurrentSearcher < Searchers.Num();
+	}
+
+	FTimespan FColumnsSearchNode::GetQuerySearchTime() const
+	{
+		return TotalSearchTime;
+	}
+
+	uint64 FColumnsSearchNode::CalculateTotalSearchedCellCount() const
+	{
+		uint64 TotalCount = 0;
+		for (const FSearchQuery& Searcher : Searchers)
+		{
+			TotalCount += Storage.RunQuery(Searcher.Query).Count;
+		}
+		return TotalCount;
+	}
+
+	void FColumnsSearchNode::RegisterColumn(const UScriptStruct* ColumnType, const Queries::FConditions& Conditions)
+	{
+		Searching::ListSearchableProperties(ColumnType,
+			[this, ColumnType, &Conditions](const FProperty* Property)
+			{
+				RegisterColumn(ColumnType, Property, Conditions);
+			});
+	}
+
+	void FColumnsSearchNode::RegisterColumn(const UScriptStruct* ColumnType, const FProperty* Property, Queries::FConditions Conditions)
+	{
+		using namespace UE::Editor::DataStorage;
+		using namespace UE::Editor::DataStorage::Queries;
+
+		if (Searching::StringSearchFunction SearchFunction = Searching::CreateSearchFunction(Property);
+			ensureMsgf(SearchFunction, TEXT("Property '%s' in Column '%s' is marked in metadata as searchable but doesn't have a searchable type."),
+			*Property->GetName(), *ColumnType->GetFName().ToString()))
+		{
+			QueryHandle Query = CreateQuery(ColumnType, MoveTemp(Conditions));
+			uint64 QueryHash = Storage.CalculateQueryTablesTopologyHash(Query);
+
+			Searchers.Add(FSearchQuery{
+				.Query = Query,
+				.Hash = QueryHash,
+				.SearchFunction = SearchFunction,
+				.Property = Property,
+				.ColumnType = ColumnType
+				});
+		}
+	}
+
+	void FColumnsSearchNode::UnregisterAllColumns()
+	{
+		for (FSearchQuery& Searcher : Searchers)
+		{
+			Storage.UnregisterQuery(Searcher.Query);
+		}
+		Searchers.Reset();
+	}
+
+	INode::RevisionId FColumnsSearchNode::GetRevision() const
+	{
+		return Revision;
+	}
+
+	void FColumnsSearchNode::VisitParents(ParentListCallback Callback)
+	{
+		// No parents to visit.
+	}
+
+	void FColumnsSearchNode::Update(FTimespan AllottedTime)
+	{
+		if (Searching::FSearchContext* LocalSearchContext = SearchContext.Get())
+		{
+			double StartTime = FPlatformTime::Seconds();
+			
+			TRACE_CPUPROFILER_EVENT_SCOPE_STR("[TEDS] FColumnsSearchNode Tick");
+
+			if (SyncAction == ESyncActions::Refresh)
+			{
+				bool bIsDirty = false;
+				for (FSearchQuery& Searcher : Searchers)
+				{
+					uint64 QueryHash = Storage.CalculateQueryTablesTopologyHash(Searcher.Query);
+					bIsDirty = bIsDirty || QueryHash != Searcher.Hash;
+					Searcher.Hash = QueryHash;
+				}
+				if (bIsDirty)
+				{
+					ResetSearch();
+					LocalSearchContext->CurrentSearcher = 0;
+				}
+			}
+
+			double TimeBudget = AllottedTime.GetTotalSeconds();
+			double EndTime = FPlatformTime::Seconds();
+			while(LocalSearchContext->CurrentSearcher < Searchers.Num() && EndTime - StartTime < TimeBudget)
+			{
+				RunSearch(*LocalSearchContext, Searchers[LocalSearchContext->CurrentSearcher]);
+				LocalSearchContext->CurrentSearcher++;
+				EndTime = FPlatformTime::Seconds();
+				Revision++;
+			}
+
+			TotalSearchTime += FTimespan::FromSeconds(EndTime - StartTime);
+		}
+	}
+	
+	FRowHandleArrayView FColumnsSearchNode::GetRows() const
+	{
+		return Rows.GetRows();
+	}
+	
+	FRowHandleArray& FColumnsSearchNode::GetMutableRows()
+	{
+		return Rows;
+	}
+
+	QueryHandle FColumnsSearchNode::CreateQuery(const UScriptStruct* ColumnType, Queries::FConditions&& Conditions)
+	{
+		using namespace UE::Editor::DataStorage::Queries;
+
+		if (Conditions.IsEmpty())
+		{
+			return Storage.RegisterQuery(
+				Select()
+					.ReadOnly(ColumnType)
+				.Compile());
+		}
+		else
+		{
+			FEditorStorageQueryConditionCompileContext CompileContext(&Storage);
+			Conditions.Compile(CompileContext);
+
+			return Storage.RegisterQuery(
+				Select()
+					.ReadOnly(ColumnType)
+				.Where(MoveTemp(Conditions))
+				.Compile());
+		}
+	}
+
+	void FColumnsSearchNode::RunSearch(const Searching::FSearchContext& Context, const FSearchQuery& Searcher)
+	{
+		using namespace UE::Editor::DataStorage;
+		using namespace UE::Editor::DataStorage::Queries;
+
+		TRACE_CPUPROFILER_EVENT_SCOPE_STR("[TEDS] FColumnsSearchNode::RunSearch");
+
+		FConcurrentRowHandleSetCollector Collector;
+		
+		Storage.RunQuery(Searcher.Query,
+			ERunQueryFlags::ParallelizeChunks | ERunQueryFlags::AutoBalanceParallelChunkProcessing,
+			[&Context, &Searcher, &Collector](TQueryContext<RowBatchInfo> QueryContext)
+				{
+					if (const UScriptStruct* ColumnType = Searcher.ColumnType.Get(); ColumnType)
+					{
+						TRACE_CPUPROFILER_EVENT_SCOPE_TEXT(
+							*FString::Printf(TEXT("[TEDS] FColumnsSearchNode query '%s'"), *ColumnType->GetFName().ToString()));
+
+						int32 ColumnSize = ColumnType->GetCppStructOps() ? ColumnType->GetCppStructOps()->GetSize() : ColumnType->GetStructureSize();
+						
+						FRowHandleArray Temp;
+						// Run the compare on all columns in this batch and add the found results to the temp array.
+						if (const char* ColumnData = reinterpret_cast<const char*>(QueryContext.GetColumnBatchAddress(ColumnType)))
+						{
+							TRACE_CPUPROFILER_EVENT_SCOPE_STR(TEXT("[TEDS] FColumnsSearchNode query - Compare"));
+
+							FString TempString;
+							bool bReservedMemory = false;
+							for (RowHandle Row : QueryContext.GetBatchRowHandles())
+							{
+								if (Searcher.SearchFunction(Context, Searcher.Property, ColumnData, TempString))
+								{
+									if (!bReservedMemory)
+									{
+										// Avoid reserving memory if there are no results.
+										Temp.Reserve(QueryContext.GetBatchRowCount());
+										bReservedMemory = true;
+									}
+									Temp.Add(Row, FRowHandleArray::EFlags::IsUnique);
+								}
+								ColumnData += ColumnSize;
+							}
+						}
+						Collector.Append(MoveTemp(Temp));
+					}
+				});
+
+		Collector.Collect(Rows);
+	}
+	
+	void FColumnsSearchNode::ResetSearch()
+	{
+		TotalSearchTime = 0;
+		Rows.Empty();
+		Revision++;
+	}
+} // namespace UE::Editor::DataStorage::QueryStack
diff --git a/Engine/Plugins/Experimental/EditorDataStorageFeatures/Source/TedsQueryStack/Private/TedsQuerySearchNode.cpp b/Engine/Plugins/Experimental/EditorDataStorageFeatures/Source/TedsQueryStack/Private/TedsQuerySearchNode.cpp
new file mode 100644
--- /dev/null
+++ b/Engine/Plugins/Experimental/EditorDataStorageFeatures/Source/TedsQueryStack/Private/TedsQuerySearchNode.cpp
@@ -0,0 +1,203 @@
+// Copyright Epic Games, Inc. All Rights Reserved.
+
+#include "TedsQuerySearchNode.h"
+
+#include "Elements/Framework/TypedElementConcurrentRowHandleCollector.h"
+#include "Elements/Framework/TypedElementQueryBuilder.h"
+#include "Elements/Framework/TypedElementQueryContext.h"
+#include "Internationalization/Text.h"
+#include "ProfilingDebugging/CpuProfilerTrace.h"
+#include "Searching/SearchUtils.h"
+#include "UObject/Class.h"
+#include "UObject/UnrealType.h"
+
+namespace UE::Editor::DataStorage::QueryStack
+{	
+	FQuerySearchNode::FQuerySearchNode(ICoreProvider& Storage, TSharedRef<IQueryNode> Query)
+		: FRowQueryResultsNode(Storage, Query)
+	{
+		UpdateSearchersList();
+	}
+
+	FQuerySearchNode::FQuerySearchNode(ICoreProvider& Storage, TSharedRef<IQueryNode> Query, ESyncActions SyncAction)
+		: FRowQueryResultsNode(Storage, Query, SyncAction)
+	{
+		UpdateSearchersList();
+	}
+
+	void FQuerySearchNode::StartSearch(FString SearchString)
+	{
+		if (!SearchString.IsEmpty())
+		{
+			bHasPendingSearch = true;
+			SearchContext = MakeUnique<Searching::FSearchContext>(MoveTemp(SearchString));
+			ResetSearch();
+		}
+		else
+		{
+			ClearSearch();
+		}
+	}
+
+	void FQuerySearchNode::ClearSearch()
+	{
+		ResetSearch();
+		SearchContext.Reset();
+	}
+
+	bool FQuerySearchNode::IsSearching() const
+	{
+		return SearchContext.IsValid() && bHasPendingSearch;
+	}
+
+	FTimespan FQuerySearchNode::GetQuerySearchTime() const
+	{
+		return TotalSearchTime;
+	}
+
+	void FQuerySearchNode::QueryUpdated()
+	{
+		if (SearchContext.IsValid())
+		{
+			bHasPendingSearch = true;
+		}
+		UpdateSearchersList();
+	}
+
+	void FQuerySearchNode::RefreshInternal(FTimespan AllottedTime, FRowHandleArray& TargetRows)
+	{
+		if (Searching::FSearchContext* LocalSearchContext = SearchContext.Get(); LocalSearchContext)
+		{
+			TRACE_CPUPROFILER_EVENT_SCOPE_STR("[TEDS] FQuerySearchNode Tick");
+
+			double StartTime = FPlatformTime::Seconds();
+			RunSearch(*LocalSearchContext, TargetRows);
+			double EndTime = FPlatformTime::Seconds();
+			TotalSearchTime = FTimespan::FromSeconds(EndTime - StartTime);
+		}
+		else
+		{
+			FRowQueryResultsNode::RefreshInternal(AllottedTime, TargetRows);
+		}
+	}
+
+	void FQuerySearchNode::UpdateSearchersList()
+	{
+		// Reset the searching if there's one active.
+		ResetSearch();
+
+		// Collect the list of property searchers.
+		const FQueryDescription& Description = Storage.GetQueryDescription(QueryNode->GetQuery());
+		Searchers.Reset();
+
+		auto AddSearcher = [](TArray<FSearchQuery>& InSearchers, const UScriptStruct* Column)
+			{
+				Searching::ListSearchableProperties(Column,
+					[&InSearchers, Column](const FProperty* Property)
+					{
+						if (Searching::StringSearchFunction SearchFunction = Searching::CreateSearchFunction(Property))
+						{
+							InSearchers.Add(FSearchQuery{
+									.SearchFunction = SearchFunction,
+									.Property = Property,
+									.ColumnType = Column,
+									.ColumnSize = Column->GetCppStructOps() ? Column->GetCppStructOps()->GetSize() : Column->GetStructureSize()
+								});
+						}
+					});
+			};
+
+		for (TWeakObjectPtr<const UScriptStruct> Column : Description.SelectionTypes)
+		{
+			if (const UScriptStruct* ColumnPtr = Column.Get())
+			{
+				AddSearcher(Searchers, ColumnPtr);
+			}
+		}
+
+		for (const FDynamicColumnDescription& DynamicColumn : Description.DynamicSelectionTypes)
+		{
+			if (DynamicColumn.TemplateType)
+			{
+				if (const UScriptStruct* ColumnPtr = Storage.FindDynamicColumnType(*DynamicColumn.TemplateType, DynamicColumn.Identifier))
+				{
+					AddSearcher(Searchers, ColumnPtr);
+				}
+			}
+		}
+	}
+
+	void FQuerySearchNode::RunSearch(const Searching::FSearchContext& Context, FRowHandleArray& TargetRows)
+	{
+		using namespace UE::Editor::DataStorage;
+		using namespace UE::Editor::DataStorage::Queries;
+
+		FConcurrentRowHandleSetCollector Collector;
+		
+		Storage.RunQuery(QueryNode->GetQuery(),
+			ERunQueryFlags::ParallelizeChunks | ERunQueryFlags::AutoBalanceParallelChunkProcessing |
+			ERunQueryFlags::AllowBoundQueries | ERunQueryFlags::IgnoreActiveState,
+				[this, &Context, &Collector](TQueryContext<RowBatchInfo> QueryContext)
+				{
+					TRACE_CPUPROFILER_EVENT_SCOPE_TEXT(TEXT("[TEDS] FQuerySearchNode query"));
+
+					// Collect the column addresses.
+					struct FInfo { const char* Address; const FSearchQuery& Searcher; };
+					TArray<FInfo, TInlineAllocator<InlineColumnCount>> ColumnInfo;
+					ColumnInfo.Reserve(Searchers.Num());
+					for (const FSearchQuery& Searcher : Searchers)
+					{
+						if (const UScriptStruct* ColumnType = Searcher.ColumnType.Get(); ColumnType)
+						{
+							ColumnInfo.Add(FInfo
+								{
+									.Address = static_cast<const char*>(QueryContext.GetColumnBatchAddress(ColumnType)),
+									.Searcher = Searcher
+								});
+						}
+					}
+
+					FRowHandleArray Temp;
+					FString TempString;
+					bool bReservedMemory = false;
+
+					for (RowHandle Row : QueryContext.GetBatchRowHandles())
+					{
+						// Compare data on current row.
+						for (const FInfo& Info : ColumnInfo)
+						{
+							if (Info.Searcher.SearchFunction(Context, Info.Searcher.Property, Info.Address, TempString))
+							{
+								if (!bReservedMemory)
+								{
+									// Avoid reserving memory if there are no results.
+									Temp.Reserve(QueryContext.GetBatchRowCount());
+									bReservedMemory = true;
+								}
+								Temp.Add(Row, FRowHandleArray::EFlags::IsUnique);
+								break; // Make sure each row is only added once.
+							}
+						}
+
+						// Move to next row
+						for (FInfo& Info : ColumnInfo)
+						{
+							Info.Address += Info.Searcher.ColumnSize;
+						}
+					}
+
+					Collector.Append(MoveTemp(Temp));
+				});
+
+		Collector.Collect(TargetRows);
+		bHasPendingSearch = false;
+	}
+
+	void FQuerySearchNode::ResetSearch()
+	{
+		TotalSearchTime = 0;
+		TopologyHash = 0; // Force a refresh.
+		Rows.Empty();
+		Revision++;
+	}
+} // namespace UE::Editor::DataStorage::QueryStack
diff --git a/Engine/Plugins/Experimental/EditorDataStorageFeatures/Source/TedsQueryStack/Public/Searching/SearchUtils.h b/Engine/Plugins/Experimental/EditorDataStorageFeatures/Source/TedsQueryStack/Public/Searching/SearchUtils.h
new file mode 100644
--- /dev/null
+++ b/Engine/Plugins/Experimental/EditorDataStorageFeatures/Source/TedsQueryStack/Public/Searching/SearchUtils.h
@@ -0,0 +1,72 @@
+// Copyright Epic Games, Inc. All Rights Reserved.
+
+#pragma once
+
+#include <functional>
+#include <type_traits>
+#include "Containers/UnrealString.h"
+#include "DataStorage/CommonTypes.h"
+#include "Internationalization/Text.h"
+#include "UObject/NameTypes.h"
+
+class FProperty;
+class UScriptStruct;
+
+namespace UE::Editor::DataStorage::QueryStack::Searching
+{
+	template <typename T> struct MemberPointerTraits;
+	template <typename Type, typename Class>
+	struct MemberPointerTraits<Type Class::*>
+	{
+		using MemberType = Type;
+		using ClassType = Class;
+	};
+
+	template <typename T>
+	concept SupportedSearchPropertyType = std::is_same_v<T, FString> || std::is_same_v<T, FText> || std::is_same_v<T, FName>;
+
+	template <typename T>
+	concept MemberObjectPointer = std::is_member_object_pointer_v<T>;
+
+	// Concept to allow binding to a searchable column variable.
+	template <auto Ptr>
+	concept SearchableMemberColumn =
+		MemberObjectPointer<decltype(Ptr)> &&
+		TDataColumnType<typename MemberPointerTraits<decltype(Ptr)>::ClassType> &&
+		SupportedSearchPropertyType<typename MemberPointerTraits<decltype(Ptr)>::MemberType>;
+
+	struct FCaseInsensitiveHash
+	{
+		size_t operator()(FString::ElementType Value) const;
+	};
+	struct FCaseInsensitiveCompare
+	{
+		bool operator()(FString::ElementType Lhs, FString::ElementType Rhs) const;
+	};
+	using FStringSearcher = std::boyer_moore_horspool_searcher<const FString::ElementType*, FCaseInsensitiveHash, FCaseInsensitiveCompare>;
+
+	struct FSearchContext
+	{
+		FSearchContext(FString InSearchString);
+
+		FString SearchString;
+		FStringSearcher StringSearcher;
+		int32 CurrentSearcher = 0;
+	};
+
+	using StringSearchFunction = bool(*)(const Searching::FSearchContext& Context, const FProperty* Property, const void* Column, FString& TempString);
+
+	template<auto MemberVariable> requires Searching::SearchableMemberColumn<MemberVariable>
+	StringSearchFunction CreateSearchFunction();
+	TEDSQUERYSTACK_API StringSearchFunction CreateSearchFunction(const FProperty* Property);
+
+	template<auto MemberVariable> requires SearchableMemberColumn<MemberVariable>
+	bool Search(const void* Column, const FSearchContext& Context, FString& TempString);
+	TEDSQUERYSTACK_API bool Search(const FString& String, const FSearchContext& Context, FString& TempString);
+	TEDSQUERYSTACK_API bool Search(const FName& String, const FSearchContext& Context, FString& TempString);
+	TEDSQUERYSTACK_API bool Search(const FText& String, const FSearchContext& Context, FString& TempString);
+
+	TEDSQUERYSTACK_API void ListSearchableProperties(const UScriptStruct* ColumnType, TFunctionRef<void(const FProperty*)> Callback);
+} // namespace UE::Editor::DataStorage::QueryStack::Searching
+
+#include "Searching/SearchUtils.inl"
\ No newline at end of file
diff --git a/Engine/Plugins/Experimental/EditorDataStorageFeatures/Source/TedsQueryStack/Public/Searching/SearchUtils.inl b/Engine/Plugins/Experimental/EditorDataStorageFeatures/Source/TedsQueryStack/Public/Searching/SearchUtils.inl
new file mode 100644
--- /dev/null
+++ b/Engine/Plugins/Experimental/EditorDataStorageFeatures/Source/TedsQueryStack/Public/Searching/SearchUtils.inl
@@ -0,0 +1,22 @@
+// Copyright Epic Games, Inc. All Rights Reserved.
+
+namespace UE::Editor::DataStorage::QueryStack::Searching
+{
+	template<auto MemberVariable> requires Searching::SearchableMemberColumn<MemberVariable>
+	StringSearchFunction CreateSearchFunction()
+	{
+		return [](const Searching::FSearchContext& Context, const FProperty* Searcher, const void* Column, FString& TempString)
+			{
+				return Searching::Search<MemberVariable>(Column, Context, TempString);
+			};
+	}
+
+	template<auto MemberVariable> requires SearchableMemberColumn<MemberVariable>
+	bool Search(const void* Column, const FSearchContext& Context, FString& TempString)
+	{
+		using Class = typename Searching::MemberPointerTraits<decltype(MemberVariable)>::ClassType;
+		
+		const Class& ColumnCast = *static_cast<const Class*>(Column);
+		return Search(ColumnCast.*MemberVariable, Context, TempString);
+	}
+} // namespace UE::Editor::DataStorage::QueryStack::Searching
diff --git a/Engine/Plugins/Experimental/EditorDataStorageFeatures/Source/TedsQueryStack/Public/TedsColumnsSearchNode.h b/Engine/Plugins/Experimental/EditorDataStorageFeatures/Source/TedsQueryStack/Public/TedsColumnsSearchNode.h
new file mode 100644
--- /dev/null
+++ b/Engine/Plugins/Experimental/EditorDataStorageFeatures/Source/TedsQueryStack/Public/TedsColumnsSearchNode.h
@@ -0,0 +1,142 @@
+// Copyright Epic Games, Inc. All Rights Reserved.
+
+#pragma once
+
+#include <type_traits>
+#include "Containers/StringFwd.h"
+#include "DataStorage/CommonTypes.h"
+#include "DataStorage/Queries/Conditions.h"
+#include "Elements/Framework/TypedElementRowHandleArray.h"
+#include "Misc/Timespan.h"
+#include "TedsQueryStackInterfaces.h"
+#include "Templates/UniquePtr.h"
+
+#include "Searching/SearchUtils.h"
+
+class FProperty;
+class UScriptStruct;
+namespace UE::Editor::DataStorage
+{
+	class ICoreProvider;
+} // namespace UE::Editor::DataStorage
+
+namespace UE::Editor::DataStorage::QueryStack
+{
+	namespace Searching
+	{
+		struct FSearchContext;
+	}
+
+	/**
+	 * Searches through the provided list of columns and generates a list of rows for all columns with a property that match the search 
+	 * string.
+	 */
+	class FColumnsSearchNode final : public IRowNode
+	{
+	public:
+		enum ESyncActions
+		{
+			None, /** Do nothing when there are changes to the tables that have columns that are being searched. */
+			Refresh, /** If any of the tables with columns that are being searched change, rerun the search. */
+		};
+
+		TEDSQUERYSTACK_API FColumnsSearchNode(ICoreProvider& Storage, ESyncActions SyncAction);
+		TEDSQUERYSTACK_API virtual ~FColumnsSearchNode() override;
+
+		/**
+		 * Starts searching for the provided string in all searchable columns that were selected in the query.
+		 * Results may not be collected and can take several frames to arrive.
+		 */
+		TEDSQUERYSTACK_API void StartSearch(FString SearchString);
+		/** Resets a search. */
+		TEDSQUERYSTACK_API void ClearSearch();
+		/** Returns true if a search is still being processed. */
+		TEDSQUERYSTACK_API bool IsSearching() const;
+
+		/**
+		 * The total time spend on performing the search.
+		 * The returned time span represents the time spend performing search operations. This differs from the wall clock time in that
+		 * it exclusively tracks the time spend on performing one or more steps in the search.
+		 */
+		TEDSQUERYSTACK_API FTimespan GetQuerySearchTime() const;
+
+		/** The total number of cells that will be searched through. */
+		TEDSQUERYSTACK_API uint64 CalculateTotalSearchedCellCount() const;
+
+		/** 
+		 * Register a column to search.
+		 * All properties in the column that are marked with "Searchable" will be included.
+		 */
+		template<TDataColumnType Column>
+		void RegisterColumn(Queries::FConditions Conditions = {});
+		/**
+		 * Register a column to search.
+		 * All properties in the column that are marked with "Searchable" will be included.
+		 */
+		TEDSQUERYSTACK_API void RegisterColumn(const UScriptStruct* ColumnType, const Queries::FConditions& Conditions = {});
+		/**
+		 * Register a specific property in a column to search.
+		 * Only the provided property will be searched if it contains a type that can be searched. It does not require to be marked with
+		 * the "Searchable" metadata.
+		 */
+		TEDSQUERYSTACK_API void RegisterColumn(const UScriptStruct* ColumnType, const FProperty* Property, Queries::FConditions Conditions = {});
+		/**
+		 * Register a specific variable in a column to search.
+		 * Only the provided variable will be searched if it contains a type that can be searched.
+		 */
+		template<auto MemberVariable>  requires Searching::SearchableMemberColumn<MemberVariable>
+		void RegisterColumn(Queries::FConditions Conditions = {});
+		/** Clear all registered columns, properties and variables. */
+		TEDSQUERYSTACK_API void UnregisterAllColumns();
+
+		TEDSQUERYSTACK_API virtual RevisionId GetRevision() const override;
+		TEDSQUERYSTACK_API virtual void VisitParents(ParentListCallback Callback) override;
+		TEDSQUERYSTACK_API virtual FRowHandleArrayView GetRows() const override;
+		TEDSQUERYSTACK_API virtual FRowHandleArray& GetMutableRows() override;
+
+	protected:
+		TEDSQUERYSTACK_API virtual void Update(FTimespan AllottedTime) override;
+	
+	private:
+		struct FSearchQuery
+		{
+			QueryHandle Query = InvalidQueryHandle;
+			uint64 Hash = 0;
+			Searching::StringSearchFunction SearchFunction = nullptr;
+			const FProperty* Property = nullptr;
+			TWeakObjectPtr<const UScriptStruct> ColumnType;
+		};
+
+		// Exported because it's used in the template to bind to a member variable.
+		TEDSQUERYSTACK_API QueryHandle CreateQuery(const UScriptStruct* ColumnType, Queries::FConditions&& Conditions);
+		void RunSearch(const Searching::FSearchContext& Context, const FSearchQuery& Searcher);
+		void ResetSearch();
+		
+		TArray<FSearchQuery> Searchers;
+
+		TUniquePtr<Searching::FSearchContext> SearchContext;
+		FRowHandleArray Rows;
+		ICoreProvider& Storage;
+		FTimespan TotalSearchTime; // actual time spend processing.
+		RevisionId Revision = UninitializedRevisionId;
+		ESyncActions SyncAction;
+	};
+
+	// Implementations
+	template<TDataColumnType Column>
+	void FColumnsSearchNode::RegisterColumn(Queries::FConditions Conditions)
+	{
+		RegisterColumn(Column::StaticStruct(), MoveTemp(Conditions));
+	}
+
+	template<auto MemberVariable> requires Searching::SearchableMemberColumn<MemberVariable>
+	void FColumnsSearchNode::RegisterColumn(Queries::FConditions Conditions)
+	{
+		using Class = typename Searching::MemberPointerTraits<decltype(MemberVariable)>::ClassType;
+		
+		Searchers.Add(FSearchQuery{
+			.Query = CreateQuery(Class::StaticStruct(), MoveTemp(Conditions)),
+			.SearchFunction = Searching::CreateSearchFunction<MemberVariable>(),
+			.ColumnType = Class::StaticStruct() });
+	}
+} // namespace UE::Editor::DataStorage::QueryStack
diff --git a/Engine/Plugins/Experimental/EditorDataStorageFeatures/Source/TedsQueryStack/Public/TedsQuerySearchNode.h b/Engine/Plugins/Experimental/EditorDataStorageFeatures/Source/TedsQueryStack/Public/TedsQuerySearchNode.h
new file mode 100644
--- /dev/null
+++ b/Engine/Plugins/Experimental/EditorDataStorageFeatures/Source/TedsQueryStack/Public/TedsQuerySearchNode.h
@@ -0,0 +1,83 @@
+// Copyright Epic Games, Inc. All Rights Reserved.
+
+#pragma once
+
+#include "Containers/StringFwd.h"
+#include "DataStorage/CommonTypes.h"
+#include "DataStorage/Queries/Conditions.h"
+#include "Elements/Framework/TypedElementRowHandleArray.h"
+#include "Misc/Timespan.h"
+#include "TedsQueryStackInterfaces.h"
+#include "TedsRowQueryResultsNode.h"
+#include "Templates/UniquePtr.h"
+
+#include "Searching/SearchUtils.h"
+
+class FProperty;
+class UScriptStruct;
+namespace UE::Editor::DataStorage
+{
+	class ICoreProvider;
+} // namespace UE::Editor::DataStorage
+
+namespace UE::Editor::DataStorage::QueryStack
+{
+	namespace Searching
+	{
+		struct FSearchContext;
+	}
+
+	/**
+	 * Query Stack node that takes a query and searches all selected columns that have searchable properties.
+	 * If the search is disabled, either by calling "ClearSearch" or providing an empty string, this node acts
+	 * as a pass through and behaves the same way as FRowQueryResultsNode.
+	 */
+	class FQuerySearchNode final : public FRowQueryResultsNode
+	{
+	public:
+		TEDSQUERYSTACK_API explicit FQuerySearchNode(ICoreProvider& Storage, TSharedRef<IQueryNode> Query);
+		TEDSQUERYSTACK_API explicit FQuerySearchNode(ICoreProvider& Storage, TSharedRef<IQueryNode> Query, ESyncActions SyncAction);
+		virtual ~FQuerySearchNode() override = default;
+
+		/**
+		 * Starts searching for the provided string in all searchable columns that were selected in the query.
+		 * Results may not be collected and can take several frames to arrive.
+		 */
+		TEDSQUERYSTACK_API void StartSearch(FString SearchString);
+		/** Resets a search and makes this node behave as a pass-through, resulting in similar behavior to FRowQueryResultsNode. */
+		TEDSQUERYSTACK_API void ClearSearch();
+		/** Returns true if a search is still being processed. */
+		TEDSQUERYSTACK_API bool IsSearching() const;
+
+		/**
+		 * The total time spent on performing the search.
+		 * This exclusively tracks the time spend on performing one or more steps in the search.
+		 */
+		TEDSQUERYSTACK_API FTimespan GetQuerySearchTime() const;
+
+	protected:
+		static constexpr uint32 InlineColumnCount = 32;
+
+		TEDSQUERYSTACK_API virtual void RefreshInternal(FTimespan AllottedTime, FRowHandleArray& TargetRows) override;
+		TEDSQUERYSTACK_API virtual void QueryUpdated();
+
+	private:
+		struct FSearchQuery
+		{
+			Searching::StringSearchFunction SearchFunction = nullptr;
+			const FProperty* Property = nullptr;
+			TWeakObjectPtr<const UScriptStruct> ColumnType;
+			int32 ColumnSize = 0;
+		};
+
+		void UpdateSearchersList();
+		void RunSearch(const Searching::FSearchContext& Context, FRowHandleArray& TargetRows);
+		void ResetSearch();
+		
+		TArray<FSearchQuery> Searchers;
+
+		TUniquePtr<Searching::FSearchContext> SearchContext;
+		FTimespan TotalSearchTime; // actual time spent processing.
+		bool bHasPendingSearch = false;
+	};
+} // namespace UE::Editor::DataStorage::QueryStack
diff --git a/Engine/Source/Runtime/TypedElementFramework/Private/Elements/Framework/TypedElementQueryBuilder.cpp b/Engine/Source/Runtime/TypedElementFramework/Private/Elements/Framework/TypedElementQueryBuilder.cpp
--- a/Engine/Source/Runtime/TypedElementFramework/Private/Elements/Framework/TypedElementQueryBuilder.cpp
+++ b/Engine/Source/Runtime/TypedElementFramework/Private/Elements/Framework/TypedElementQueryBuilder.cpp
@@ -521,15 +521,15 @@ namespace UE::Editor::DataStorage::Queries
 		return FDependency{ &Query };
 	}
 
-	Select& Select::AccessesHierarchy(const FName& HierachyName)
+	Select& Select::AccessesHierarchy(const FName& HierarchyName)
 	{
-		Query.Hierarchy = HierachyName;
+		Query.Hierarchy = HierarchyName;
 		return *this;
 	}
 
-	FQueryConditionQuery Select::Where(const Queries::FConditions& Condition)
+	FQueryConditionQuery Select::Where(Queries::FConditions Condition)
 	{
-		Query.Conditions.Emplace(Condition);
+		Query.Conditions.Emplace(MoveTemp(Condition));
 		return FQueryConditionQuery(&Query);
 	}
 
diff --git a/Engine/Source/Runtime/TypedElementFramework/Public/Elements/Framework/TypedElementQueryBuilder.h b/Engine/Source/Runtime/TypedElementFramework/Public/Elements/Framework/TypedElementQueryBuilder.h
--- a/Engine/Source/Runtime/TypedElementFramework/Public/Elements/Framework/TypedElementQueryBuilder.h
+++ b/Engine/Source/Runtime/TypedElementFramework/Public/Elements/Framework/TypedElementQueryBuilder.h
@@ -429,7 +429,7 @@ namespace UE::Editor::DataStorage::Queries
 		Select& ReadWrite();
 
 		// Query Condition API
-		TYPEDELEMENTFRAMEWORK_API FQueryConditionQuery Where(const Queries::FConditions& Condition);
+		TYPEDELEMENTFRAMEWORK_API FQueryConditionQuery Where(Queries::FConditions Condition);
 
 		TYPEDELEMENTFRAMEWORK_API FQueryDescription&& Compile();
 		TYPEDELEMENTFRAMEWORK_API FSimpleQuery Where();
