FRowFilterNode : fix revision update

A call to ForceRefresh() when search text changes will reset the CachedParentRevisionID to UninitializedRevisionId but in the next Update it will get set back to ParentRowNode->GetRevision().  Therefore, if the ParentRowNode has not changed revision then downstream child nodes of this node would not see any revision change related to the search text change prior to this fix.

#Synced-CL 48876874

diff --git a/Engine/Plugins/Experimental/EditorDataStorageFeatures/Source/TedsQueryStack/Private/TedsRowFilterNode.cpp b/Engine/Plugins/Experimental/EditorDataStorageFeatures/Source/TedsQueryStack/Private/TedsRowFilterNode.cpp
--- a/Engine/Plugins/Experimental/EditorDataStorageFeatures/Source/TedsQueryStack/Private/TedsRowFilterNode.cpp
+++ b/Engine/Plugins/Experimental/EditorDataStorageFeatures/Source/TedsQueryStack/Private/TedsRowFilterNode.cpp
@@ -16,7 +16,7 @@ namespace UE::Editor::DataStorage::QueryStack
 
 	IRowNode::RevisionId FRowFilterNode::GetRevision() const
 	{
-		return CachedParentRevisionID;
+		return Revision;
 	}
 
 	void FRowFilterNode::Update(FTimespan AllottedTime)
@@ -26,6 +26,7 @@ namespace UE::Editor::DataStorage::QueryStack
 			Rows.Reset();
 			Storage->FilterRowsBy(Rows, ParentRowNode->GetRows(), Options, Filter);
 
+			++Revision;
 			CachedParentRevisionID = ParentRowNode->GetRevision();
 		}
 	}
diff --git a/Engine/Plugins/Experimental/EditorDataStorageFeatures/Source/TedsQueryStack/Public/TedsRowFilterNode.h b/Engine/Plugins/Experimental/EditorDataStorageFeatures/Source/TedsQueryStack/Public/TedsRowFilterNode.h
--- a/Engine/Plugins/Experimental/EditorDataStorageFeatures/Source/TedsQueryStack/Public/TedsRowFilterNode.h
+++ b/Engine/Plugins/Experimental/EditorDataStorageFeatures/Source/TedsQueryStack/Public/TedsRowFilterNode.h
@@ -85,6 +85,7 @@ namespace UE::Editor::DataStorage::QueryStack
 		FRowHandleArray Rows;
 		TSharedPtr<IRowNode> ParentRowNode;
 		RevisionId CachedParentRevisionID = UninitializedRevisionId;
+		RevisionId Revision = UninitializedRevisionId;
 		ICoreProvider::EFilterOptions Options;
 	};
 }
