Editor: Pipe AllowInvisibleItemSelection through STedsCompositeHierarchyViewer, and add ExpandWithParents.

#Synced-CL 48918043

diff --git a/Engine/Plugins/Experimental/EditorDataStorageFeatures/Source/TedsTableViewer/Private/Widgets/Composite/STedsCompositeHierarchyViewer.cpp b/Engine/Plugins/Experimental/EditorDataStorageFeatures/Source/TedsTableViewer/Private/Widgets/Composite/STedsCompositeHierarchyViewer.cpp
--- a/Engine/Plugins/Experimental/EditorDataStorageFeatures/Source/TedsTableViewer/Private/Widgets/Composite/STedsCompositeHierarchyViewer.cpp
+++ b/Engine/Plugins/Experimental/EditorDataStorageFeatures/Source/TedsTableViewer/Private/Widgets/Composite/STedsCompositeHierarchyViewer.cpp
@@ -244,4 +244,10 @@ namespace UE::Editor::DataStorage
 	{
 		HierarchyViewer->CollapseAll();
 	}
+
+	void STedsCompositeHierarchyViewer::ExpandWithParents(RowHandle Row)
+	{
+		HierarchyViewer->ExpandWithParents(Row);
+	}
+
 } // namespace UE::Editor::DataStorage
diff --git a/Engine/Plugins/Experimental/EditorDataStorageFeatures/Source/TedsTableViewer/Private/Widgets/STedsHierarchyViewer.cpp b/Engine/Plugins/Experimental/EditorDataStorageFeatures/Source/TedsTableViewer/Private/Widgets/STedsHierarchyViewer.cpp
--- a/Engine/Plugins/Experimental/EditorDataStorageFeatures/Source/TedsTableViewer/Private/Widgets/STedsHierarchyViewer.cpp
+++ b/Engine/Plugins/Experimental/EditorDataStorageFeatures/Source/TedsTableViewer/Private/Widgets/STedsHierarchyViewer.cpp
@@ -92,7 +92,8 @@ namespace UE::Editor::DataStorage
 			.OnGenerateRow(this, &SHierarchyViewer::MakeTableRowWidget)
 			.OnExpansionChanged(this, &SHierarchyViewer::OnExpansionChanged)
 			.OnSelectionChanged(this, &SHierarchyViewer::OnListSelectionChanged)
-			.SelectionMode(InArgs._ListSelectionMode);
+			.SelectionMode(InArgs._ListSelectionMode)
+			.AllowInvisibleItemSelection(InArgs._AllowInvisibleItemSelection);
 
 		CreateInternalWidget();
 		
@@ -418,6 +419,19 @@ namespace UE::Editor::DataStorage
 		}
 	}
 
+	void SHierarchyViewer::ExpandWithParents(RowHandle Row)
+	{
+		FTedsRowHandle TedsRowHandle{ .RowHandle = Row };
+		int32 SanityCheckCount = 0;
+		constexpr int32 SANITY_CHECK_LIMIT = 1000;
+
+		while (TedsRowHandle.IsValid() && ensure(SanityCheckCount++ < SANITY_CHECK_LIMIT))
+		{
+			TreeView->SetItemExpansion(TedsRowHandle, true);
+			TedsRowHandle = GetParentRow(TedsRowHandle);
+		}
+	}
+
 	bool SHierarchyViewer::IsItemVisible(TableViewerItemPtr InItem) const
 	{
 		return TreeView->IsItemVisible(InItem);
diff --git a/Engine/Plugins/Experimental/EditorDataStorageFeatures/Source/TedsTableViewer/Private/Widgets/STedsTreeView.cpp b/Engine/Plugins/Experimental/EditorDataStorageFeatures/Source/TedsTableViewer/Private/Widgets/STedsTreeView.cpp
--- a/Engine/Plugins/Experimental/EditorDataStorageFeatures/Source/TedsTableViewer/Private/Widgets/STedsTreeView.cpp
+++ b/Engine/Plugins/Experimental/EditorDataStorageFeatures/Source/TedsTableViewer/Private/Widgets/STedsTreeView.cpp
@@ -29,7 +29,8 @@ namespace UE::Editor::DataStorage
 			.OnContextMenuOpening(Binder.BindEvent(&FWidgetContextMenuColumn::OnContextMenuOpening))
 			.OnItemScrolledIntoView(Binder.BindEvent(&FWidgetRowScrolledIntoView::OnItemScrolledIntoView))
 			.OnMouseButtonDoubleClick(Binder.BindEvent(&FWidgetDoubleClickedColumn::OnMouseButtonDoubleClick))
-			.OnGetChildren(this, &STedsTreeView::GetChildren_Internal));
+			.OnGetChildren(this, &STedsTreeView::GetChildren_Internal)
+			.AllowInvisibleItemSelection(InArgs._AllowInvisibleItemSelection));
 	}
 
 	void STedsTreeView::Tick( const FGeometry& AllottedGeometry, const double InCurrentTime, const float InDeltaTime )
diff --git a/Engine/Plugins/Experimental/EditorDataStorageFeatures/Source/TedsTableViewer/Private/Widgets/STedsTreeView.h b/Engine/Plugins/Experimental/EditorDataStorageFeatures/Source/TedsTableViewer/Private/Widgets/STedsTreeView.h
--- a/Engine/Plugins/Experimental/EditorDataStorageFeatures/Source/TedsTableViewer/Private/Widgets/STedsTreeView.h
+++ b/Engine/Plugins/Experimental/EditorDataStorageFeatures/Source/TedsTableViewer/Private/Widgets/STedsTreeView.h
@@ -39,6 +39,9 @@ namespace UE::Editor::DataStorage
 
 			/** The header row widget to use */
 			SLATE_ARGUMENT( TSharedPtr<SHeaderRow>, HeaderRow )
+
+			/** Allows collapsed or temporarily hidden items to still be selected */
+			SLATE_ARGUMENT(bool, AllowInvisibleItemSelection);
 			
 		SLATE_END_ARGS()
 		
diff --git a/Engine/Plugins/Experimental/EditorDataStorageFeatures/Source/TedsTableViewer/Public/Widgets/Composite/STedsCompositeHierarchyViewer.h b/Engine/Plugins/Experimental/EditorDataStorageFeatures/Source/TedsTableViewer/Public/Widgets/Composite/STedsCompositeHierarchyViewer.h
--- a/Engine/Plugins/Experimental/EditorDataStorageFeatures/Source/TedsTableViewer/Public/Widgets/Composite/STedsCompositeHierarchyViewer.h
+++ b/Engine/Plugins/Experimental/EditorDataStorageFeatures/Source/TedsTableViewer/Public/Widgets/Composite/STedsCompositeHierarchyViewer.h
@@ -120,6 +120,9 @@ namespace UE::Editor::DataStorage
 		// Recursively Collapse all rows currently present in the viewer
 		TEDSTABLEVIEWER_API void CollapseAll() const;
 
+		// Expand this row, and all its parents
+		TEDSTABLEVIEWER_API void ExpandWithParents(RowHandle Row);
+
 	private:
 		// Row Node passed to the SearchBox that receives the searched result
 		TSharedPtr<QueryStack::IRowNode> SearchNode;
diff --git a/Engine/Plugins/Experimental/EditorDataStorageFeatures/Source/TedsTableViewer/Public/Widgets/STedsHierarchyViewer.h b/Engine/Plugins/Experimental/EditorDataStorageFeatures/Source/TedsTableViewer/Public/Widgets/STedsHierarchyViewer.h
--- a/Engine/Plugins/Experimental/EditorDataStorageFeatures/Source/TedsTableViewer/Public/Widgets/STedsHierarchyViewer.h
+++ b/Engine/Plugins/Experimental/EditorDataStorageFeatures/Source/TedsTableViewer/Public/Widgets/STedsHierarchyViewer.h
@@ -58,6 +58,7 @@ namespace UE::Editor::DataStorage
 		SLATE_BEGIN_ARGS(SHierarchyViewer)
 			: _CanSelectGeneratedColumn(true)
 			, _ListSelectionMode(ESelectionMode::Type::Single)
+			, _AllowInvisibleItemSelection(false)
 		{
 			
 		}
@@ -95,6 +96,9 @@ namespace UE::Editor::DataStorage
 		// The primary column in the hierarchy viewer is the one where the expansion row shows up
 		SLATE_ARGUMENT(TWeakObjectPtr<const UScriptStruct>, PrimaryColumn)
 
+		// Allows collapsed or temporarily hidden items to still be selected
+		SLATE_ARGUMENT(bool, AllowInvisibleItemSelection);
+
 		SLATE_END_ARGS()
 
 	public:
@@ -145,6 +149,9 @@ namespace UE::Editor::DataStorage
 		// Recursively Collapse all rows currently present in the viewer
 		TEDSTABLEVIEWER_API void CollapseAll() const;
 
+		// Expand this row and all its parents
+		TEDSTABLEVIEWER_API void ExpandWithParents(RowHandle Row);
+
 	protected:
 		
 		RowHandle AddOrFindTableViewerUIRow(TableViewerItemPtr InItem) const;
