STedsCompositeHierarchyViewer: Fix missing parameter forward from 48929987.

#Synced-CL 48929987

diff --git a/Engine/Plugins/Experimental/EditorDataStorageFeatures/Source/TedsTableViewer/Private/Widgets/Composite/STedsCompositeHierarchyViewer.cpp b/Engine/Plugins/Experimental/EditorDataStorageFeatures/Source/TedsTableViewer/Private/Widgets/Composite/STedsCompositeHierarchyViewer.cpp
--- a/Engine/Plugins/Experimental/EditorDataStorageFeatures/Source/TedsTableViewer/Private/Widgets/Composite/STedsCompositeHierarchyViewer.cpp
+++ b/Engine/Plugins/Experimental/EditorDataStorageFeatures/Source/TedsTableViewer/Private/Widgets/Composite/STedsCompositeHierarchyViewer.cpp
@@ -114,7 +114,8 @@ namespace UE::Editor::DataStorage
 			.ColumnMetaData(HierarchyViewerArgs._ColumnMetaData)
 			.ListSelectionMode(HierarchyViewerArgs._ListSelectionMode)
 			.OnSelectionChanged(HierarchyViewerArgs._OnSelectionChanged)
-			.PrimaryColumn(HierarchyViewerArgs._PrimaryColumn);
+			.PrimaryColumn(HierarchyViewerArgs._PrimaryColumn)
+			.AllowInvisibleItemSelection(HierarchyViewerArgs._AllowInvisibleItemSelection);
 		
 		ContentWidget->AddSlot()
 		[
