[TEDS] Updates to the STedsSearchBbox to support the new query search node.

This updates uses the new query search node if the input to the search box is a query node. If a row node is provided it will continue to use the filter approach. The search box was integrated into the TEDS Debugger to test it at scale.

#Synced-CL 48813241

diff --git a/Engine/Plugins/Experimental/EditorDataStorageFeatures/Source/TedsDebugger/Private/Widgets/QueryEditor/TedsQueryEditorResultsView.cpp b/Engine/Plugins/Experimental/EditorDataStorageFeatures/Source/TedsDebugger/Private/Widgets/QueryEditor/TedsQueryEditorResultsView.cpp
--- a/Engine/Plugins/Experimental/EditorDataStorageFeatures/Source/TedsDebugger/Private/Widgets/QueryEditor/TedsQueryEditorResultsView.cpp
+++ b/Engine/Plugins/Experimental/EditorDataStorageFeatures/Source/TedsDebugger/Private/Widgets/QueryEditor/TedsQueryEditorResultsView.cpp
@@ -23,6 +23,7 @@
 #include "Widgets/STedsTableViewer.h"
 #include "Widgets/SRowDetails.h"
 #include "Widgets/Text/STextBlock.h"
+#include "Widgets/STedsSearchBox.h"
 
 #define LOCTEXT_NAMESPACE "TedsDebuggerModule"
 
@@ -58,12 +59,6 @@ namespace UE::Editor::DataStorage::Debug::QueryEditor
 		CreateRowHandleColumn();
 	
 		QueryNode = MakeShared<FQueryNode>(Model->GetTedsInterface());
-
-		// Use ForceRefreshOnUpdate to re-run the query every frame since we don't have a good change detection mechanism for hierarchies right now
-		// and just using observers can cause the visual hierarchy to go out of sync
-		RowQueryStack = MakeShared<FRowQueryResultsNode>(Model->GetTedsInterface(), QueryNode,
-			FRowQueryResultsNode::ESyncActions::ForceRefreshOnUpdate);
-		//RowQueryStack = MakeShared<FRowMonitorNode>(Model->GetTedsInterface(), RowQueryStack, QueryNode);
 		
 		ChildSlot
 		[
@@ -74,18 +69,30 @@ namespace UE::Editor::DataStorage::Debug::QueryEditor
 				+SSplitter::Slot()
 				.Value(0.5f)
 				[
-					SAssignNew(TableViewerSlot, SBorder)
+					SNew(SVerticalBox)
+					+SVerticalBox::Slot()
+					.AutoHeight()
+					[
+						SNew(STedsSearchBox)
+						.InStorage(&Model->GetTedsInterface())
+						.InSearchableQueryNode(QueryNode)
+						.InSearchableQueryFlags(STedsSearchBox::ESearchableQueryFlags::RefreshOnUpdate)
+						.OutSearchNode(&RowQueryStack)
+					]
+					+SVerticalBox::Slot()
+					.FillContentHeight(1.0f)
 					[
-						CreateTableViewer()
+						SAssignNew(TableViewerSlot, SBorder)
+						[
+							CreateTableViewer()
+						]
 					]
-					
 				]
 				+SSplitter::Slot()
 				.Value(0.5f)
 				[
 					SAssignNew(RowDetailsWidget, SRowDetails)
 				]
-		
 			]
 			+SVerticalBox::Slot()
 			.AutoHeight()
@@ -108,6 +115,8 @@ namespace UE::Editor::DataStorage::Debug::QueryEditor
 			]
 		];
 
+		RowQueryStack = MakeShared<FRowMonitorNode>(Model->GetTedsInterface(), RowQueryStack, QueryNode);
+
 		if(RowHandleColumn)
 		{
 			TableViewer->AddCustomRowWidget(RowHandleColumn.ToSharedRef());
diff --git a/Engine/Plugins/Experimental/EditorDataStorageFeatures/Source/TedsTableViewer/Private/Widgets/STedsSearchBox.cpp b/Engine/Plugins/Experimental/EditorDataStorageFeatures/Source/TedsTableViewer/Private/Widgets/STedsSearchBox.cpp
--- a/Engine/Plugins/Experimental/EditorDataStorageFeatures/Source/TedsTableViewer/Private/Widgets/STedsSearchBox.cpp
+++ b/Engine/Plugins/Experimental/EditorDataStorageFeatures/Source/TedsTableViewer/Private/Widgets/STedsSearchBox.cpp
@@ -21,8 +21,9 @@ namespace UE::Editor::DataStorage
 	{
 		using namespace UE::Editor::DataStorage::Queries;
 
-		ICoreProvider* Storage = GetMutableDataStorageFeature<ICoreProvider>(UE::Editor::DataStorage::StorageFeatureName);
-
+		ICoreProvider* Storage = InArgs._InStorage 
+			? InArgs._InStorage
+			: GetMutableDataStorageFeature<ICoreProvider>(UE::Editor::DataStorage::StorageFeatureName);
 		if(!ensureMsgf(Storage, TEXT("Cannot create a TEDS search widget before TEDS is initialized")))
 		{
 			ChildSlot
@@ -33,48 +34,89 @@ namespace UE::Editor::DataStorage
 			return;
 		}
 
-		if(InArgs._OutSearchNode && InArgs._InSearchableRowNode)
+		if(InArgs._OutSearchNode)
 		{
-			OutSearchNode = MakeShared<QueryStack::FRowFilterNode>(Storage, InArgs._InSearchableRowNode,
-				[this](TQueryContext<SingleRowInfo> Context, const FTypedElementLabelColumn& LabelColumn)
-				{
-					if(SearchText.IsEmpty())
+			if (InArgs._InSearchableQueryNode)
+			{
+				TSharedRef<QueryStack::FQuerySearchNode> NewSearchNode = MakeShared<QueryStack::FQuerySearchNode>(
+					*Storage, 
+					InArgs._InSearchableQueryNode.ToSharedRef(),
+					InArgs._InSearchableQueryFlags);
+				SearchNode.Emplace<TSharedRef<QueryStack::FQuerySearchNode>>(NewSearchNode);
+				*InArgs._OutSearchNode = MoveTemp(NewSearchNode);
+			}
+			else if (InArgs._InSearchableRowNode)
+			{
+				TSharedPtr<QueryStack::FRowFilterNode> NewSearchNode = MakeShared<QueryStack::FRowFilterNode>(Storage, InArgs._InSearchableRowNode,
+					[Widget = AsWeak()](TQueryContext<SingleRowInfo> Context, const FTypedElementLabelColumn& LabelColumn)
 					{
-						return true;
-					}
-
-					TArray<FString> SearchTokens;
-					SearchText.ToString().ParseIntoArrayWS(SearchTokens);
-					const FString Label = LabelColumn.Label;
-
-					for (const FString& SearchToken : SearchTokens)
-					{
-						if (Label.IsEmpty() || !Label.Contains(SearchToken, ESearchCase::Type::IgnoreCase))
+						if (TSharedPtr<SWidget> WidgetPtr = Widget.Pin())
 						{
-							return false;
+							STedsSearchBox& SearchBox = *static_cast<STedsSearchBox*>(WidgetPtr.Get());
+							if (!SearchBox.SearchText.IsEmpty())
+							{
+								TArray<FString> SearchTokens;
+								SearchBox.SearchText.ToString().ParseIntoArrayWS(SearchTokens);
+								const FString Label = LabelColumn.Label;
+
+								for (const FString& SearchToken : SearchTokens)
+								{
+									if (Label.IsEmpty() || !Label.Contains(SearchToken, ESearchCase::Type::IgnoreCase))
+									{
+										return false;
+									}
+								}
+							}
 						}
-					}
-					return true;
-				});
-				
-			*InArgs._OutSearchNode = OutSearchNode;
+						return true;
+					});
+
+				SearchNode.Emplace<TSharedPtr<QueryStack::FRowFilterNode>>(NewSearchNode);
+				*InArgs._OutSearchNode = MoveTemp(NewSearchNode);
+			}
 		}
 		
 		ChildSlot
 		[
 			SNew(SSearchBox)
 				.OnTextChanged(this, &STedsSearchBox::OnTextChanged)
+				.DelayChangeNotificationsWhileTyping(false)
 		];
 	}
 
 	void STedsSearchBox::OnTextChanged(const FText& InSearchText)
 	{
 		SearchText = InSearchText;
-		if(OutSearchNode)
+
+		struct FVisitor
 		{
-			OutSearchNode->ForceRefresh();
-		}
-		
+			void operator()(const TSharedPtr<QueryStack::FRowFilterNode>& Node)
+			{
+				if (Node)
+				{
+					Node->ForceRefresh();
+				}
+			}
+			
+			void operator()(const TSharedPtr<QueryStack::FQuerySearchNode>& Node)
+			{
+				if (Node)
+				{
+					if (This->SearchText.IsEmpty())
+					{
+						Node->ClearSearch();
+					}
+					else
+					{
+						Node->StartSearch(This->SearchText.ToString());
+					}
+				}
+			}
+
+			STedsSearchBox* This;
+		};
+
+		Visit(FVisitor{ .This = this }, SearchNode);
 	}
 	
 } // namespace UE::Editor::DataStorage
diff --git a/Engine/Plugins/Experimental/EditorDataStorageFeatures/Source/TedsTableViewer/Public/Widgets/STedsSearchBox.h b/Engine/Plugins/Experimental/EditorDataStorageFeatures/Source/TedsTableViewer/Public/Widgets/STedsSearchBox.h
--- a/Engine/Plugins/Experimental/EditorDataStorageFeatures/Source/TedsTableViewer/Public/Widgets/STedsSearchBox.h
+++ b/Engine/Plugins/Experimental/EditorDataStorageFeatures/Source/TedsTableViewer/Public/Widgets/STedsSearchBox.h
@@ -1,9 +1,10 @@
 // Copyright Epic Games, Inc. All Rights Reserved.
 #pragma once
 
+#include "Misc/Variant.h"
 #include "TedsQueryStackInterfaces.h"
 #include "TedsRowFilterNode.h"
-#include "TedsRowQueryResultsNode.h"
+#include "TedsQuerySearchNode.h"
 #include "Widgets/SCompoundWidget.h"
 
 namespace UE::Editor::DataStorage
@@ -24,26 +25,49 @@ namespace UE::Editor::DataStorage
 	class STedsSearchBox : public SCompoundWidget
 	{
 	public:
-		
+		using ESearchableQueryFlags = QueryStack::FQuerySearchNode::ESyncActions;
+
 		SLATE_BEGIN_ARGS(STedsSearchBox) {}
 
-		// The node whose rows we'll perform a search on
+		/** The instance of the Editor Data Storage used to perform the search. */
+		SLATE_ARGUMENT(ICoreProvider*, InStorage);
+
+		/** 
+		 * The node provider to search through.
+		 * If InSearchableQueryNode is also set, this node will be ignored.
+		 */
 		SLATE_ARGUMENT(TSharedPtr<QueryStack::IRowNode>, InSearchableRowNode)
 
-		// The search results node we create and receive back
+		/**
+		 * Query node used for searching.If this and InSearchableRowNode are set, the query node will take priority
+		 * and InSearchableRowNode will be ignored. If there's no active search this node will execute the query without
+		 * filtering and collect the results.
+		 */
+		SLATE_ARGUMENT(TSharedPtr<QueryStack::IQueryNode>, InSearchableQueryNode)
+
+		/** If a searchable query is set, this sets how updates happen that run the query. */
+		SLATE_ARGUMENT(ESearchableQueryFlags, InSearchableQueryFlags)
+
+		/** 
+		 * The result node containing a list of rows. This can be used as an input to other
+		 * nodes in the Query Stack.
+		 * If a search is active this will result in a list of rows filtered to the search criteria.
+		 * If there's no search than this will return the results of running the query if a query node was
+		 * provided. If a row node was provided the original list of rows will be returned.
+		 */
 		SLATE_ARGUMENT(TSharedPtr<QueryStack::IRowNode>*, OutSearchNode)
 			
 		SLATE_END_ARGS()
-
-	public:
-		
+	
 		TEDSTABLEVIEWER_API void Construct(const FArguments& InArgs);
 
 	private:
 		void OnTextChanged(const FText& InSearchText);
 
-		/** Internal SharedPtr of the OutSearchNode so that when text is updated, it can ForceRefresh the node */
-		TSharedPtr<QueryStack::FRowFilterNode> OutSearchNode;
+		TVariant<
+			TSharedPtr<QueryStack::FRowFilterNode>,
+			TSharedRef<QueryStack::FQuerySearchNode>
+		> SearchNode;
 		FText SearchText;
 	};
 }
\ No newline at end of file
