[Prefab Editor] Revert returning the node Revision in the RowFilterNode as a band-aid for the Prefab Editor

#Synced-CL 48986818

diff --git a/Engine/Plugins/Experimental/EditorDataStorageFeatures/Source/TedsQueryStack/Private/TedsRowFilterNode.cpp b/Engine/Plugins/Experimental/EditorDataStorageFeatures/Source/TedsQueryStack/Private/TedsRowFilterNode.cpp
--- a/Engine/Plugins/Experimental/EditorDataStorageFeatures/Source/TedsQueryStack/Private/TedsRowFilterNode.cpp
+++ b/Engine/Plugins/Experimental/EditorDataStorageFeatures/Source/TedsQueryStack/Private/TedsRowFilterNode.cpp
@@ -16,7 +16,7 @@ namespace UE::Editor::DataStorage::QueryStack
 
 	IRowNode::RevisionId FRowFilterNode::GetRevision() const
 	{
-		return Revision;
+		return CachedParentRevisionID;
 	}
 
 	void FRowFilterNode::Update(FTimespan AllottedTime)
