STedsSearchBox: Remove Storage argument (that was not init properly) and grab storage from the global instance instead

#Synced-CL 48986405

diff --git a/Engine/Plugins/Experimental/EditorDataStorageFeatures/Source/TedsDebugger/Private/Widgets/QueryEditor/TedsQueryEditorResultsView.cpp b/Engine/Plugins/Experimental/EditorDataStorageFeatures/Source/TedsDebugger/Private/Widgets/QueryEditor/TedsQueryEditorResultsView.cpp
--- a/Engine/Plugins/Experimental/EditorDataStorageFeatures/Source/TedsDebugger/Private/Widgets/QueryEditor/TedsQueryEditorResultsView.cpp
+++ b/Engine/Plugins/Experimental/EditorDataStorageFeatures/Source/TedsDebugger/Private/Widgets/QueryEditor/TedsQueryEditorResultsView.cpp
@@ -74,7 +74,6 @@ namespace UE::Editor::DataStorage::Debug::QueryEditor
 					.AutoHeight()
 					[
 						SNew(STedsSearchBox)
-						.InStorage(&Model->GetTedsInterface())
 						.InSearchableQueryNode(QueryNode)
 						.InSearchableQueryFlags(STedsSearchBox::ESearchableQueryFlags::RefreshOnUpdate)
 						.OutSearchNode(&RowQueryStack)
diff --git a/Engine/Plugins/Experimental/EditorDataStorageFeatures/Source/TedsTableViewer/Private/Widgets/STedsSearchBox.cpp b/Engine/Plugins/Experimental/EditorDataStorageFeatures/Source/TedsTableViewer/Private/Widgets/STedsSearchBox.cpp
--- a/Engine/Plugins/Experimental/EditorDataStorageFeatures/Source/TedsTableViewer/Private/Widgets/STedsSearchBox.cpp
+++ b/Engine/Plugins/Experimental/EditorDataStorageFeatures/Source/TedsTableViewer/Private/Widgets/STedsSearchBox.cpp
@@ -21,9 +21,7 @@ namespace UE::Editor::DataStorage
 	{
 		using namespace UE::Editor::DataStorage::Queries;
 
-		ICoreProvider* Storage = InArgs._InStorage 
-			? InArgs._InStorage
-			: GetMutableDataStorageFeature<ICoreProvider>(UE::Editor::DataStorage::StorageFeatureName);
+		ICoreProvider* Storage = GetMutableDataStorageFeature<ICoreProvider>(UE::Editor::DataStorage::StorageFeatureName);
 		if(!ensureMsgf(Storage, TEXT("Cannot create a TEDS search widget before TEDS is initialized")))
 		{
 			ChildSlot
diff --git a/Engine/Plugins/Experimental/EditorDataStorageFeatures/Source/TedsTableViewer/Public/Widgets/STedsSearchBox.h b/Engine/Plugins/Experimental/EditorDataStorageFeatures/Source/TedsTableViewer/Public/Widgets/STedsSearchBox.h
--- a/Engine/Plugins/Experimental/EditorDataStorageFeatures/Source/TedsTableViewer/Public/Widgets/STedsSearchBox.h
+++ b/Engine/Plugins/Experimental/EditorDataStorageFeatures/Source/TedsTableViewer/Public/Widgets/STedsSearchBox.h
@@ -29,9 +29,6 @@ namespace UE::Editor::DataStorage
 
 		SLATE_BEGIN_ARGS(STedsSearchBox) {}
 
-		/** The instance of the Editor Data Storage used to perform the search. */
-		SLATE_ARGUMENT(ICoreProvider*, InStorage);
-
 		/** 
 		 * The node provider to search through.
 		 * If InSearchableQueryNode is also set, this node will be ignored.
