Adding back in the correct revision

#Synced-CL 48986818

diff --git a/Engine/Plugins/Experimental/EditorDataStorageFeatures/Source/TedsQueryStack/Private/TedsRowFilterNode.cpp b/Engine/Plugins/Experimental/EditorDataStorageFeatures/Source/TedsQueryStack/Private/TedsRowFilterNode.cpp
--- a/Engine/Plugins/Experimental/EditorDataStorageFeatures/Source/TedsQueryStack/Private/TedsRowFilterNode.cpp
+++ b/Engine/Plugins/Experimental/EditorDataStorageFeatures/Source/TedsQueryStack/Private/TedsRowFilterNode.cpp
@@ -16,7 +16,7 @@ namespace UE::Editor::DataStorage::QueryStack
 
 	IRowNode::RevisionId FRowFilterNode::GetRevision() const
 	{
-		return CachedParentRevisionID;
+		return Revision;
 	}
 
 	void FRowFilterNode::Update(FTimespan AllottedTime)
