WidgetDropHandler now selects placed actors without teds rows in same tick as drop operation.

diff --git a/Engine/Plugins/Experimental/EditorDataStorageFeatures/Source/TedsOperations/Private/DragAndDrop/Widgets/WidgetDropHandler.cpp b/Engine/Plugins/Experimental/EditorDataStorageFeatures/Source/TedsOperations/Private/DragAndDrop/Widgets/WidgetDropHandler.cpp
--- a/Engine/Plugins/Experimental/EditorDataStorageFeatures/Source/TedsOperations/Private/DragAndDrop/Widgets/WidgetDropHandler.cpp
+++ b/Engine/Plugins/Experimental/EditorDataStorageFeatures/Source/TedsOperations/Private/DragAndDrop/Widgets/WidgetDropHandler.cpp
@@ -7,7 +7,8 @@
 #include "DragAndDrop/DecoratedDragDropOp.h"
 #include "DragAndDrop/DropOperationSystem.h"
 #include "DragAndDrop/TedsDragDropOpUtility.h"
-#include "Elements/Columns/TypedElementHandleColumn.h"
+#include "Editor.h"
+#include "Elements/Framework/EngineElementsLibrary.h"
 #include "Elements/Framework/TypedElementSelectionSet.h"
 #include "Elements/Interfaces/TypedElementDataStorageInterface.h"
 #include "Input/DragAndDrop.h"
@@ -19,42 +20,6 @@
 namespace UE::Editor
 {
 
-namespace WidgetDropHandler_Private
-{
-	static void SelectRows(DataStorage::ICoreProvider& Storage, UTypedElementSelectionSet& SelectionSet, DataStorage::FRowHandleArray RowsToSelect)
-	{
-		using namespace UE::Editor::DataStorage;
-
-		if (RowsToSelect.IsEmpty())
-		{
-			return;
-		}
-		
-		// Select next tick, so that all columns have been set and updated.
-		FTSTicker::GetCoreTicker().AddTicker(FTickerDelegate::CreateLambda(
-			[Storage = &Storage, SelectionSet = TWeakObjectPtr(&SelectionSet), RowsToSelect = MoveTemp(RowsToSelect)](float)
-			{
-				if (TStrongObjectPtr<UTypedElementSelectionSet> SelectionSetPin = SelectionSet.Pin())
-				{
-					TArray<FTypedElementHandle> Elements;
-					Elements.Reserve(RowsToSelect.Num());
-					for (RowHandle Row : RowsToSelect.GetRows())
-					{
-						Compatibility::FTypedElementColumn* Column = Storage->GetColumn<Compatibility::FTypedElementColumn>(Row);
-						if (!Column || !Column->Handle)
-						{
-							continue;
-						}
-						
-						Elements.Add(Column->Handle);
-					}
-					SelectionSetPin->SetSelection(Elements, {});
-				}
-				return false;
-			}));
-	}
-}
-
 FWidgetDropHandler::FWidgetDropHandler(TWeakObjectPtr<UTypedElementSelectionSet> SelectionSet)
 	: SelectionSet(MoveTemp(SelectionSet))
 	, CurrentOperation(nullptr)
@@ -276,30 +241,50 @@ void FWidgetDropHandler::Drop(DataStorage::ICoreProvider& Storage, UDropOperatio
 	(void)UpdateParameters(Storage, Geometry, InputEvent);
 	SetRowsFromParameters(Storage, InputRows.GetRows());
 	{
+		// Bit of a hack to select our newly placed objects.
+		// Usually we would like use teds rows for this, but since the synchronisation is deferred we cannot access TEv1 elements from our rows yet.
+		// We also cannot access the rows next tick, because we might be within a transaction scope which would not capture the selection.
+		// To solve this we collect all actors that are placed while the drop process is ongoing.
+		UTypedElementSelectionSet* SelectionSetPtr = SelectionSet.Get();	
+		TSet<TObjectPtr<AActor>> PlacedActors;	
+		FDelegateHandle DelegateHandle;
+		ON_SCOPE_EXIT{ GEditor->OnLevelActorAdded().Remove(DelegateHandle); };
+		if (SelectionSetPtr)
+		{
+			DelegateHandle = GEditor->OnLevelActorAdded().AddLambda([&PlacedActors](AActor* Actor)
+				{
+					if (Actor)
+					{
+						PlacedActors.Add(Actor);
+					}
+				});
+		}
+		
 		FScopedTransaction Transaction(LOCTEXT("Drop_Transaction", "Drop on widget."));
 		(void)DropSystem.Apply(InputRows.GetRows());
-	}
 
-	UTypedElementSelectionSet* SelectionSetPtr = SelectionSet.Get();
-	FRowHandleArray RowsToSelect;
-	if (SelectionSetPtr)
-	{
-		RowsToSelect.Reserve(InputRows.Num());
-		for (RowHandle InputRow : InputRows.GetRows())
+		// Now we select the roots of all captured actors.
+		if (SelectionSetPtr)
 		{
-			if (const FResultColumn* Column = Storage.GetColumn<FResultColumn>(InputRow))
+			TArray<FTypedElementHandle> RootElements;
+			for (const TObjectPtr<AActor>& ActorPtr : PlacedActors)
 			{
-				RowsToSelect.Append(Column->Value.Created.GetRows());
+				if (AActor* Actor = ActorPtr.Get())
+				{
+					AActor* Parent = Actor->GetParentActor();
+					if (!Parent || !PlacedActors.Contains(Parent))
+					{
+						constexpr bool bAllowCreate = true;
+						RootElements.Add(UEngineElementsLibrary::AcquireEditorActorElementHandle(Actor, bAllowCreate));
+					}
+				}
 			}
+			
+			SelectionSetPtr->SetSelection(RootElements, FTypedElementSelectionOptions());
 		}
 	}
 	
 	Stop(Storage, DropSystem);
-
-	if (SelectionSetPtr)
-	{
-		WidgetDropHandler_Private::SelectRows(Storage, *SelectionSetPtr, MoveTemp(RowsToSelect));
-	}
 }
 
 void FWidgetDropHandler::UpdateOperationDecorator(FDragDropOperation& Operation, TOptional<FDropResult> Result) const
