[TEDS Outliner] Add default sorting to TEDS Hierarchies on the Primary Column

#Synced-CL 49032379

diff --git a/Engine/Plugins/Experimental/EditorDataStorageFeatures/Source/TedsTableViewer/Private/TedsTableViewerModel.cpp b/Engine/Plugins/Experimental/EditorDataStorageFeatures/Source/TedsTableViewer/Private/TedsTableViewerModel.cpp
--- a/Engine/Plugins/Experimental/EditorDataStorageFeatures/Source/TedsTableViewer/Private/TedsTableViewerModel.cpp
+++ b/Engine/Plugins/Experimental/EditorDataStorageFeatures/Source/TedsTableViewer/Private/TedsTableViewerModel.cpp
@@ -283,6 +283,17 @@ namespace UE::Editor::DataStorage
 		if (const TSharedRef<FTedsTableViewerColumn>* FoundPrimaryColumn = NewColumnMap.Find(PrimaryColumn))
 		{
 			PrimaryUIColumnName = (*FoundPrimaryColumn)->GetColumnName();
+
+			// Initialize default sorting with the primary column, if it exists
+			// Eventually there will be a use case to make this into a separate parameter, but for now this is sufficient for current user needs
+			const FMetaDataEntryView MetaDataEntryView = (*FoundPrimaryColumn)->GetMetaData().FindGeneric("bColumnInitialSortIsAscending");
+			const bool* bColumnInitialSortIsAscending = MetaDataEntryView.TryGetExact<bool>();
+			EColumnSortMode::Type SortMode = EColumnSortMode::Ascending;
+			if (bColumnInitialSortIsAscending)
+			{
+				SortMode = *bColumnInitialSortIsAscending ? EColumnSortMode::Ascending : EColumnSortMode::Descending;
+			}
+			(*FoundPrimaryColumn)->OnSortCallback(EColumnSortPriority::Primary, PrimaryUIColumnName, SortMode);
 		}
 		else
 		{
diff --git a/Engine/Plugins/Experimental/EditorDataStorageFeatures/Source/TedsTableViewer/Private/Widgets/STedsTableViewer.cpp b/Engine/Plugins/Experimental/EditorDataStorageFeatures/Source/TedsTableViewer/Private/Widgets/STedsTableViewer.cpp
--- a/Engine/Plugins/Experimental/EditorDataStorageFeatures/Source/TedsTableViewer/Private/Widgets/STedsTableViewer.cpp
+++ b/Engine/Plugins/Experimental/EditorDataStorageFeatures/Source/TedsTableViewer/Private/Widgets/STedsTableViewer.cpp
@@ -59,7 +59,8 @@ namespace UE::Editor::DataStorage
 			.ColumnMetaData = InArgs._ColumnMetaData,
 			.CellWidgetPurpose = CellWidgetPurpose,
 			.HeaderWidgetPurpose = HeaderWidgetPurpose,
-			.IsItemVisibleDelegate = FTedsTableViewerModel::FIsItemVisible::CreateSP(this, &STedsTableViewer::IsItemVisible)
+			.IsItemVisibleDelegate = FTedsTableViewerModel::FIsItemVisible::CreateSP(this, &STedsTableViewer::IsItemVisible),
+			.PrimaryColumn = InArgs._PrimaryColumn
 		});
 
 		InstanceCounter++;
diff --git a/Engine/Plugins/Experimental/EditorDataStorageFeatures/Source/TedsTableViewer/Public/TedsTableViewerColumn.h b/Engine/Plugins/Experimental/EditorDataStorageFeatures/Source/TedsTableViewer/Public/TedsTableViewerColumn.h
--- a/Engine/Plugins/Experimental/EditorDataStorageFeatures/Source/TedsTableViewer/Public/TedsTableViewerColumn.h
+++ b/Engine/Plugins/Experimental/EditorDataStorageFeatures/Source/TedsTableViewer/Public/TedsTableViewerColumn.h
@@ -71,6 +71,7 @@ namespace UE::Editor::DataStorage
 		TEDSTABLEVIEWER_API bool HasSortInfo() const;
 		TEDSTABLEVIEWER_API void ClearSort();
 		TEDSTABLEVIEWER_API void OnSortCallback(EColumnSortPriority::Type InSortPriority, const FName& InColumnId, EColumnSortMode::Type InSortDirection);
+		TEDSTABLEVIEWER_API FComboMetaDataView<FGenericMetaDataView, FColumnsMetaDataView> GetMetaData() const;
 
 	protected:
 
@@ -78,7 +79,6 @@ namespace UE::Editor::DataStorage
 		void UnRegisterQueries();
 		bool IsRowVisible(const RowHandle InRowHandle) const;
 		void UpdateWidgets();
-		FComboMetaDataView<FGenericMetaDataView, FColumnsMetaDataView> GetMetaData() const;
 		
 	private:
 
diff --git a/Engine/Plugins/Experimental/EditorDataStorageFeatures/Source/TedsTableViewer/Public/Widgets/STedsTableViewer.h b/Engine/Plugins/Experimental/EditorDataStorageFeatures/Source/TedsTableViewer/Public/Widgets/STedsTableViewer.h
--- a/Engine/Plugins/Experimental/EditorDataStorageFeatures/Source/TedsTableViewer/Public/Widgets/STedsTableViewer.h
+++ b/Engine/Plugins/Experimental/EditorDataStorageFeatures/Source/TedsTableViewer/Public/Widgets/STedsTableViewer.h
@@ -81,6 +81,9 @@ namespace UE::Editor::DataStorage
 		// The message to show in place of the table viewer when there are no rows provided by the current query stack
 		// Empty = simply show the column headers instead of a message
 		SLATE_ATTRIBUTE(FText, EmptyRowsMessage)
+		
+		// The primary column in the table viewer is the one that is sorted by default
+		SLATE_ARGUMENT(TWeakObjectPtr<const UScriptStruct>, PrimaryColumn)
 
 		/** The height of the list item */
 		SLATE_ATTRIBUTE(float, ItemHeight)
