TedsSettings : add a hierarchy and settings property rows as children of the settings section rows.

Also, remove the notion of "inactive" settings sections.  This feature is no longer used and did not play nicely with TEDS hierarchies because the hierarchies internally adds private columns to the rows and these were not detectable for migration between active <-> inactive section rows.  Since the feature was essentially just technical debt it simplifies things to remove it.

#Synced-CL 48966007

diff --git a/Engine/Plugins/Experimental/EditorDataStorageFeatures/Source/TedsSettings/Private/TedsSettingsEditorSubsystem.cpp b/Engine/Plugins/Experimental/EditorDataStorageFeatures/Source/TedsSettings/Private/TedsSettingsEditorSubsystem.cpp
--- a/Engine/Plugins/Experimental/EditorDataStorageFeatures/Source/TedsSettings/Private/TedsSettingsEditorSubsystem.cpp
+++ b/Engine/Plugins/Experimental/EditorDataStorageFeatures/Source/TedsSettings/Private/TedsSettingsEditorSubsystem.cpp
@@ -33,19 +33,29 @@ UTedsSettingsEditorSubsystem::FOnEnabledChanged& UTedsSettingsEditorSubsystem::O
 	return EnabledChangedDelegate;
 }
 
-UE::Editor::DataStorage::RowHandle UTedsSettingsEditorSubsystem::FindSettingsSection(const FName& ContainerName, const FName& CategoryName, const FName& SectionName)
+UE::Editor::DataStorage::RowHandle UTedsSettingsEditorSubsystem::FindSettingsContainer(const FName& ContainerName) const
+{
+	return SettingsManager->FindSettingsContainer(ContainerName);
+}
+
+UE::Editor::DataStorage::RowHandle UTedsSettingsEditorSubsystem::FindSettingsCategory(const FName& ContainerName, const FName& CategoryName) const
+{
+	return SettingsManager->FindSettingsCategory(ContainerName, CategoryName);
+}
+
+UE::Editor::DataStorage::RowHandle UTedsSettingsEditorSubsystem::FindSettingsSection(const FName& ContainerName, const FName& CategoryName, const FName& SectionName) const
 {
 	return SettingsManager->FindSettingsSection(ContainerName, CategoryName, SectionName);
 }
 
-UE::Editor::DataStorage::RowHandle UTedsSettingsEditorSubsystem::FindOrAddSettingsSection(const FName& ContainerName, const FName& CategoryName, const FName& SectionName)
+bool UTedsSettingsEditorSubsystem::GetSettingsSectionFromRow(UE::Editor::DataStorage::RowHandle SectionRow, FName& OutContainerName, FName& OutCategoryName, FName& OutSectionName)
 {
-	return SettingsManager->FindOrAddSettingsSection(ContainerName, CategoryName, SectionName);
+	return SettingsManager->GetSettingsSectionFromRow(SectionRow, OutContainerName, OutCategoryName, OutSectionName);
 }
 
-bool UTedsSettingsEditorSubsystem::GetSettingsSectionFromRow(UE::Editor::DataStorage::RowHandle Row, FName& OutContainerName, FName& OutCategoryName, FName& OutSectionName)
+bool UTedsSettingsEditorSubsystem::GetSettingsCategoryFromRow(UE::Editor::DataStorage::RowHandle CategoryRow, FName& OutContainerName, FName& OutCategoryName)
 {
-	return SettingsManager->GetSettingsSectionFromRow(Row, OutContainerName, OutCategoryName, OutSectionName);
+	return SettingsManager->GetSettingsCategoryFromRow(CategoryRow, OutContainerName, OutCategoryName);
 }
 
 void UTedsSettingsEditorSubsystem::Initialize(FSubsystemCollectionBase& Collection)
diff --git a/Engine/Plugins/Experimental/EditorDataStorageFeatures/Source/TedsSettings/Private/TedsSettingsManager.cpp b/Engine/Plugins/Experimental/EditorDataStorageFeatures/Source/TedsSettings/Private/TedsSettingsManager.cpp
--- a/Engine/Plugins/Experimental/EditorDataStorageFeatures/Source/TedsSettings/Private/TedsSettingsManager.cpp
+++ b/Engine/Plugins/Experimental/EditorDataStorageFeatures/Source/TedsSettings/Private/TedsSettingsManager.cpp
@@ -22,7 +22,25 @@
 
 namespace UE::Editor::Settings::Private
 {
-	static UE::Editor::DataStorage::FMapKey GenerateIndexKey(const FName& ContainerName, const FName& CategoryName, const FName& SectionName)
+	static UE::Editor::DataStorage::FMapKey GenerateContainerIndexKey(const FName& ContainerName)
+	{
+		FString Key = "ISettingsContainer: ";
+		ContainerName.AppendString(Key);
+
+		return UE::Editor::DataStorage::FMapKey(MoveTemp(Key));
+	}
+
+	static UE::Editor::DataStorage::FMapKey GenerateCategoryIndexKey(const FName& ContainerName, const FName& CategoryName)
+	{
+		FString Key = "ISettingsCategory: ";
+		ContainerName.AppendString(Key);
+		Key.AppendChar(TEXT(','));
+		CategoryName.AppendString(Key);
+
+		return UE::Editor::DataStorage::FMapKey(MoveTemp(Key));
+	}
+
+	static UE::Editor::DataStorage::FMapKey GenerateSectionIndexKey(const FName& ContainerName, const FName& CategoryName, const FName& SectionName)
 	{
 		FString Key = "ISettingsSection: ";
 		ContainerName.AppendString(Key);
@@ -33,31 +51,15 @@ namespace UE::Editor::Settings::Private
 	
 		return UE::Editor::DataStorage::FMapKey(MoveTemp(Key));
 	}
-
-	static bool IsUnknownColumn(const UScriptStruct* ColumnType)
-	{
-		check(ColumnType);
-		return ColumnType != FSettingsSectionTag::StaticStruct()
-			&& ColumnType != FSettingsInactiveSectionTag::StaticStruct()
-			&& ColumnType != FSettingsContainerReferenceColumn::StaticStruct()
-			&& ColumnType != FSettingsCategoryReferenceColumn::StaticStruct()
-			&& ColumnType != FSettingsNameColumn::StaticStruct()
-			&& ColumnType != FDisplayNameColumn::StaticStruct()
-			&& ColumnType != FDescriptionColumn::StaticStruct()
-			&& ColumnType != FTypedElementUObjectColumn::StaticStruct()
-			&& ColumnType != FTypedElementUObjectIdColumn::StaticStruct()
-			&& ColumnType != FTypedElementClassTypeInfoColumn::StaticStruct()
-			&& ColumnType != FTypedElementClassDefaultObjectTag::StaticStruct();
-	}
 }
 
 FTedsSettingsManager::FTedsSettingsManager()
 	: bIsInitialized{ false }
-	, SelectAllActiveSettingsQuery{ UE::Editor::DataStorage::InvalidQueryHandle }
-	, SelectAllInactiveSettingsQuery{ UE::Editor::DataStorage::InvalidQueryHandle }
+	, SelectAllSettingsQuery{ UE::Editor::DataStorage::InvalidQueryHandle }
 	, SettingsContainerTable{ UE::Editor::DataStorage::InvalidTableHandle }
 	, SettingsCategoryTable{ UE::Editor::DataStorage::InvalidTableHandle }
-	, SettingsInactiveSectionTable{ UE::Editor::DataStorage::InvalidTableHandle }
+	, SettingsPropertyTable{ UE::Editor::DataStorage::InvalidTableHandle }
+	, SettingsHierarchy{ }
 {
 }
 
@@ -72,6 +74,7 @@ void FTedsSettingsManager::Initialize()
 				check(DataStorage);
 
 				RegisterTables(*DataStorage);
+				RegisterHierarchies(*DataStorage);
 				RegisterQueries(*DataStorage);
 				RegisterActiveSettings();
 			};
@@ -102,7 +105,6 @@ void FTedsSettingsManager::Shutdown()
 			check(DataStorage);
 
 			UnregisterActiveSettings();
-			UnregisterInactiveSettings();
 			UnregisterQueries(*DataStorage);
 		}
 
@@ -110,58 +112,54 @@ void FTedsSettingsManager::Shutdown()
 	}
 }
 
-UE::Editor::DataStorage::RowHandle FTedsSettingsManager::FindSettingsSection(const FName& ContainerName, const FName& CategoryName, const FName& SectionName)
+UE::Editor::DataStorage::RowHandle FTedsSettingsManager::FindSettingsContainer(const FName& ContainerName) const
 {
 	using namespace UE::Editor::DataStorage;
 	using namespace UE::Editor::Settings::Private;
 
 	if (bIsInitialized)
 	{
-		if (ICoreProvider* DataStorage = GetMutableDataStorageFeature<ICoreProvider>(StorageFeatureName))
+		if (const ICoreProvider* DataStorage = GetDataStorageFeature<ICoreProvider>(StorageFeatureName))
 		{
-			// Ensure that the SettingsInactiveSectionTable is registered since it is possible for OnDataStorageInterfacesSet to be called
-			// for a delegate bound in another module before it has been called for our bound delegate.
-			RegisterTables(*DataStorage);
-
-			FMapKey SectionIndexKey = GenerateIndexKey(ContainerName, CategoryName, SectionName);
+			FMapKey ContainerIndexKey = GenerateContainerIndexKey(ContainerName);
 
-			return DataStorage->LookupMappedRow(Settings::MappingDomain, SectionIndexKey);
+			return DataStorage->LookupMappedRow(Settings::MappingDomain, ContainerIndexKey);
 		}
 	}
 
 	return InvalidRowHandle;
 }
 
-UE::Editor::DataStorage::RowHandle FTedsSettingsManager::FindOrAddSettingsSection(const FName& ContainerName, const FName& CategoryName, const FName& SectionName)
+UE::Editor::DataStorage::RowHandle FTedsSettingsManager::FindSettingsCategory(const FName& ContainerName, const FName& CategoryName) const
 {
 	using namespace UE::Editor::DataStorage;
 	using namespace UE::Editor::Settings::Private;
 
 	if (bIsInitialized)
 	{
-		if (ICoreProvider* DataStorage = GetMutableDataStorageFeature<ICoreProvider>(StorageFeatureName))
+		if (const ICoreProvider* DataStorage = GetDataStorageFeature<ICoreProvider>(StorageFeatureName))
 		{
-			// Ensure that the SettingsInactiveSectionTable is registered since it is possible for OnDataStorageInterfacesSet to be called
-			// for a delegate bound in another module before it has been called for our bound delegate.
-			RegisterTables(*DataStorage);
-
-			FMapKey SectionIndexKey = GenerateIndexKey(ContainerName, CategoryName, SectionName);
+			FMapKey CategoryIndexKey = GenerateCategoryIndexKey(ContainerName, CategoryName);
 
-			RowHandle SectionRow = DataStorage->LookupMappedRow(Settings::MappingDomain, SectionIndexKey);
+			return DataStorage->LookupMappedRow(Settings::MappingDomain, CategoryIndexKey);
+		}
+	}
 
-			if (SectionRow == InvalidRowHandle)
-			{
-				SectionRow = DataStorage->AddRow(SettingsInactiveSectionTable);
+	return InvalidRowHandle;
+}
 
-				DataStorage->AddColumn<FSettingsInactiveSectionTag>(SectionRow);
-				DataStorage->AddColumn<FSettingsNameColumn>(SectionRow, { .Name = SectionName });
-				DataStorage->AddColumn<FSettingsContainerReferenceColumn>(SectionRow, { .ContainerName = ContainerName, .ContainerRow = InvalidRowHandle });
-				DataStorage->AddColumn<FSettingsCategoryReferenceColumn>(SectionRow, { .CategoryName = CategoryName, .CategoryRow = InvalidRowHandle });
+UE::Editor::DataStorage::RowHandle FTedsSettingsManager::FindSettingsSection(const FName& ContainerName, const FName& CategoryName, const FName& SectionName) const
+{
+	using namespace UE::Editor::DataStorage;
+	using namespace UE::Editor::Settings::Private;
 
-				DataStorage->MapRow(Settings::MappingDomain, MoveTemp(SectionIndexKey), SectionRow);
-			}
+	if (bIsInitialized)
+	{
+		if (const ICoreProvider* DataStorage = GetDataStorageFeature<ICoreProvider>(StorageFeatureName))
+		{
+			FMapKey SectionIndexKey = GenerateSectionIndexKey(ContainerName, CategoryName, SectionName);
 
-			return SectionRow;
+			return DataStorage->LookupMappedRow(Settings::MappingDomain, SectionIndexKey);
 		}
 	}
 
@@ -176,7 +174,7 @@ bool FTedsSettingsManager::GetSettingsSectionFromRow(UE::Editor::DataStorage::Ro
 	{
 		if (const ICoreProvider* DataStorage = GetDataStorageFeature<ICoreProvider>(StorageFeatureName))
 		{
-			if (DataStorage->IsRowAvailable(Row) && (DataStorage->HasColumns<FSettingsSectionTag>(Row) || DataStorage->HasColumns<FSettingsInactiveSectionTag>(Row)))
+			if (DataStorage->IsRowAvailable(Row) && (DataStorage->HasColumns<FSettingsSectionTag>(Row)))
 			{
 				const FSettingsContainerReferenceColumn* ContainerReferenceColumn = DataStorage->GetColumn<FSettingsContainerReferenceColumn>(Row);
 				const FSettingsCategoryReferenceColumn* CategoryReferenceColumn = DataStorage->GetColumn<FSettingsCategoryReferenceColumn>(Row);
@@ -197,9 +195,37 @@ bool FTedsSettingsManager::GetSettingsSectionFromRow(UE::Editor::DataStorage::Ro
 	return false;
 }
 
+bool FTedsSettingsManager::GetSettingsCategoryFromRow(UE::Editor::DataStorage::RowHandle Row, FName& OutContainerName, FName& OutCategoryName)
+{
+	using namespace UE::Editor::DataStorage;
+
+	if (bIsInitialized)
+	{
+		if (const ICoreProvider* DataStorage = GetDataStorageFeature<ICoreProvider>(StorageFeatureName))
+		{
+			if (DataStorage->IsRowAvailable(Row) && DataStorage->HasColumns<FSettingsCategoryTag>(Row))
+			{
+				const FSettingsContainerReferenceColumn* ContainerReferenceColumn = DataStorage->GetColumn<FSettingsContainerReferenceColumn>(Row);
+				const FSettingsNameColumn* NameColumn = DataStorage->GetColumn<FSettingsNameColumn>(Row);
+
+				if (ContainerReferenceColumn && NameColumn)
+				{
+					OutContainerName = ContainerReferenceColumn->ContainerName;
+					OutCategoryName = NameColumn->Name;
+
+					return true;
+				}
+			}
+		}
+	}
+
+	return false;
+}
+
 void FTedsSettingsManager::RegisterTables(UE::Editor::DataStorage::ICoreProvider& DataStorage)
 {
 	using namespace UE::Editor::DataStorage;
+
 	if (SettingsContainerTable == InvalidTableHandle)
 	{
 		SettingsContainerTable = DataStorage.RegisterTable(
@@ -214,46 +240,45 @@ void FTedsSettingsManager::RegisterTables(UE::Editor::DataStorage::ICoreProvider
 			FName(TEXT("Editor_SettingsCategoryTable")));
 	}
 
-	if (SettingsInactiveSectionTable == InvalidTableHandle)
+	if (SettingsPropertyTable == InvalidTableHandle)
 	{
-		SettingsInactiveSectionTable = DataStorage.RegisterTable(
-			TTypedElementColumnTypeList<FSettingsContainerReferenceColumn, FSettingsCategoryReferenceColumn, FSettingsNameColumn, FSettingsInactiveSectionTag>(),
-			FName(TEXT("Editor_SettingsInactiveSectionTable")));
+		SettingsPropertyTable = DataStorage.RegisterTable(
+			TTypedElementColumnTypeList<FSettingsNameColumn, FDisplayNameColumn, FDescriptionColumn, FSettingsPropertyTag>(),
+			FName(TEXT("Editor_SettingsPropertyTable")));
 	}
 }
 
+void FTedsSettingsManager::RegisterHierarchies(UE::Editor::DataStorage::ICoreProvider& DataStorage)
+{
+	using namespace UE::Editor::DataStorage;
+
+	FHierarchyRegistrationParams HierarchyRegistrationParams
+	{
+		.Name = UE::Editor::DataStorage::Settings::HierarchyName
+	};
+
+	SettingsHierarchy = DataStorage.RegisterHierarchy(HierarchyRegistrationParams);
+}
+
 void FTedsSettingsManager::RegisterQueries(UE::Editor::DataStorage::ICoreProvider& DataStorage)
 {
 	using namespace UE::Editor::DataStorage::Queries;
 
-	if (SelectAllActiveSettingsQuery == UE::Editor::DataStorage::InvalidQueryHandle)
+	if (SelectAllSettingsQuery == UE::Editor::DataStorage::InvalidQueryHandle)
 	{
-		SelectAllActiveSettingsQuery = DataStorage.RegisterQuery(
+		SelectAllSettingsQuery = DataStorage.RegisterQuery(
 			Select()
 				.ReadOnly<FSettingsContainerReferenceColumn, FSettingsCategoryReferenceColumn, FSettingsNameColumn>()
 			.Where()
 				.All<FSettingsSectionTag>()
 			.Compile());
 	}
-
-	if (SelectAllInactiveSettingsQuery == UE::Editor::DataStorage::InvalidQueryHandle)
-	{
-		SelectAllInactiveSettingsQuery = DataStorage.RegisterQuery(
-			Select()
-				.ReadOnly<FSettingsContainerReferenceColumn, FSettingsCategoryReferenceColumn, FSettingsNameColumn>()
-			.Where()
-				.All<FSettingsInactiveSectionTag>()
-			.Compile());
-	}
 }
 
 void FTedsSettingsManager::UnregisterQueries(UE::Editor::DataStorage::ICoreProvider& DataStorage)
 {
-	DataStorage.UnregisterQuery(SelectAllActiveSettingsQuery);
-	SelectAllActiveSettingsQuery = UE::Editor::DataStorage::InvalidQueryHandle;
-
-	DataStorage.UnregisterQuery(SelectAllInactiveSettingsQuery);
-	SelectAllInactiveSettingsQuery = UE::Editor::DataStorage::InvalidQueryHandle;
+	DataStorage.UnregisterQuery(SelectAllSettingsQuery);
+	SelectAllSettingsQuery = UE::Editor::DataStorage::InvalidQueryHandle;
 }
 
 void FTedsSettingsManager::RegisterActiveSettings()
@@ -277,6 +302,7 @@ void FTedsSettingsManager::RegisterActiveSettings()
 void FTedsSettingsManager::RegisterSettingsContainer(const FName& ContainerName)
 {
 	using namespace UE::Editor::DataStorage;
+	using namespace UE::Editor::Settings::Private;
 
 	TRACE_CPUPROFILER_EVENT_SCOPE(TedsSettingsManager.RegisterSettingsContainer);
 
@@ -290,7 +316,7 @@ void FTedsSettingsManager::RegisterSettingsContainer(const FName& ContainerName)
 
 	ISettingsContainerPtr ContainerPtr = SettingsModule->GetContainer(ContainerName);
 
-	FMapKey ContainerIndexKey = FMapKey(ContainerPtr.Get());
+	FMapKey ContainerIndexKey = GenerateContainerIndexKey(ContainerName);
 	RowHandle ContainerRow = DataStorage->LookupMappedRow(Settings::MappingDomain, ContainerIndexKey);
 	if (ContainerRow == InvalidRowHandle)
 	{
@@ -352,65 +378,32 @@ void FTedsSettingsManager::UnregisterActiveSettings()
 
 		ContainerPtr->OnCategoryModified().RemoveAll(this);
 
-		TArray<ISettingsCategoryPtr> Categories;
-		ContainerPtr->GetCategories(Categories);
-
-		for (ISettingsCategoryPtr CategoryPtr : Categories)
+		FMapKey ContainerIndexKey = GenerateContainerIndexKey(ContainerName);
+		RowHandle ContainerRow = DataStorage->LookupMappedRow(Settings::MappingDomain, ContainerIndexKey);
+		if (ContainerRow != InvalidRowHandle)
 		{
-			UE_LOG(LogTedsSettings, Verbose, TEXT("Unregister Settings Category : '%s'"), *CategoryPtr->GetName().ToString());
-
-			TArray<ISettingsSectionPtr> Sections;
-			const bool bIgnoreVisibility = true;
-			CategoryPtr->GetSections(Sections, bIgnoreVisibility);
-
-			for (ISettingsSectionPtr SectionPtr : Sections)
+			TArray<TWeakObjectPtr<UObject>> ObjectsToRemove;
+			FRowHandleArray RowsToRemove;
+			auto VisitRow = [DataStorage, DataStorageCompatibility, &RowsToRemove, &ObjectsToRemove](const ICoreProvider&, RowHandle, RowHandle TargetRow)
 			{
-				if (TStrongObjectPtr<UObject> SettingsObjectPtr = SectionPtr->GetSettingsObject().Pin(); SettingsObjectPtr)
+				if (FTypedElementUObjectColumn* UObjectColumn = DataStorage->GetColumn<FTypedElementUObjectColumn>(TargetRow))
 				{
-					DataStorageCompatibility->RemoveCompatibleObject(SettingsObjectPtr);
-
-					UE_LOG(LogTedsSettings, Verbose, TEXT("Removed Settings Section : '%s'"), *SectionPtr->GetName().ToString());
+					ObjectsToRemove.Add(UObjectColumn->Object);
 				}
-			}
+				else
+				{
+					RowsToRemove.Add(TargetRow);
+				}
+			};
 
-			FMapKeyView CategoryIndexKey = FMapKeyView(CategoryPtr.Get());
-			RowHandle CategoryRow = DataStorage->LookupMappedRow(Settings::MappingDomain, CategoryIndexKey);
-			if (CategoryRow != InvalidRowHandle)
+			DataStorage->WalkDepthFirst(SettingsHierarchy, ContainerRow, VisitRow);
+			DataStorage->BatchRemoveRows(RowsToRemove.GetRows());
+
+			for (TWeakObjectPtr<UObject> ObjectToRemove : ObjectsToRemove)
 			{
-				DataStorage->RemoveRow(CategoryRow);
+				DataStorageCompatibility->RemoveCompatibleObject(ObjectToRemove);
 			}
 		}
-
-		FMapKeyView ContainerIndexKey = FMapKeyView(ContainerPtr.Get());
-		RowHandle ContainerRow = DataStorage->LookupMappedRow(Settings::MappingDomain, ContainerIndexKey);
-		if (ContainerRow != InvalidRowHandle)
-		{
-			DataStorage->RemoveRow(ContainerRow);
-		}
-	}
-}
-
-void FTedsSettingsManager::UnregisterInactiveSettings()
-{
-	TRACE_CPUPROFILER_EVENT_SCOPE(TedsSettingsManager.UnregisterInactiveSettings);
-
-	using namespace UE::Editor::DataStorage::Queries;
-	using namespace UE::Editor::Settings::Private;
-
-	ICoreProvider* DataStorage = GetMutableDataStorageFeature<ICoreProvider>(StorageFeatureName);
-	check(DataStorage);
-	
-	TSet<RowHandle> RowHandles;
-
-	DataStorage->RunQuery(SelectAllInactiveSettingsQuery, CreateDirectQueryCallbackBinding(
-		[&RowHandles](IDirectQueryContext& Context, const RowHandle*)
-		{
-			RowHandles.Append(Context.GetRowHandles());
-		}));
-
-	for (const RowHandle& Row : RowHandles)
-	{
-		DataStorage->RemoveRow(Row);
 	}
 }
 
@@ -432,7 +425,7 @@ void FTedsSettingsManager::UpdateSettingsCategory(TSharedPtr<ISettingsCategory>
 
 	UE_LOG(LogTedsSettings, Verbose, TEXT("Update Settings Category: '%s->%s'"), *ContainerName.ToString(), *CategoryName.ToString());
 
-	FMapKey CategoryIndexKey = FMapKey(SettingsCategory.Get());
+	FMapKey CategoryIndexKey = GenerateCategoryIndexKey(ContainerName, CategoryName);
 
 	RowHandle CategoryRow = DataStorage->LookupMappedRow(Settings::MappingDomain, CategoryIndexKey);
 	if (CategoryRow == InvalidRowHandle)
@@ -445,18 +438,20 @@ void FTedsSettingsManager::UpdateSettingsCategory(TSharedPtr<ISettingsCategory>
 		DataStorage->AddColumn<FDescriptionColumn>(CategoryRow, { .Description = SettingsCategory->GetDescription() });
 		DataStorage->AddColumn<FSettingsCategoryTag>(CategoryRow);
 
+		DataStorage->SetParentRow(SettingsHierarchy, CategoryRow, ContainerRow);
+
 		DataStorage->MapRow(Settings::MappingDomain, MoveTemp(CategoryIndexKey), CategoryRow);
 	}
 
 	TArray<RowHandle> OldRowHandles;
 	TArray<FName> OldSectionNames;
 
-	// Gather all existing active rows for the given { ContainerName, CategoryName } pair.
+	// Gather all existing section rows for the given { ContainerName, CategoryName } pair.
 	if (bQueryExistingRows)
 	{
 		using namespace UE::Editor::DataStorage::Queries;
 
-		DataStorage->RunQuery(SelectAllActiveSettingsQuery, CreateDirectQueryCallbackBinding(
+		DataStorage->RunQuery(SelectAllSettingsQuery, CreateDirectQueryCallbackBinding(
 			[&OldRowHandles, &OldSectionNames, &ContainerName, &CategoryName](
 				IDirectQueryContext& Context,
 				const FSettingsContainerReferenceColumn* ContainerColumns,
@@ -494,25 +489,12 @@ void FTedsSettingsManager::UpdateSettingsCategory(TSharedPtr<ISettingsCategory>
 		{
 			NewSectionNames.Emplace(SectionName);
 
-			FMapKey SectionIndexKey = GenerateIndexKey(ContainerName, CategoryName, SectionName);
-
-			TArray<void*> ColumnsToCopy;
-			TArray<const UScriptStruct*> ColumnTypesToCopy;
+			FMapKey SectionIndexKey = GenerateSectionIndexKey(ContainerName, CategoryName, SectionName);
 
 			RowHandle OldSectionRow = DataStorage->LookupMappedRow(Settings::MappingDomain, SectionIndexKey);
 			if (OldSectionRow != InvalidRowHandle)
 			{
 				UE_LOG(LogTedsSettings, Verbose, TEXT("Settings Section : '%s' is already in data storage"), *SectionName.ToString());
-				
-				// Capture unknown columns and copy them to the new row before we delete the old row.
-				DataStorage->ListColumns(OldSectionRow, [&ColumnsToCopy, &ColumnTypesToCopy](void* Column, const UScriptStruct& ColumnType)
-					{
-						if (IsUnknownColumn(&ColumnType))
-						{
-							ColumnsToCopy.Add(Column);
-							ColumnTypesToCopy.Add(&ColumnType);
-						}
-					});
 			}
 
 			RowHandle NewSectionRow = DataStorageCompatibility->AddCompatibleObject(SettingsObjectPtr);
@@ -526,10 +508,27 @@ void FTedsSettingsManager::UpdateSettingsCategory(TSharedPtr<ISettingsCategory>
 			
 			if (OldSectionRow != InvalidRowHandle && OldSectionRow != NewSectionRow)
 			{
-				AddColumns(*DataStorage, NewSectionRow, ColumnsToCopy, ColumnTypesToCopy);
+				// Remove the old section row and any child property rows.
+				FRowHandleArray RowsToRemove;
+				auto VisitRow = [DataStorage, OldSectionRow, &RowsToRemove](const ICoreProvider&, RowHandle, RowHandle TargetRow)
+				{
+					if (TargetRow != OldSectionRow)
+					{
+						RowsToRemove.Add(TargetRow);
+					}
+				};
 
-				// Remove the old row after we copy the unknown columns to the new row.
-				DataStorage->RemoveRow(OldSectionRow);
+				DataStorage->WalkDepthFirst(SettingsHierarchy, OldSectionRow, VisitRow);
+				DataStorage->BatchRemoveRows(RowsToRemove.GetRows());
+
+				if (FTypedElementUObjectColumn* UObjectColumn = DataStorage->GetColumn<FTypedElementUObjectColumn>(OldSectionRow))
+				{
+					DataStorageCompatibility->RemoveCompatibleObject(UObjectColumn->Object);
+				}
+				else
+				{
+					DataStorage->RemoveRow(OldSectionRow);
+				}
 
 				UE_LOG(LogTedsSettings, Verbose, TEXT("Removed Settings Section : '%s'"), *SectionName.ToString());
 			}
@@ -537,11 +536,33 @@ void FTedsSettingsManager::UpdateSettingsCategory(TSharedPtr<ISettingsCategory>
 			UE_CLOG(OldSectionRow != NewSectionRow, LogTedsSettings, Verbose, TEXT("Added Settings Section : '%s'"), *SectionName.ToString());
 			UE_CLOG(OldSectionRow == NewSectionRow, LogTedsSettings, Verbose, TEXT("Updated Settings Section : '%s'"), *SectionName.ToString());
 
+			DataStorage->SetParentRow(SettingsHierarchy, NewSectionRow, CategoryRow);						
+			
 			DataStorage->MapRow(Settings::MappingDomain, MoveTemp(SectionIndexKey), NewSectionRow);
+
+			if (OldSectionRow != NewSectionRow)
+			{
+				for (FProperty* Property = SettingsObjectPtr->GetClass()->PropertyLink; Property; Property = Property->PropertyLinkNext)
+				{
+					if (Property->HasAnyPropertyFlags(CPF_Transient) || !Property->HasAnyPropertyFlags(CPF_Edit))
+					{
+						continue;
+					}
+
+					RowHandle NewPropertyRow = DataStorage->AddRow(SettingsPropertyTable);
+
+					DataStorage->AddColumn<FSettingsPropertyTag>(NewPropertyRow);
+					DataStorage->AddColumn<FSettingsNameColumn>(NewPropertyRow, { .Name = Property->GetFName() });
+					DataStorage->AddColumn<FDisplayNameColumn>(NewPropertyRow, { .DisplayName = Property->GetDisplayNameText() });
+					DataStorage->AddColumn<FDescriptionColumn>(NewPropertyRow, { .Description = Property->GetToolTipText() });
+
+					DataStorage->SetUnresolvedParent(SettingsHierarchy, NewPropertyRow, FMapKey(SettingsObjectPtr.Get()), ICompatibilityProvider::ObjectMappingDomain);
+				}
+			}
 		}
 	}
 
-	// Iterate the old active sections and inactivate rows not in the new active sections list.
+	// Iterate the old sections and remove rows not in the new sections list.
 	for (int32 RowIndex = 0; RowIndex < OldSectionNames.Num(); ++RowIndex)
 	{
 		const FName& OldSectionName = OldSectionNames[RowIndex];
@@ -554,52 +575,28 @@ void FTedsSettingsManager::UpdateSettingsCategory(TSharedPtr<ISettingsCategory>
 		RowHandle OldSectionRow = OldRowHandles[RowIndex];
 		check(OldSectionRow != InvalidRowHandle);
 
-		// Add an inactive row for the settings section.
-		RowHandle NewSectionRow = DataStorage->AddRow(SettingsInactiveSectionTable);
-
-		DataStorage->AddColumn<FSettingsInactiveSectionTag>(NewSectionRow);
-		DataStorage->AddColumn<FSettingsNameColumn>(NewSectionRow, { .Name = OldSectionName });
-		DataStorage->AddColumn<FSettingsContainerReferenceColumn>(NewSectionRow, { .ContainerName = ContainerName, .ContainerRow = InvalidRowHandle });
-		DataStorage->AddColumn<FSettingsCategoryReferenceColumn>(NewSectionRow, { .CategoryName = CategoryName, .CategoryRow = InvalidRowHandle });
-
-		// Capture unknown columns and copy them to the new row before we delete the old row.
-		TArray<void*> ColumnsToCopy;
-		TArray<const UScriptStruct*> ColumnTypesToCopy;		
-		
-		DataStorage->ListColumns(OldSectionRow, [&ColumnsToCopy, &ColumnTypesToCopy](void* Column, const UScriptStruct& ColumnType)
+		// Remove the old section row and any child property rows.
+		FRowHandleArray RowsToRemove;
+		auto VisitRow = [DataStorage, OldSectionRow, &RowsToRemove](const ICoreProvider&, RowHandle, RowHandle TargetRow)
 			{
-				if (IsUnknownColumn(&ColumnType))
+				if (TargetRow != OldSectionRow)
 				{
-					ColumnsToCopy.Add(Column);
-					ColumnTypesToCopy.Add(&ColumnType);
+					RowsToRemove.Add(TargetRow);
 				}
-			});
-
-		AddColumns(*DataStorage, NewSectionRow, ColumnsToCopy, ColumnTypesToCopy);
-
-		DataStorage->RemoveRow(OldSectionRow);
-
-		DataStorage->MapRow(Settings::MappingDomain, GenerateIndexKey(ContainerName, CategoryName, OldSectionName), NewSectionRow);
+			};
 
-		UE_LOG(LogTedsSettings, Verbose, TEXT("Inactivated Settings Section : '%s'"), *OldSectionName.ToString());
-	}
-}
+		DataStorage->WalkDepthFirst(SettingsHierarchy, OldSectionRow, VisitRow);
+		DataStorage->BatchRemoveRows(RowsToRemove.GetRows());
 
-void FTedsSettingsManager::AddColumns(UE::Editor::DataStorage::ICoreProvider& DataStorage, UE::Editor::DataStorage::RowHandle Row, const TArray<void*>& Columns, const TArray<const UScriptStruct*>& ColumnTypes)
-{
-	for (int32 ColumnIndex = 0; ColumnIndex < Columns.Num(); ++ColumnIndex)
-	{
-		void* SourceColumn = Columns[ColumnIndex];
-		const UScriptStruct* SourceColumnType = ColumnTypes[ColumnIndex];
+		if (FTypedElementUObjectColumn* UObjectColumn = DataStorage->GetColumn<FTypedElementUObjectColumn>(OldSectionRow))
+		{
+			DataStorageCompatibility->RemoveCompatibleObject(UObjectColumn->Object);
+		}
+		else
+		{
+			DataStorage->RemoveRow(OldSectionRow);
+		}
 
-		DataStorage.AddColumnData(Row, SourceColumnType,
-			[SourceColumn](void* TargetColumn, const UScriptStruct& ColumnType)
-			{
-				ColumnType.CopyScriptStruct(TargetColumn, SourceColumn);
-			},
-			[](const UScriptStruct& ColumnType, void* Destination, void* Source)
-			{
-				ColumnType.CopyScriptStruct(Destination, Source);
-			});
+		UE_LOG(LogTedsSettings, Verbose, TEXT("Removed Settings Section : '%s'"), *OldSectionName.ToString());
 	}
 }
diff --git a/Engine/Plugins/Experimental/EditorDataStorageFeatures/Source/TedsSettings/Private/TedsSettingsManager.h b/Engine/Plugins/Experimental/EditorDataStorageFeatures/Source/TedsSettings/Private/TedsSettingsManager.h
--- a/Engine/Plugins/Experimental/EditorDataStorageFeatures/Source/TedsSettings/Private/TedsSettingsManager.h
+++ b/Engine/Plugins/Experimental/EditorDataStorageFeatures/Source/TedsSettings/Private/TedsSettingsManager.h
@@ -21,34 +21,35 @@ public:
 	void Initialize();
 	void Shutdown();
 
-	UE::Editor::DataStorage::RowHandle FindSettingsSection(const FName& ContainerName, const FName& CategoryName, const FName& SectionName);
+	UE::Editor::DataStorage::RowHandle FindSettingsContainer(const FName& ContainerName) const;
 
-	UE::Editor::DataStorage::RowHandle FindOrAddSettingsSection(const FName& ContainerName, const FName& CategoryName, const FName& SectionName);
+	UE::Editor::DataStorage::RowHandle FindSettingsCategory(const FName& ContainerName, const FName& CategoryName) const;
+
+	UE::Editor::DataStorage::RowHandle FindSettingsSection(const FName& ContainerName, const FName& CategoryName, const FName& SectionName) const;
 
 	bool GetSettingsSectionFromRow(UE::Editor::DataStorage::RowHandle Row, FName& OutContainerName, FName& OutCategoryName, FName& OutSectionName);
 
+	bool GetSettingsCategoryFromRow(UE::Editor::DataStorage::RowHandle Row, FName& OutContainerName, FName& OutCategoryName);
+
 private:
 
 	void RegisterTables(UE::Editor::DataStorage::ICoreProvider& DataStorage);
+	void RegisterHierarchies(UE::Editor::DataStorage::ICoreProvider& DataStorage);
 
 	void RegisterQueries(UE::Editor::DataStorage::ICoreProvider& DataStorage);
 	void UnregisterQueries(UE::Editor::DataStorage::ICoreProvider& DataStorage);
 
 	void RegisterActiveSettings();
 	void UnregisterActiveSettings();
-	void UnregisterInactiveSettings();
 
 	void RegisterSettingsContainer(const FName& ContainerName);
 
 	void UpdateSettingsCategory(TSharedPtr<ISettingsCategory> SettingsCategory, UE::Editor::DataStorage::RowHandle ContainerRow, const bool bQueryExistingRows = true);
 
-	void AddColumns(UE::Editor::DataStorage::ICoreProvider& DataStorage, UE::Editor::DataStorage::RowHandle Row, const TArray<void*>& Columns, const TArray<const UScriptStruct*>& ColumnTypes);
-
 	bool bIsInitialized;
-	UE::Editor::DataStorage::QueryHandle SelectAllActiveSettingsQuery;
-	UE::Editor::DataStorage::QueryHandle SelectAllInactiveSettingsQuery;
+	UE::Editor::DataStorage::QueryHandle SelectAllSettingsQuery;
 	UE::Editor::DataStorage::TableHandle SettingsContainerTable;
 	UE::Editor::DataStorage::TableHandle SettingsCategoryTable;
-	UE::Editor::DataStorage::TableHandle SettingsInactiveSectionTable;
-
+	UE::Editor::DataStorage::TableHandle SettingsPropertyTable;
+	UE::Editor::DataStorage::FHierarchyHandle SettingsHierarchy;
 };
diff --git a/Engine/Plugins/Experimental/EditorDataStorageFeatures/Source/TedsSettings/Public/TedsSettingsColumns.h b/Engine/Plugins/Experimental/EditorDataStorageFeatures/Source/TedsSettings/Public/TedsSettingsColumns.h
--- a/Engine/Plugins/Experimental/EditorDataStorageFeatures/Source/TedsSettings/Public/TedsSettingsColumns.h
+++ b/Engine/Plugins/Experimental/EditorDataStorageFeatures/Source/TedsSettings/Public/TedsSettingsColumns.h
@@ -10,6 +10,7 @@
 namespace UE::Editor::DataStorage::Settings
 {
 	inline static const FName MappingDomain = "Settings";
+	inline static const FName HierarchyName = "SettingsHierarchy";
 }
 
 USTRUCT(meta = (DisplayName = "Settings Container Reference"))
@@ -65,8 +66,8 @@ struct FSettingsSectionTag final : public FEditorDataStorageTag
 	GENERATED_BODY()
 };
 
-USTRUCT(meta = (DisplayName = "Inactive Settings Section"))
-struct FSettingsInactiveSectionTag final : public FEditorDataStorageTag
+USTRUCT(meta = (DisplayName = "Settings Property"))
+struct FSettingsPropertyTag final : public FEditorDataStorageTag
 {
 	GENERATED_BODY()
 };
diff --git a/Engine/Plugins/Experimental/EditorDataStorageFeatures/Source/TedsSettings/Public/TedsSettingsEditorSubsystem.h b/Engine/Plugins/Experimental/EditorDataStorageFeatures/Source/TedsSettings/Public/TedsSettingsEditorSubsystem.h
--- a/Engine/Plugins/Experimental/EditorDataStorageFeatures/Source/TedsSettings/Public/TedsSettingsEditorSubsystem.h
+++ b/Engine/Plugins/Experimental/EditorDataStorageFeatures/Source/TedsSettings/Public/TedsSettingsEditorSubsystem.h
@@ -25,8 +25,27 @@ public:
 	DECLARE_MULTICAST_DELEGATE(FOnEnabledChanged)
 	UE_API FOnEnabledChanged& OnEnabledChanged();
 
+	/**
+	 * Finds an existing container row.
+	 *
+	 * @param	ContainerName	The ContainerName to search for.
+	 *
+	 * @return	The found row or InvalidRowHandle if no existing row is found.
+	 */
+	UE_API UE::Editor::DataStorage::RowHandle FindSettingsContainer(const FName& ContainerName) const;
+
+	/**
+	 * Finds an existing category row.
+	 *
+	 * @param	ContainerName	The ContainerName to search for.
+	 * @param	CategoryName	The CategoryName to search for.
+	 *
+	 * @return	The found row or InvalidRowHandle if no existing row is found.
+	 */
+	UE_API UE::Editor::DataStorage::RowHandle FindSettingsCategory(const FName& ContainerName, const FName& CategoryName) const;
+
 	/** 
-	 * Finds an existing row (may be active or inactive).
+	 * Finds an existing section row.
 	 * 
 	 * @param	ContainerName	The ContainerName to search for.
 	 * @param	CategoryName	The CategoryName to search for.
@@ -34,30 +53,30 @@ public:
 	 * 
 	 * @return	The found row or InvalidRowHandle if no existing row is found.
 	 */
-	UE_API UE::Editor::DataStorage::RowHandle FindSettingsSection(const FName& ContainerName, const FName& CategoryName, const FName& SectionName);
+	UE_API UE::Editor::DataStorage::RowHandle FindSettingsSection(const FName& ContainerName, const FName& CategoryName, const FName& SectionName) const;
 
 	/**
-	 * Finds an existing row (may be active or inactive) or adds a new inactive settings section row if no existing row is found.
+	 * Gets the settings section details for the given section row.
 	 *
-	 * @param	ContainerName	The ContainerName to search for.
-	 * @param	CategoryName	The CategoryName to search for.
-	 * @param	SectionName		The SectionName to search for.
+	 * @param	SectionRow			The section row from which to read settings details.
+	 * @param	OutContainerName	Reference to the ContainerName for the given row, only set if this method returns true.
+	 * @param	OutCategoryName		Reference to the CategoryName for the given row, only set if this method returns true.
+	 * @param	OutSectionName		Reference to the SectionName for the given row, only set if this method returns true.
 	 *
-	 * @return	The row that was either found or added, or InvalidRowHandle if data storage is not available.
+	 * @return	True if the settings section details are successfully returned in the out parameters.
 	 */
-	UE_API UE::Editor::DataStorage::RowHandle FindOrAddSettingsSection(const FName& ContainerName, const FName& CategoryName, const FName& SectionName);
+	UE_API bool GetSettingsSectionFromRow(UE::Editor::DataStorage::RowHandle SectionRow, FName& OutContainerName, FName& OutCategoryName, FName& OutSectionName);
 
 	/**
-	 * Gets the settings section details for the given row.
+	 * Gets the settings category details for the given category row.
 	 *
-	 * @param	Row					The row from which to read settings details.
+	 * @param	CategoryRow			The category row from which to read settings details.
 	 * @param	OutContainerName	Reference to the ContainerName for the given row, only set if this method returns true.
 	 * @param	OutCategoryName		Reference to the CategoryName for the given row, only set if this method returns true.
-	 * @param	OutSectionName		Reference to the SectionName for the given row, only set if this method returns true.
 	 *
-	 * @return	True if the settings section details are successfully returned in the out parameters.
+	 * @return	True if the settings category details are successfully returned in the out parameters.
 	 */
-	UE_API bool GetSettingsSectionFromRow(UE::Editor::DataStorage::RowHandle Row, FName& OutContainerName, FName& OutCategoryName, FName& OutSectionName);
+	UE_API bool GetSettingsCategoryFromRow(UE::Editor::DataStorage::RowHandle CategoryRow, FName& OutContainerName, FName& OutCategoryName);
 
 protected: // USubsystem
 
diff --git a/Engine/Plugins/Experimental/EditorDataStorageFeatures/Source/TedsSettings/Tests/TedsSettingsTests.cpp b/Engine/Plugins/Experimental/EditorDataStorageFeatures/Source/TedsSettings/Tests/TedsSettingsTests.cpp
--- a/Engine/Plugins/Experimental/EditorDataStorageFeatures/Source/TedsSettings/Tests/TedsSettingsTests.cpp
+++ b/Engine/Plugins/Experimental/EditorDataStorageFeatures/Source/TedsSettings/Tests/TedsSettingsTests.cpp
@@ -87,7 +87,7 @@ namespace UE::Editor::Settings::Tests
 				CountAllSettingsQuery = DataStorage->RegisterQuery(
 					Count()
 					.Where()
-						.Any<FSettingsSectionTag, FSettingsInactiveSectionTag>()
+						.Any<FSettingsSectionTag>()
 					.Compile());
 			}
 
@@ -145,7 +145,7 @@ namespace UE::Editor::Settings::Tests
 
 		Describe("UnregisterSettings", [this]()
 		{
-			LatentIt("Should inactivate a row in editor data storage", [this](const FDoneDelegate& Done)
+			LatentIt("Should remove a row from editor data storage", [this](const FDoneDelegate& Done)
 			{
 				const FName& ContainerName = FName(TEXT("TestContainer"));
 				const FName& CategoryName = FName(TEXT("TestCategory"));
@@ -175,7 +175,7 @@ namespace UE::Editor::Settings::Tests
 					SettingsModule->UnregisterSettings(ContainerName, CategoryName, SectionName);
 
 					uint32 AfterUnregisterRowCount = CountSettingsRowsInDataStorage();
-					TestEqual(TEXT("RowCount"), AfterUnregisterRowCount, BeforeRowCount + 1);
+					TestEqual(TEXT("RowCount"), AfterUnregisterRowCount, BeforeRowCount);
 
 					TestFalse(TEXT("IsRowAssigned"), DataStorage->IsRowAssigned(RowHandle));
 
@@ -184,15 +184,11 @@ namespace UE::Editor::Settings::Tests
 
 					UTedsSettingsEditorSubsystem* SettingsEditorSubsystem = GEditor->GetEditorSubsystem<UTedsSettingsEditorSubsystem>();
 
-					DataStorage::RowHandle InactiveRowHandle = SettingsEditorSubsystem->FindOrAddSettingsSection(ContainerName, CategoryName, SectionName);
-					TestNotEqual(TEXT("InactiveRowHandle"), InactiveRowHandle, DataStorage::InvalidRowHandle);
-
-					TestRowHandles.Push(InactiveRowHandle);
+					TestInvalidRowHandle = SettingsEditorSubsystem->FindSettingsSection(ContainerName, CategoryName, SectionName);
+					TestEqual(TEXT("InvalidRowHandle"), TestInvalidRowHandle, DataStorage::InvalidRowHandle);
 
 					uint32 AfterFindOrAddRowCount = CountSettingsRowsInDataStorage();
-					TestEqual(TEXT("RowCount"), AfterFindOrAddRowCount, BeforeRowCount + 1);
-
-					TestTrue(TEXT("FSettingsInactiveSectionTag"), DataStorage->HasColumns<FSettingsInactiveSectionTag>(InactiveRowHandle));
+					TestEqual(TEXT("RowCount"), AfterFindOrAddRowCount, BeforeRowCount);
 				});
 			});
 		});
@@ -242,96 +238,6 @@ namespace UE::Editor::Settings::Tests
 				});
 			});
 		});
-
-		Describe("FindOrAddSettingsSection with no existing active rows", [this]()
-		{
-			LatentIt("Should result in only a single inactive row in editor data storage", [this](const FDoneDelegate& Done)
-			{
-				UTedsSettingsEditorSubsystem* SettingsEditorSubsystem = GEditor->GetEditorSubsystem<UTedsSettingsEditorSubsystem>();
-
-				const FName& ContainerName = FName(TEXT("TestContainer"));
-				const FName& CategoryName = FName(TEXT("TestCategory"));
-				const FName& SectionName = FName(TEXT("TestSection"));
-
-				DataStorage::RowHandle RowHandle = SettingsEditorSubsystem->FindOrAddSettingsSection(ContainerName, CategoryName, SectionName);
-				TestNotEqual(TEXT("RowHandle"), RowHandle, DataStorage::InvalidRowHandle);
-
-				if (RowHandle == DataStorage::InvalidRowHandle)
-				{
-					Done.Execute();
-					return;
-				}
-
-				TestRowHandles.Push(RowHandle);
-
-				AwaitRowHandleThenVerify(RowHandle, Done, [this, RowHandle, ContainerName, CategoryName, SectionName]()
-				{
-					uint32 AfterRowCount = CountSettingsRowsInDataStorage();
-
-					TestEqual(TEXT("RowCount"), AfterRowCount, BeforeRowCount + 1);
-
-					TestEqual(TEXT("ContainerName"), DataStorage->GetColumn<FSettingsContainerReferenceColumn>(RowHandle)->ContainerName, ContainerName);
-					TestEqual(TEXT("CategoryName"), DataStorage->GetColumn<FSettingsCategoryReferenceColumn>(RowHandle)->CategoryName, CategoryName);
-					TestEqual(TEXT("SectionName"), DataStorage->GetColumn<FSettingsNameColumn>(RowHandle)->Name, SectionName);
-
-					TestTrue(TEXT("FSettingsInactiveSectionTag"), DataStorage->HasColumns<FSettingsInactiveSectionTag>(RowHandle));
-				});
-			});
-		});
-
-		Describe("RegisterSettings with an existing inactive row for same container/category/section", [this]()
-		{
-			LatentIt("Should result in only a single active row in editor data storage", [this](const FDoneDelegate& Done)
-			{
-				UTedsSettingsEditorSubsystem* SettingsEditorSubsystem = GEditor->GetEditorSubsystem<UTedsSettingsEditorSubsystem>();
-
-				const FName& ContainerName = FName(TEXT("TestContainer"));
-				const FName& CategoryName = FName(TEXT("TestCategory"));
-				const FName& SectionName = FName(TEXT("TestSection"));
-
-				DataStorage::RowHandle RowHandle = SettingsEditorSubsystem->FindOrAddSettingsSection(ContainerName, CategoryName, SectionName);
-				TestNotEqual(TEXT("RowHandle"), RowHandle, DataStorage::InvalidRowHandle);
-
-				if (RowHandle == DataStorage::InvalidRowHandle)
-				{
-					Done.Execute();
-					return;
-				}
-
-				TestRowHandles.Push(RowHandle);
-
-				TestTrue(TEXT("FSettingsInactiveSectionTag"), DataStorage->HasColumns<FSettingsInactiveSectionTag>(RowHandle));
-
-				TObjectPtr<UTestSettings> TestSettingsObject = NewObject<UTestSettings>();
-
-				SettingsModule->RegisterSettings(ContainerName, CategoryName, SectionName, FText(), FText(), TestSettingsObject);
-
-				RowHandle = DataStorageCompatibility->FindRowWithCompatibleObject(TestSettingsObject);
-				TestNotEqual(TEXT("RowHandle"), RowHandle, DataStorage::InvalidRowHandle);
-
-				if (RowHandle == DataStorage::InvalidRowHandle)
-				{
-					Done.Execute();
-					return;
-				}
-
-				TestRowHandles.Push(RowHandle);
-
-				AwaitRowHandleThenVerify(RowHandle, Done, [this, RowHandle, ContainerName, CategoryName, SectionName]()
-				{
-					uint32 AfterRowCount = CountSettingsRowsInDataStorage();
-
-					TestEqual(TEXT("RowCount"), AfterRowCount, BeforeRowCount + 1);
-
-					TestEqual(TEXT("ContainerName"), DataStorage->GetColumn<FSettingsContainerReferenceColumn>(RowHandle)->ContainerName, ContainerName);
-					TestEqual(TEXT("CategoryName"), DataStorage->GetColumn<FSettingsCategoryReferenceColumn>(RowHandle)->CategoryName, CategoryName);
-					TestEqual(TEXT("SectionName"), DataStorage->GetColumn<FSettingsNameColumn>(RowHandle)->Name, SectionName);
-
-					TestFalse(TEXT("FSettingsInactiveSectionTag"), DataStorage->HasColumns<FSettingsInactiveSectionTag>(RowHandle));
-					TestTrue(TEXT("FSettingsSectionTag"), DataStorage->HasColumns<FSettingsSectionTag>(RowHandle));
-				});
-			});
-		});
 	}
 } // namespace UE::Editor::Settings::Tests
 
