TedsSettings : use SetParentRow on the property rows.

I meant to make this change before submitting 49045812 so this is just a quick followup to that CL.

#Synced-CL 49045812

diff --git a/Engine/Plugins/Experimental/EditorDataStorageFeatures/Source/TedsSettings/Private/TedsSettingsManager.cpp b/Engine/Plugins/Experimental/EditorDataStorageFeatures/Source/TedsSettings/Private/TedsSettingsManager.cpp
--- a/Engine/Plugins/Experimental/EditorDataStorageFeatures/Source/TedsSettings/Private/TedsSettingsManager.cpp
+++ b/Engine/Plugins/Experimental/EditorDataStorageFeatures/Source/TedsSettings/Private/TedsSettingsManager.cpp
@@ -556,7 +556,7 @@ void FTedsSettingsManager::UpdateSettingsCategory(TSharedPtr<ISettingsCategory>
 					DataStorage->AddColumn<FDisplayNameColumn>(NewPropertyRow, { .DisplayName = Property->GetDisplayNameText() });
 					DataStorage->AddColumn<FDescriptionColumn>(NewPropertyRow, { .Description = Property->GetToolTipText() });
 
-					DataStorage->SetUnresolvedParent(SettingsHierarchy, NewPropertyRow, FMapKey(SettingsObjectPtr.Get()), ICompatibilityProvider::ObjectMappingDomain);
+					DataStorage->SetParentRow(SettingsHierarchy, NewPropertyRow, NewSectionRow);
 				}
 			}
 		}
