[TEDS] Add option to not include top level rows in HierarchyNode output

#Synced-CL 48874164

diff --git a/Engine/Plugins/Experimental/EditorDataStorageFeatures/Source/TedsQueryStack/Private/TedsHierarchyNode.cpp b/Engine/Plugins/Experimental/EditorDataStorageFeatures/Source/TedsQueryStack/Private/TedsHierarchyNode.cpp
--- a/Engine/Plugins/Experimental/EditorDataStorageFeatures/Source/TedsQueryStack/Private/TedsHierarchyNode.cpp
+++ b/Engine/Plugins/Experimental/EditorDataStorageFeatures/Source/TedsQueryStack/Private/TedsHierarchyNode.cpp
@@ -11,11 +11,13 @@ namespace UE::Editor::DataStorage::QueryStack
 		ICoreProvider& Storage,
 		FHierarchyHandle InHierarchyHandle,
 		TSharedPtr<QueryStack::IRowNode>& InTopLevelRows,
-		ESyncFlags InSyncFlags)
+		ESyncFlags InSyncFlags,
+		EFilterFlags InFilterFlags)
 		: TopLevelRows(InTopLevelRows)
 		, Storage(Storage)
 		, HierarchyHandle(InHierarchyHandle)
 		, SyncFlags(InSyncFlags)
+		, FilterFlags(InFilterFlags)
 	{
 	}
 
@@ -87,8 +89,13 @@ namespace UE::Editor::DataStorage::QueryStack
 	{
 		InRows.Empty();
 		
-		auto VisitFn = [&InRows](const ICoreProvider&, RowHandle, RowHandle TargetRow)
+		auto VisitFn = [&InRows, FilterFlags_ = FilterFlags](const ICoreProvider&, RowHandle ParentRow, RowHandle TargetRow)
 		{
+			if ((FilterFlags_ == EFilterFlags::OnlyIncludeChildRows) && (ParentRow == InvalidRowHandle))
+			{
+				return;
+			}
+
 			InRows.Add(TargetRow);
 		};
 
diff --git a/Engine/Plugins/Experimental/EditorDataStorageFeatures/Source/TedsQueryStack/Public/TedsHierarchyNode.h b/Engine/Plugins/Experimental/EditorDataStorageFeatures/Source/TedsQueryStack/Public/TedsHierarchyNode.h
--- a/Engine/Plugins/Experimental/EditorDataStorageFeatures/Source/TedsQueryStack/Public/TedsHierarchyNode.h
+++ b/Engine/Plugins/Experimental/EditorDataStorageFeatures/Source/TedsQueryStack/Public/TedsHierarchyNode.h
@@ -26,13 +26,27 @@ namespace UE::Editor::DataStorage::QueryStack
 			 */
 			Always
 		};
+
+		enum class EFilterFlags
+		{
+			/**
+			 * No filtering will be applied
+			 */
+			None = 0,
+			/**
+			 * Root parent rows will not be included in the row list
+			 */
+			OnlyIncludeChildRows = 1,
+		};
+
 		using EmitRowFn = TFunctionRef<void(TArrayView<const RowHandle>)>;
 		
 		TEDSQUERYSTACK_API FHierarchyRowNode(
 			ICoreProvider& Storage,
 			FHierarchyHandle InHierarchyHandle,
 			TSharedPtr<QueryStack::IRowNode>& InTopLevelRows,
-			ESyncFlags SyncFlags);
+			ESyncFlags SyncFlags,
+			EFilterFlags FilterFlags = EFilterFlags::None);
 		
 		TEDSQUERYSTACK_API virtual RevisionId GetRevision() const override;
 		TEDSQUERYSTACK_API virtual void VisitParents(ParentListCallback Callback) override;
@@ -56,6 +70,7 @@ namespace UE::Editor::DataStorage::QueryStack
 		
 		RevisionId Revision = UninitializedRevisionId;
 		ESyncFlags SyncFlags;
+		EFilterFlags FilterFlags;
 	};
 
 	ENUM_CLASS_FLAGS(FHierarchyRowNode::ESyncFlags);
