TedsSettings : workaround a timing issue with calling SetParentRow where the target (child) row is the result of a call to AddCompatibleObject.

We need to wait until both rows are assigned before establishing the relationship.

#Synced-CL 49050857

diff --git a/Engine/Plugins/Experimental/EditorDataStorageFeatures/Source/TedsSettings/Private/TedsSettingsManager.cpp b/Engine/Plugins/Experimental/EditorDataStorageFeatures/Source/TedsSettings/Private/TedsSettingsManager.cpp
--- a/Engine/Plugins/Experimental/EditorDataStorageFeatures/Source/TedsSettings/Private/TedsSettingsManager.cpp
+++ b/Engine/Plugins/Experimental/EditorDataStorageFeatures/Source/TedsSettings/Private/TedsSettingsManager.cpp
@@ -2,6 +2,7 @@
 
 #include "TedsSettingsManager.h"
 
+#include "Containers/Ticker.h"
 #include "DataStorage/Features.h"
 #include "Elements/Columns/TypedElementCompatibilityColumns.h"
 #include "Elements/Columns/TypedElementMiscColumns.h"
@@ -536,7 +537,7 @@ void FTedsSettingsManager::UpdateSettingsCategory(TSharedPtr<ISettingsCategory>
 			UE_CLOG(OldSectionRow != NewSectionRow, LogTedsSettings, Verbose, TEXT("Added Settings Section : '%s'"), *SectionName.ToString());
 			UE_CLOG(OldSectionRow == NewSectionRow, LogTedsSettings, Verbose, TEXT("Updated Settings Section : '%s'"), *SectionName.ToString());
 
-			DataStorage->SetParentRow(SettingsHierarchy, NewSectionRow, CategoryRow);						
+			SetParentRowOnceBothRowsAreAssigned(SettingsHierarchy, NewSectionRow, CategoryRow);
 			
 			DataStorage->MapRow(Settings::MappingDomain, MoveTemp(SectionIndexKey), NewSectionRow);
 
@@ -600,3 +601,41 @@ void FTedsSettingsManager::UpdateSettingsCategory(TSharedPtr<ISettingsCategory>
 		UE_LOG(LogTedsSettings, Verbose, TEXT("Removed Settings Section : '%s'"), *OldSectionName.ToString());
 	}
 }
+
+void FTedsSettingsManager::SetParentRowOnceBothRowsAreAssigned(
+	UE::Editor::DataStorage::FHierarchyHandle Hierarchy,
+	UE::Editor::DataStorage::RowHandle Target,
+	UE::Editor::DataStorage::RowHandle Parent)
+{
+	using namespace UE::Editor::DataStorage;
+
+	TRACE_CPUPROFILER_EVENT_SCOPE(TedsSettingsManager.SetParentRowOnceBothRowsAreAssigned);
+
+	ICoreProvider* DataStorage = GetMutableDataStorageFeature<ICoreProvider>(StorageFeatureName);
+
+	// If the hierarchy is no longer valid or either row is no longer available, then we abort the attempt to establish the relationship.
+	if (!DataStorage ||
+		!DataStorage->IsValidHierarchyHandle(Hierarchy) ||
+		!DataStorage->IsRowAvailable(Target) ||
+		!DataStorage->IsRowAvailable(Parent))
+	{
+		return;
+	}
+
+	// If both rows are assigned, then it is now safe to establish the relationship.
+	if (DataStorage->IsRowAssigned(Target) && DataStorage->IsRowAssigned(Parent))
+	{
+		DataStorage->SetParentRow(Hierarchy, Target, Parent);
+	}
+	// Else one of the rows must be available but not yet assigned, so defer another tick and try again.
+	else
+	{
+		ExecuteOnGameThread(UE_SOURCE_LOCATION, [WeakThis = SharedThis(this).ToWeakPtr(), Hierarchy, Target, Parent]()
+		{
+			if (TSharedPtr<FTedsSettingsManager> This = WeakThis.Pin())
+			{
+				This->SetParentRowOnceBothRowsAreAssigned(Hierarchy, Target, Parent);
+			}
+		});
+	}
+}
diff --git a/Engine/Plugins/Experimental/EditorDataStorageFeatures/Source/TedsSettings/Private/TedsSettingsManager.h b/Engine/Plugins/Experimental/EditorDataStorageFeatures/Source/TedsSettings/Private/TedsSettingsManager.h
--- a/Engine/Plugins/Experimental/EditorDataStorageFeatures/Source/TedsSettings/Private/TedsSettingsManager.h
+++ b/Engine/Plugins/Experimental/EditorDataStorageFeatures/Source/TedsSettings/Private/TedsSettingsManager.h
@@ -46,6 +46,13 @@ private:
 
 	void UpdateSettingsCategory(TSharedPtr<ISettingsCategory> SettingsCategory, UE::Editor::DataStorage::RowHandle ContainerRow, const bool bQueryExistingRows = true);
 
+	// This is a workaround for an issue where calling SetParentRow with a not-yet-assigned Target (child) row
+	// sometimes fails to establish the relationship.
+	void SetParentRowOnceBothRowsAreAssigned(
+		UE::Editor::DataStorage::FHierarchyHandle Hierarchy,
+		UE::Editor::DataStorage::RowHandle Target,
+		UE::Editor::DataStorage::RowHandle Parent);
+
 	bool bIsInitialized;
 	UE::Editor::DataStorage::QueryHandle SelectAllSettingsQuery;
 	UE::Editor::DataStorage::TableHandle SettingsContainerTable;
