Teds Debugger: use ForceRefreshOnUpdate for the hierarchy viewer and RefreshOnUpdate for the other views. We currently cannot detect hierarchy changes using the caching mechanism so this leads to issues. Worst case is crashes in slate code since UI doesn't match data, "best" case is incorrect data.

#Synced-CL 49034284

diff --git a/Engine/Plugins/Experimental/EditorDataStorageFeatures/Source/TedsDebugger/Private/Widgets/QueryEditor/TedsQueryEditorResultsView.cpp b/Engine/Plugins/Experimental/EditorDataStorageFeatures/Source/TedsDebugger/Private/Widgets/QueryEditor/TedsQueryEditorResultsView.cpp
--- a/Engine/Plugins/Experimental/EditorDataStorageFeatures/Source/TedsDebugger/Private/Widgets/QueryEditor/TedsQueryEditorResultsView.cpp
+++ b/Engine/Plugins/Experimental/EditorDataStorageFeatures/Source/TedsDebugger/Private/Widgets/QueryEditor/TedsQueryEditorResultsView.cpp
@@ -73,10 +73,10 @@ namespace UE::Editor::DataStorage::Debug::QueryEditor
 					+SVerticalBox::Slot()
 					.AutoHeight()
 					[
-						SNew(STedsSearchBox)
-						.InSearchableQueryNode(QueryNode)
-						.InSearchableQueryFlags(STedsSearchBox::ESearchableQueryFlags::RefreshOnUpdate)
-						.OutSearchNode(&RowQueryStack)
+						SAssignNew(SearchBoxSlot, SBorder)
+						[
+							CreateSearchBox()
+						]
 					]
 					+SVerticalBox::Slot()
 					.FillContentHeight(1.0f)
@@ -186,6 +186,7 @@ namespace UE::Editor::DataStorage::Debug::QueryEditor
 	{
 		ResultsViewMode = InViewMode;
 		TableViewerSlot->SetContent(CreateTableViewer()); // Regenerate the table viewer widget
+		SearchBoxSlot->SetContent(CreateSearchBox()); // Regenerate the search box since the sync flag depends on the view mode
 	}
 
 	void SResultsView::OnModelChanged()
@@ -285,6 +286,20 @@ namespace UE::Editor::DataStorage::Debug::QueryEditor
 		
 		return TableViewer->AsWidget();
 	}
+
+	TSharedRef<SWidget> SResultsView::CreateSearchBox()
+	{
+		// We cannot detect hierarchy changes by checking the cached state since they are data changes, so we use force refresh when viewing hierarchies
+		STedsSearchBox::ESearchableQueryFlags QueryFlags =
+			(ResultsViewMode == ETableViewMode::Tree)
+			? STedsSearchBox::ESearchableQueryFlags::ForceRefreshOnUpdate 
+			: STedsSearchBox::ESearchableQueryFlags::RefreshOnUpdate;
+		
+		return SNew(STedsSearchBox)
+			.InSearchableQueryNode(QueryNode)
+			.InSearchableQueryFlags(QueryFlags)
+			.OutSearchNode(&RowQueryStack);
+	}
 } // namespace UE::Editor::DataStorage::Debug::QueryEditor
 
 #undef LOCTEXT_NAMESPACE
diff --git a/Engine/Plugins/Experimental/EditorDataStorageFeatures/Source/TedsDebugger/Private/Widgets/QueryEditor/TedsQueryEditorResultsView.h b/Engine/Plugins/Experimental/EditorDataStorageFeatures/Source/TedsDebugger/Private/Widgets/QueryEditor/TedsQueryEditorResultsView.h
--- a/Engine/Plugins/Experimental/EditorDataStorageFeatures/Source/TedsDebugger/Private/Widgets/QueryEditor/TedsQueryEditorResultsView.h
+++ b/Engine/Plugins/Experimental/EditorDataStorageFeatures/Source/TedsDebugger/Private/Widgets/QueryEditor/TedsQueryEditorResultsView.h
@@ -51,6 +51,7 @@ namespace UE::Editor::DataStorage
 			void OnHierarchyChanged();
 			void CreateRowHandleColumn();
 			TSharedRef<SWidget> CreateTableViewer();
+			TSharedRef<SWidget> CreateSearchBox();
 		
 			FTedsQueryEditorModel* Model = nullptr;
 			FDelegateHandle ModelChangedDelegateHandle;
@@ -81,6 +82,9 @@ namespace UE::Editor::DataStorage
 			// Border inside which the table viewer is stored so we can re-create it when needed
 			TSharedPtr<SBorder> TableViewerSlot;
 
+			// Border inside which the search box is stored so we can re-create it when needed
+			TSharedPtr<SBorder> SearchBoxSlot;
+
 			ETableViewMode::Type ResultsViewMode = ETableViewMode::List;
 		};
 
