TEDS Outliner: Replace the query description in FTedsOutlinerColumnQueryColumn with a query handle instead
Also modify the creation param and few other classes in the Teds Outliner to use query handles instead of descriptions for column info. The description is now grabbed on demand.

#Synced-CL 49084582

diff --git a/Engine/Plugins/ChaosVD/Source/ChaosVD/Private/ChaosVDWorldOutlinerTab.cpp b/Engine/Plugins/ChaosVD/Source/ChaosVD/Private/ChaosVDWorldOutlinerTab.cpp
--- a/Engine/Plugins/ChaosVD/Source/ChaosVD/Private/ChaosVDWorldOutlinerTab.cpp
+++ b/Engine/Plugins/ChaosVD/Source/ChaosVD/Private/ChaosVDWorldOutlinerTab.cpp
@@ -54,10 +54,13 @@ void FChaosVDWorldOutlinerTab::CreateWorldOutlinerWidget()
 	Params.bShowRowHandleColumn = false;
 	Params.bUseDefaultObservers = false; // CVD uses custom observers to time slice addition and removal currently
 
-	Params.ColumnQueryDescription =
-		Select()
-			.ReadOnly<FVisibleInEditorColumn>()
-		.Compile();
+	if (ICoreProvider* Storage = GetMutableDataStorageFeature<ICoreProvider>(StorageFeatureName))
+	{
+		Params.ColumnQuery = Storage->RegisterQuery(
+			Select()
+				.ReadOnly<FVisibleInEditorColumn>()
+			.Compile());
+	}
 
 	const FTedsOutlinerHierarchyData::FGetParentRowHandle RowHandleGetter = FTedsOutlinerHierarchyData::FGetParentRowHandle::CreateLambda([](const void* InColumnData)
 	{
diff --git a/Engine/Plugins/Experimental/Dataflow/Source/DataflowEditor/Private/Dataflow/DataflowOutlinerView.cpp b/Engine/Plugins/Experimental/Dataflow/Source/DataflowEditor/Private/Dataflow/DataflowOutlinerView.cpp
--- a/Engine/Plugins/Experimental/Dataflow/Source/DataflowEditor/Private/Dataflow/DataflowOutlinerView.cpp
+++ b/Engine/Plugins/Experimental/Dataflow/Source/DataflowEditor/Private/Dataflow/DataflowOutlinerView.cpp
@@ -38,7 +38,7 @@ TSharedPtr<ISceneOutliner> FDataflowOutlinerView::CreateWidget()
 		.Where(TColumn<FDataflowSceneObjectTag>(GetEditorContent()->GetDataflowOwner().GetFName()) || TColumn<FDataflowSceneStructTag>(GetEditorContent()->GetDataflowOwner().GetFName()) )
 		.Compile();
 
-	const UE::Editor::DataStorage::FQueryDescription ColumnQueryDescription =
+	UE::Editor::DataStorage::FQueryDescription ColumnQueryDescription =
 		Select()
 		.ReadOnly<FTypedElementClassTypeInfoColumn, FVisibleInEditorColumn>()
 		.Compile();
@@ -52,7 +52,7 @@ TSharedPtr<ISceneOutliner> FDataflowOutlinerView::CreateWidget()
 	UE::Editor::Outliner::FTedsOutlinerParams Params(nullptr);
 	Params.QueryDescription = RowQueryDescription;
 	Params.bShowRowHandleColumn = false;
-	Params.ColumnQueryDescription = ColumnQueryDescription;
+	Params.ColumnQuery = Storage->RegisterQuery(MoveTemp(ColumnQueryDescription));
 	
 	// Add outliner filter queries
 	Params.Filters.Emplace("Dataflow Construction",
diff --git a/Engine/Plugins/Experimental/EditorDataStorageFeatures/Source/TedsOutliner/Private/TedsOutlinerColumn.cpp b/Engine/Plugins/Experimental/EditorDataStorageFeatures/Source/TedsOutliner/Private/TedsOutlinerColumn.cpp
--- a/Engine/Plugins/Experimental/EditorDataStorageFeatures/Source/TedsOutliner/Private/TedsOutlinerColumn.cpp
+++ b/Engine/Plugins/Experimental/EditorDataStorageFeatures/Source/TedsOutliner/Private/TedsOutlinerColumn.cpp
@@ -11,7 +11,7 @@
 
 namespace UE::Editor::Outliner
 {
-	FTedsOutlinerUiColumnInitParams::FTedsOutlinerUiColumnInitParams(DataStorage::FQueryDescription& InQuery, DataStorage::ICoreProvider& InStorage,
+	FTedsOutlinerUiColumnInitParams::FTedsOutlinerUiColumnInitParams(DataStorage::QueryHandle InQuery, DataStorage::ICoreProvider& InStorage,
 		DataStorage::IUiProvider& InStorageUi, DataStorage::ICompatibilityProvider& InStorageCompatibility)
 		: Query(InQuery)
 		, Storage(InStorage)
@@ -24,7 +24,7 @@ namespace UE::Editor::Outliner
 		: Storage(InitParams.Storage)
 		, StorageUi(InitParams.StorageUi)
 		, StorageCompatibility(InitParams.StorageCompatibility)
-		, QueryDescription(InitParams.Query)
+		, Query(InitParams.Query)
 		, NameId(InitParams.NameId)
 		, FallbackColumn(InitParams.FallbackColumn)
 		, OwningOutliner(InitParams.OwningOutliner)
@@ -36,6 +36,8 @@ namespace UE::Editor::Outliner
 
 		using namespace UE::Editor::DataStorage;
 
+		FQueryDescription QueryDescription = Storage.GetQueryDescription(Query);
+
 		TableViewerColumnImpl = MakeUnique<FTedsTableViewerColumn>(NameId, InitParams.CellWidgetConstructor, InitParams.ColumnTypes,
 			InitParams.HeaderWidgetConstructor, QueryDescription.MetaData);
 
diff --git a/Engine/Plugins/Experimental/EditorDataStorageFeatures/Source/TedsOutliner/Private/TedsOutlinerColumn.h b/Engine/Plugins/Experimental/EditorDataStorageFeatures/Source/TedsOutliner/Private/TedsOutlinerColumn.h
--- a/Engine/Plugins/Experimental/EditorDataStorageFeatures/Source/TedsOutliner/Private/TedsOutlinerColumn.h
+++ b/Engine/Plugins/Experimental/EditorDataStorageFeatures/Source/TedsOutliner/Private/TedsOutlinerColumn.h
@@ -14,11 +14,11 @@ namespace UE::Editor::Outliner
 
 	struct FTedsOutlinerUiColumnInitParams
 	{
-		FTedsOutlinerUiColumnInitParams(DataStorage::FQueryDescription& InQuery, DataStorage::ICoreProvider& InStorage,
+		FTedsOutlinerUiColumnInitParams(DataStorage::QueryHandle InQuery, DataStorage::ICoreProvider& InStorage,
 			DataStorage::IUiProvider& InStorageUi, DataStorage::ICompatibilityProvider& InStorageCompatibility);
 		
 		// The query used for this column
-		DataStorage::FQueryDescription& Query;
+		DataStorage::QueryHandle Query;
 
 		// Data storage interfaces
 		DataStorage::ICoreProvider& Storage;
@@ -74,7 +74,7 @@ namespace UE::Editor::Outliner
 		DataStorage::ICoreProvider& Storage;
 		DataStorage::IUiProvider& StorageUi;
 		DataStorage::ICompatibilityProvider& StorageCompatibility;
-		DataStorage::FQueryDescription QueryDescription;
+		DataStorage::QueryHandle Query;
 		DataStorage::FMetaData MetaData;
 		FName NameId;
 		TSharedPtr<ISceneOutlinerColumn> FallbackColumn;
diff --git a/Engine/Plugins/Experimental/EditorDataStorageFeatures/Source/TedsOutliner/Private/TedsOutlinerImpl.cpp b/Engine/Plugins/Experimental/EditorDataStorageFeatures/Source/TedsOutliner/Private/TedsOutlinerImpl.cpp
--- a/Engine/Plugins/Experimental/EditorDataStorageFeatures/Source/TedsOutliner/Private/TedsOutlinerImpl.cpp
+++ b/Engine/Plugins/Experimental/EditorDataStorageFeatures/Source/TedsOutliner/Private/TedsOutlinerImpl.cpp
@@ -349,6 +349,8 @@ FTedsOutlinerImpl::~FTedsOutlinerImpl()
 {
 	if (Storage)
 	{
+		// We take ownership of the column query and so have to handle unregistering it
+		Storage->UnregisterQuery(CreationParams.ColumnQuery);
 		Storage->RemoveRow(OutlinerRowHandle);
 		Storage->OnUpdateCompleted().RemoveAll(this);
 	}
@@ -489,7 +491,7 @@ void FTedsOutlinerImpl::RegisterTedsOutliner()
 		
 		if (FTedsOutlinerColumnQueryColumn* ColumnQueryColumn = Storage->GetColumn<FTedsOutlinerColumnQueryColumn>(OutlinerRowHandle))
 		{
-			ColumnQueryColumn->ColumnQueryDescription = CreationParams.ColumnQueryDescription;
+			ColumnQueryColumn->ColumnQueryHandle = CreationParams.ColumnQuery;
 
 			ColumnQueryColumn->OnRefreshColumns.BindSP(this, &FTedsOutlinerImpl::RegenerateColumns);
 		}
@@ -538,13 +540,15 @@ void FTedsOutlinerImpl::RegenerateColumns()
 	using MatchApproach = DataStorage::IUiProvider::EMatchApproach;
 	constexpr uint8 DefaultPriorityIndex = 100;
 
-	DataStorage::FQueryMetaDataView MetaDataView(ColumnQueryColumn->ColumnQueryDescription);
+	FQueryDescription ColumnQueryDescription = Storage->GetQueryDescription(ColumnQueryColumn->ColumnQueryHandle);
 	
-	int32 SelectionCount = ColumnQueryColumn->ColumnQueryDescription.SelectionTypes.Num();
+	DataStorage::FQueryMetaDataView MetaDataView(ColumnQueryDescription);
+	
+	int32 SelectionCount = ColumnQueryDescription.SelectionTypes.Num();
 	AddedColumns.Reset(SelectionCount);
 
 	TArray<TWeakObjectPtr<const UScriptStruct>> ColumnTypes =
-		DataStorage::TableViewerUtils::CreateVerifiedColumnTypeArray(ColumnQueryColumn->ColumnQueryDescription.SelectionTypes);
+		DataStorage::TableViewerUtils::CreateVerifiedColumnTypeArray(ColumnQueryDescription.SelectionTypes);
 	
 	DataStorage::RowHandle CellPurposeRow = StorageUi->FindPurpose(CellWidgetPurpose);
 	DataStorage::RowHandle HeaderPurposeRow = StorageUi->FindPurpose(HeaderWidgetPurpose);
@@ -591,7 +595,7 @@ void FTedsOutlinerImpl::RegenerateColumns()
 								.IsSortingHandler = DataStorage::FTedsTableViewerColumn::FIsSorting::CreateSP(this, &FTedsOutlinerImpl::IsSorting)
 							};
 
-							FTedsOutlinerUiColumnInitParams Params(ColumnQueryColumn->ColumnQueryDescription, *Storage, *StorageUi, *StorageCompatibility);
+							FTedsOutlinerUiColumnInitParams Params(ColumnQueryColumn->ColumnQueryHandle, *Storage, *StorageUi, *StorageCompatibility);
 							Params.NameId = NameId;
 							Params.ColumnTypes = TArray<TWeakObjectPtr<const UScriptStruct>>(ColumnTypes.GetData(), ColumnTypes.Num());
 							Params.HeaderWidgetConstructor = MoveTemp(HeaderConstructor);
diff --git a/Engine/Plugins/Experimental/EditorDataStorageFeatures/Source/TedsOutliner/Private/TedsOutlinerModule.cpp b/Engine/Plugins/Experimental/EditorDataStorageFeatures/Source/TedsOutliner/Private/TedsOutlinerModule.cpp
--- a/Engine/Plugins/Experimental/EditorDataStorageFeatures/Source/TedsOutliner/Private/TedsOutlinerModule.cpp
+++ b/Engine/Plugins/Experimental/EditorDataStorageFeatures/Source/TedsOutliner/Private/TedsOutlinerModule.cpp
@@ -22,6 +22,7 @@
 #include "Modules/ModuleManager.h"
 #include "Compatibility/SceneOutlinerTedsBridge.h"
 #include "Compatibility/SceneOutlinerRowHandleColumn.h"
+#include "Columns/ActorFolderColumns.h"
 #include "Widgets/Docking/SDockTab.h"
 
 #include "Engine/DirectionalLight.h"
@@ -118,7 +119,7 @@ void FTedsOutlinerModule::ShutdownModule()
 	IModuleInterface::ShutdownModule();
 }
 
-UE::Editor::DataStorage::FQueryDescription FTedsOutlinerModule::GetLevelEditorTedsOutlinerColumnQueryDescription()
+QueryHandle FTedsOutlinerModule::GetLevelEditorTedsOutlinerColumnQuery()
 {
 	using namespace UE::Editor::DataStorage;
 	using namespace UE::Editor::DataStorage::Columns;
@@ -126,30 +127,38 @@ UE::Editor::DataStorage::FQueryDescription FTedsOutlinerModule::GetLevelEditorTe
 		
 	using namespace DataStorage::Queries;
 
-	static FQueryDescription ColumnQuery = 
-		Select()
-			.ReadOnly<FTypedElementClassTypeInfoColumn,
-					  FAlertColumn, 
+	QueryHandle ColumnQuery;
+
+	if (!CVarUseNewRevisionControlWidgets->GetBool())
+	{
+		ColumnQuery = Storage->RegisterQuery(
+			Select()
+				.ReadOnly<FTypedElementClassTypeInfoColumn,
+					  FAlertColumn,
 					  FChildAlertColumn,
 					  FUObjectIdNameColumn,
 					  FTedsActorMobilityColumn,
 					  FLevelColumn,
 					  FVisibleInEditorColumn>()
-		.Compile();
+			.Compile());
 
-	// Query to also include revision control info
-	static FQueryDescription RevisionControlQuery = 
-		Select()
-			.ReadOnly<FTypedElementClassTypeInfoColumn,  
+	}
+	else
+	{
+		ColumnQuery = Storage->RegisterQuery( 
+			Select()
+				.ReadOnly<FTypedElementClassTypeInfoColumn,  
 					  FTypedElementPackageReference, 
 					  FAlertColumn, 
 					  FUObjectIdNameColumn,
 					  FTedsActorMobilityColumn,
 					  FLevelColumn,
 					  FVisibleInEditorColumn>()
-		.Compile();
+			.Compile());
 
-	return Private::bUseNewRevisionControlWidgets ? RevisionControlQuery : ColumnQuery;
+	}
+
+	return ColumnQuery;
 }
 
 TSharedRef<SWidget> FTedsOutlinerModule::CreateLevelEditorTedsOutliner()
@@ -165,8 +174,7 @@ TSharedRef<SWidget> FTedsOutlinerModule::CreateLevelEditorTedsOutliner()
 	// The Outliner is populated with Actors and Entities
 	DataStorage::FQueryDescription OutlinerQueryDescription =
 		Select()
-		.Where()
-			.All<FTypedElementClassTypeInfoColumn>() // TEDS-Outliner TODO: Currently looking at all entries with type info in TEDS
+		.Where(TColumn<FTypedElementActorTag>() || TColumn<FFolderCompatibilityColumn>())
 		.Compile();
 
 	FSceneOutlinerInitializationOptions InitOptions;
@@ -177,9 +185,8 @@ TSharedRef<SWidget> FTedsOutlinerModule::CreateLevelEditorTedsOutliner()
 
 	FTedsOutlinerParams Params(nullptr);
 	Params.QueryDescription = OutlinerQueryDescription;
-	Params.ColumnQueryDescription = GetLevelEditorTedsOutlinerColumnQueryDescription();
-
-	Params.ClassFilters.Add(AActor::StaticClass());
+	Params.ColumnQuery = GetLevelEditorTedsOutlinerColumnQuery();
+		Params.ClassFilters.Add(AActor::StaticClass());
 	Params.ClassFilters.Add(AStaticMeshActor::StaticClass());
 	Params.ClassFilters.Add(ADirectionalLight::StaticClass());
 	
diff --git a/Engine/Plugins/Experimental/EditorDataStorageFeatures/Source/TedsOutliner/Public/Columns/TedsOutlinerColumns.h b/Engine/Plugins/Experimental/EditorDataStorageFeatures/Source/TedsOutliner/Public/Columns/TedsOutlinerColumns.h
--- a/Engine/Plugins/Experimental/EditorDataStorageFeatures/Source/TedsOutliner/Public/Columns/TedsOutlinerColumns.h
+++ b/Engine/Plugins/Experimental/EditorDataStorageFeatures/Source/TedsOutliner/Public/Columns/TedsOutlinerColumns.h
@@ -54,8 +54,8 @@ struct FTedsOutlinerColumnQueryColumn final : public FEditorDataStorageColumn
 {
 	GENERATED_BODY()
 
-	// The SELECT() part of the query description specifies the columns to be displayed
-	UE::Editor::DataStorage::FQueryDescription ColumnQueryDescription;
+	// The SELECT() part of the query description attached to this handle specifies the columns to be displayed
+	UE::Editor::DataStorage::QueryHandle ColumnQueryHandle;
 
 	// Call this delegate to refresh the columns in the Teds Outliner with the values specified in ColumnQueryDescription
 	UE::Editor::Outliner::FOnRefreshColumns OnRefreshColumns;
diff --git a/Engine/Plugins/Experimental/EditorDataStorageFeatures/Source/TedsOutliner/Public/TedsOutlinerImpl.h b/Engine/Plugins/Experimental/EditorDataStorageFeatures/Source/TedsOutliner/Public/TedsOutlinerImpl.h
--- a/Engine/Plugins/Experimental/EditorDataStorageFeatures/Source/TedsOutliner/Public/TedsOutlinerImpl.h
+++ b/Engine/Plugins/Experimental/EditorDataStorageFeatures/Source/TedsOutliner/Public/TedsOutlinerImpl.h
@@ -112,8 +112,9 @@ struct FTedsOutlinerParams
 	// The query description that will be used to populate rows in the TEDS Outliner
 	DataStorage::FQueryDescription QueryDescription;
 
-	// The query description that will be used to populate the initial columns in the TEDS Outliner
-	DataStorage::FQueryDescription ColumnQueryDescription;
+	// The query that will be used to populate the initial columns in the TEDS Outliner
+	// NOTE: The TEDS Outliner takes ownership of this query and will unregister it on destruction.
+	DataStorage::QueryHandle ColumnQuery;
 	
 	// TEDS Filters that utilize QueryDescriptions or QueryFunctions in this TEDS Outliner
 	TArray<FTedsFilterData> Filters;
diff --git a/Engine/Plugins/Experimental/EditorDataStorageFeatures/Source/TedsOutliner/Public/TedsOutlinerModule.h b/Engine/Plugins/Experimental/EditorDataStorageFeatures/Source/TedsOutliner/Public/TedsOutlinerModule.h
--- a/Engine/Plugins/Experimental/EditorDataStorageFeatures/Source/TedsOutliner/Public/TedsOutlinerModule.h
+++ b/Engine/Plugins/Experimental/EditorDataStorageFeatures/Source/TedsOutliner/Public/TedsOutlinerModule.h
@@ -45,7 +45,7 @@ public:
 	virtual void ShutdownModule() override;
 
 	// Get the column query the default table viewer uses
-	virtual UE::Editor::DataStorage::FQueryDescription GetLevelEditorTedsOutlinerColumnQueryDescription();
+	virtual UE::Editor::DataStorage::QueryHandle GetLevelEditorTedsOutlinerColumnQuery();
 
 	// The name of the tab the default table viewer is opened in
 	FName GetTedsOutlinerTabName();
