[TEDS] Type Info improvements and new columns

#Synced-CL 49077765

diff --git a/Engine/Plugins/Experimental/EditorDataStorageFeatures/Source/TedsTypeInfo/Private/TedsTypeInfoFactory.cpp b/Engine/Plugins/Experimental/EditorDataStorageFeatures/Source/TedsTypeInfo/Private/TedsTypeInfoFactory.cpp
--- a/Engine/Plugins/Experimental/EditorDataStorageFeatures/Source/TedsTypeInfo/Private/TedsTypeInfoFactory.cpp
+++ b/Engine/Plugins/Experimental/EditorDataStorageFeatures/Source/TedsTypeInfo/Private/TedsTypeInfoFactory.cpp
@@ -11,6 +11,7 @@
 #include "TedsTypeInfoModule.h"
 #include "TedsTypeInfoUtils.h"
 
+#include "Containers/VersePath.h"
 #include "DataStorage/Features.h"
 #include "Elements/Interfaces/TypedElementDataStorageCompatibilityInterface.h"
 #include "Elements/Interfaces/TypedElementDataStorageInterface.h"
@@ -26,9 +27,25 @@
 #include "UObject/UObjectIterator.h"
 #include "VerseVM/VVMVerseClass.h"
 
-const FName UTypeInfoFactory::PropertyTableName = "Property Information";
-const FName UTypeInfoFactory::ClassHierarchyName = "ClassHierarchy";
-const FName UTypeInfoFactory::PropertyHierarchyName = "PropertyHierarchy";
+namespace UE::Editor::DataStorage::Internal
+{
+	// Helper for unmangling verse names. Copied here as there's no shared helper util currently
+	static FString UnmanglePropertyName(const FName MaybeMangledName, bool& bOutNameWasMangled)
+	{
+		FString Result = MaybeMangledName.ToString();
+		if (Result.StartsWith(TEXTVIEW("__verse_0x")))
+		{
+			// chop "__verse_0x" (10 char) + CRC (8 char) + "_" (1 char)
+			Result = Result.RightChop(19);
+			bOutNameWasMangled = true;
+		}
+		else
+		{
+			bOutNameWasMangled = false;
+		}
+		return Result;
+	}
+}
 
 TAutoConsoleVariable<bool> CVarTEDSTypeInfoOnlyVerseProperties(TEXT("TEDS.Feature.TypeInfoIntegration.OnlyStoreVerseProperties"),
 	true,
@@ -45,13 +62,18 @@ void UTypeInfoFactory::RegisterTables(UE::Editor::DataStorage::ICoreProvider& Da
 
 	Super::RegisterTables(DataStorage, DataStorageCompatability);
 
+	DataStorage.RegisterTable({
+		FTypeInfoTag::StaticStruct(),
+		FTypedElementLabelColumn::StaticStruct(),
+	}, TypeInfo::TypeTableName);
+
 	PropertyTable = DataStorage.RegisterTable({
 		FPropertyInfoTag::StaticStruct(),
 		FPropertyNameColumn::StaticStruct(),
 		FPropertyOwnerTypeRowHandleColumn::StaticStruct(),
 		FPropertyTypeNameInfoColumn::StaticStruct(),
 		FPropertyTypeNameColumn::StaticStruct(),
-	}, PropertyTableName);
+	}, TypeInfo::PropertyTableName);
 }
 
 void UTypeInfoFactory::RegisterHierarchies(UE::Editor::DataStorage::ICoreProvider& DataStorage)
@@ -60,14 +82,14 @@ void UTypeInfoFactory::RegisterHierarchies(UE::Editor::DataStorage::ICoreProvide
 
 	const FHierarchyRegistrationParams ClassHierarchyParams
 	{
-		.Name = ClassHierarchyName,
+		.Name = TypeInfo::ClassHierarchyName,
 	};
 
 	DataStorage.RegisterHierarchy(ClassHierarchyParams);
 
 	const FHierarchyRegistrationParams PropertyHierarchyParams
 	{
-		.Name = PropertyHierarchyName,
+		.Name = TypeInfo::PropertyHierarchyName,
 	};
 
 	PropertyHierarchyHandle = DataStorage.RegisterHierarchy(PropertyHierarchyParams);
@@ -96,7 +118,7 @@ void UTypeInfoFactory::RegisterQueries(UE::Editor::DataStorage::ICoreProvider& D
 
 				Context.RemoveColumns(Row, { FTypeInfoRequiresHierarchyUpdateTag::StaticStruct() });
 			})
-			.AccessesHierarchy(ClassHierarchyName)
+			.AccessesHierarchy(TypeInfo::ClassHierarchyName)
 			.Where()
 				.All<FTypeInfoRequiresHierarchyUpdateTag>()
 			.Compile());
@@ -161,7 +183,7 @@ void UTypeInfoFactory::PopulateTypeInfo()
 		Type* CurrentType = *TypeIterator;
 		if (CurrentType && TryAddTypeInfo(*DataStorage, CurrentType, OutRowhandle))
 		{
-			if (!bOnlyAddVerseProperties || CurrentType->template IsChildOf<UVerseClass>())
+			if (!bOnlyAddVerseProperties || Cast<UVerseClass>(CurrentType))
 			{
 				TypesWithPendingPropertyInfo.Emplace(CurrentType, OutRowhandle);
 			}
@@ -251,6 +273,12 @@ void UTypeInfoFactory::AddPropertyInfo(UE::Editor::DataStorage::ICoreProvider& D
 			// TODO: Cleanup property rows when associated type row is destroyed
 			// We get away with it for now because Verse Properties are destroyed during compiles and we refresh the whole type info table then
 			DataStorage.SetParentRow(PropertyHierarchyHandle, PropertyRowHandle, OwnerTypeRowHandle);
+
+			if (Cast<UVerseClass>(InTypeInfo))
+			{
+				DataStorage.AddColumn<FVersePropertyTag>(PropertyRowHandle);
+			}
+
 			AddCommonPropertyColumns(DataStorage, OwnerTypeRowHandle, PropertyRowHandle, *CurrentPropertyIt);
 		}
 	}
@@ -285,6 +313,8 @@ void UTypeInfoFactory::AddCommonColumns(UE::Editor::DataStorage::ICoreProvider&
 	checkf(InTypeInfo, TEXT("Attempted to add type info columns with null TypeInfo"));
 
 	DataStorage.AddColumn(TypeRowHandle, FTypedElementLabelColumn{ .Label = InTypeInfo->GetDisplayNameText().ToString() });
+	DataStorage.AddColumn(TypeRowHandle, FTypeInfoObjectPathColumn{ .ObjectPath = InTypeInfo->GetPathName() });
+	DataStorage.AddColumn(TypeRowHandle, FTypeInfoVersePathColumn{ .VersePath = InTypeInfo->GetVersePath().AsText() });
 	DataStorage.AddColumns<FTypeInfoTag, FTypeInfoRequiresHierarchyUpdateTag>(TypeRowHandle);
 }
 
@@ -376,13 +406,23 @@ void UTypeInfoFactory::AddCommonPropertyColumns(UE::Editor::DataStorage::ICorePr
 
 	checkf(InPropertyInfo, TEXT("Attempted to add property info columns with null PropertyInfo"));
 
-	DataStorage.AddColumn(PropertyRowHandle, FPropertyNameColumn{ .PropertyName = InPropertyInfo->GetDisplayNameText() });
+	bool bNameWasMangled = false;
+	FString PropertyName = UE::Editor::DataStorage::Internal::UnmanglePropertyName(InPropertyInfo->GetFName(), bNameWasMangled);
+
+	UE::FPropertyTypeName PropertyTypeNameInfo(InPropertyInfo);
+	DataStorage.AddColumn(PropertyRowHandle, FPropertyNameColumn{ .PropertyName = FText::FromString(PropertyName) });
 	DataStorage.AddColumn(PropertyRowHandle, FPropertyOwnerTypeRowHandleColumn{ .PropertyOwnerTypeRowHandle = OwnerTypeRowHandle });
-	DataStorage.AddColumn(PropertyRowHandle, FPropertyTypeNameInfoColumn{ .PropertyTypeNameInfo = UE::FPropertyTypeName(InPropertyInfo) });
+	DataStorage.AddColumn(PropertyRowHandle, FPropertyTypeNameInfoColumn{ .PropertyTypeNameInfo = PropertyTypeNameInfo });
 
 	if (FFieldClass* PropertyClass = InPropertyInfo->GetClass())
 	{
-		DataStorage.AddColumn(PropertyRowHandle, FPropertyTypeNameColumn{ .PropertyTypeName = InPropertyInfo->GetClass()->GetDisplayNameText() });
+		DataStorage.AddColumn(PropertyRowHandle,
+			FPropertyTypeNameColumn{ .PropertyTypeName = FText::FromString(PropertyTypeNameInfo.GetName().ToString()) });
+	}
+
+	if (!InPropertyInfo->HasAnyPropertyFlags(CPF_BlueprintReadOnly))
+	{
+		DataStorage.AddColumn<FPropertyMutableTag>(PropertyRowHandle);
 	}
 
 	DataStorage.AddColumn<FPropertyInfoTag>(PropertyRowHandle);
diff --git a/Engine/Plugins/Experimental/EditorDataStorageFeatures/Source/TedsTypeInfo/Private/TedsTypeInfoFactory.h b/Engine/Plugins/Experimental/EditorDataStorageFeatures/Source/TedsTypeInfo/Private/TedsTypeInfoFactory.h
--- a/Engine/Plugins/Experimental/EditorDataStorageFeatures/Source/TedsTypeInfo/Private/TedsTypeInfoFactory.h
+++ b/Engine/Plugins/Experimental/EditorDataStorageFeatures/Source/TedsTypeInfo/Private/TedsTypeInfoFactory.h
@@ -35,9 +35,6 @@ class UTypeInfoFactory final : public UEditorDataStorageFactory
 	GENERATED_BODY()
 
 public:
-	static const FName PropertyTableName;
-	static const FName ClassHierarchyName;
-	static const FName PropertyHierarchyName;
 
 	virtual ~UTypeInfoFactory() override = default;
 
@@ -79,7 +76,6 @@ private:
 
 	bool bRefreshTypeInfoQueued = false;
 
-	UE::Editor::DataStorage::TableHandle TypeTable = UE::Editor::DataStorage::InvalidTableHandle;
 	UE::Editor::DataStorage::TableHandle PropertyTable = UE::Editor::DataStorage::InvalidTableHandle;
 	UE::Editor::DataStorage::FHierarchyHandle PropertyHierarchyHandle;
 };
diff --git a/Engine/Plugins/Experimental/EditorDataStorageFeatures/Source/TedsTypeInfo/Public/Elements/Columns/TedsTypeInfoColumns.h b/Engine/Plugins/Experimental/EditorDataStorageFeatures/Source/TedsTypeInfo/Public/Elements/Columns/TedsTypeInfoColumns.h
--- a/Engine/Plugins/Experimental/EditorDataStorageFeatures/Source/TedsTypeInfo/Public/Elements/Columns/TedsTypeInfoColumns.h
+++ b/Engine/Plugins/Experimental/EditorDataStorageFeatures/Source/TedsTypeInfo/Public/Elements/Columns/TedsTypeInfoColumns.h
@@ -38,6 +38,30 @@ struct FDataStorageVerseTypeInfoTag : public FEditorDataStorageTag
 	GENERATED_BODY()
 };
 
+USTRUCT(meta = (DisplayName = "Row type is virtualized"))
+struct FDataStorageVirtualTypeInfoTag : public FEditorDataStorageTag
+{
+	GENERATED_BODY()
+};
+
+USTRUCT(meta = (DisplayName = "Type Object Path"))
+struct FDataStorageTypeInfoObjectPathColumn final : public FEditorDataStorageColumn
+{
+	GENERATED_BODY()
+
+	UPROPERTY(meta = (Searchable, Sortable))
+	FString ObjectPath;
+};
+
+USTRUCT(meta = (DisplayName = "Type Verse Path"))
+struct FDataStorageTypeInfoVersePathColumn final : public FEditorDataStorageColumn
+{
+	GENERATED_BODY()
+
+	UPROPERTY(meta = (Searchable, Sortable))
+	FText VersePath;
+};
+
 USTRUCT(meta = (DisplayName = "Verse type access level", EditorDataStorage_DynamicColumnTemplate))
 struct FDataStorageVerseTypeInfoAccessLevel : public FEditorDataStorageTag
 {
@@ -86,11 +110,23 @@ struct FDataStoragePropertyTypeNameInfoColumn final : public FEditorDataStorageC
 	UE::FPropertyTypeName PropertyTypeNameInfo;
 };
 
+USTRUCT(meta = (DisplayName = "Property is mutable"))
+struct FDataStoragePropertyMutableTag : public FEditorDataStorageTag
+{
+	GENERATED_BODY()
+};
+
 USTRUCT(meta = (DisplayName = "Row is a property"))
 struct FDataStoragePropertyInfoTag : public FEditorDataStorageTag
 {
 	GENERATED_BODY()
 };
+
+USTRUCT(meta = (DisplayName = "Property is owned by a Verse type"))
+struct FDataStorageVersePropertyTag : public FEditorDataStorageTag
+{
+	GENERATED_BODY()
+};
 // Property Info Columns End
 
 namespace UE::Editor::DataStorage
@@ -100,7 +136,10 @@ namespace UE::Editor::DataStorage
 	using FClassTypeInfoTag = FDataStorageClassTypeInfoTag;
 	using FStructTypeInfoTag = FDataStorageStructTypeInfoTag;
 	using FTypeInfoInterfaceTag = FDataStorageTypeInfoInterfaceTag;
+	using FTypeInfoObjectPathColumn = FDataStorageTypeInfoObjectPathColumn;
+	using FTypeInfoVersePathColumn = FDataStorageTypeInfoVersePathColumn;
 	using FVerseTypeInfoTag = FDataStorageVerseTypeInfoTag;
+	using FVirtualTypeInfoTag = FDataStorageVirtualTypeInfoTag;
 	using FVerseTypeInfoAccessLevel = FDataStorageVerseTypeInfoAccessLevel;
 	using FHasUnknownPropertiesTag = FDataStorageHasUnknownPropertiesTag;
 	// Type Info Tags End
@@ -110,11 +149,17 @@ namespace UE::Editor::DataStorage
 	using FPropertyOwnerTypeRowHandleColumn = FDataStoragePropertyOwnerTypeRowHandleColumn;
 	using FPropertyTypeNameColumn = FDataStoragePropertyTypeNameColumn;
 	using FPropertyTypeNameInfoColumn = FDataStoragePropertyTypeNameInfoColumn;
+	using FPropertyMutableTag = FDataStoragePropertyMutableTag;
 	using FPropertyInfoTag = FDataStoragePropertyInfoTag;
+	using FVersePropertyTag = FDataStorageVersePropertyTag;
 	// Property Info Tags End
 }
 
 namespace UE::Editor::DataStorage::TypeInfo
 {
 	static const FName TypeMappingDomain = "TypeMappingDomain";
+	static const FName TypeTableName = "Type Information";
+	static const FName PropertyTableName = "Property Information";
+	static const FName ClassHierarchyName = "ClassHierarchy";
+	static const FName PropertyHierarchyName = "PropertyHierarchy";
 }
\ No newline at end of file
