Teds Visibility Widget: Fix visibility toggle not working in the hierarchy viewer. The hierarchy data column now only exists on the table viewer row and can be accessed by widgets (instead of duplicating onto every widget row)

#Synced-CL 49084582

diff --git a/Engine/Plugins/Experimental/EditorDataStorageFeatures/Source/TedsOutliner/Private/Widgets/OutlinerVisibilityWidget.cpp b/Engine/Plugins/Experimental/EditorDataStorageFeatures/Source/TedsOutliner/Private/Widgets/OutlinerVisibilityWidget.cpp
--- a/Engine/Plugins/Experimental/EditorDataStorageFeatures/Source/TedsOutliner/Private/Widgets/OutlinerVisibilityWidget.cpp
+++ b/Engine/Plugins/Experimental/EditorDataStorageFeatures/Source/TedsOutliner/Private/Widgets/OutlinerVisibilityWidget.cpp
@@ -19,6 +19,7 @@
 #include "ActorTreeItem.h"
 #include "TedsOutlinerHelpers.h"
 #include "TedsOutlinerItem.h"
+#include "TedsTableViewerUtils.h"
 #include "TedsTableViewerWidgetColumns.h"
 
 #include UE_INLINE_GENERATED_CPP_BY_NAME(OutlinerVisibilityWidget)
@@ -301,8 +302,10 @@ void STedsVisibilityWidget::SetIsVisible(RowHandle InRow, const bool bVisible)
 
 	if (ICoreProvider* DataStorage = GetDataStorage())
 	{
-		// Use the hierarchy data on this widget directly if present
-		if (FHierarchyViewerDataColumn* DataColumn = DataStorage->GetColumn<FHierarchyViewerDataColumn>(WidgetRow))
+		// Use the hierarchy data if present
+		const RowHandle ParentRow = TableViewerUtils::GetTableViewerUiRow(DataStorage, WidgetRow);
+		
+		if (FHierarchyViewerDataColumn* DataColumn = DataStorage->GetColumn<FHierarchyViewerDataColumn>(ParentRow))
 		{
 			if (TSharedPtr<IHierarchyViewerDataInterface> DataInterface = DataColumn->HierarchyData.Pin())
 			{
diff --git a/Engine/Plugins/Experimental/EditorDataStorageFeatures/Source/TedsTableViewer/Private/Widgets/STedsTableViewerRow.cpp b/Engine/Plugins/Experimental/EditorDataStorageFeatures/Source/TedsTableViewer/Private/Widgets/STedsTableViewerRow.cpp
--- a/Engine/Plugins/Experimental/EditorDataStorageFeatures/Source/TedsTableViewer/Private/Widgets/STedsTableViewerRow.cpp
+++ b/Engine/Plugins/Experimental/EditorDataStorageFeatures/Source/TedsTableViewer/Private/Widgets/STedsTableViewerRow.cpp
@@ -158,17 +158,6 @@ namespace UE::Editor::DataStorage
 
 		return ActualWidget;
 	}
-
-	void SHierarchyViewerRow::SetupWidgetColumns(ICoreProvider& InStorage, const RowHandle& InUIRowHandle)
-	{
-		// Forward hierarchy data from parent widget (row) to cell widget for convenient access
-		if (const FHierarchyViewerDataColumn* Data = InStorage.GetColumn<FHierarchyViewerDataColumn>(WidgetRowHandle))
-		{
-			InStorage.AddColumn(InUIRowHandle, FHierarchyViewerDataColumn{.HierarchyData = Data->HierarchyData });
-		}
-		
-		STedsTableViewerRow::SetupWidgetColumns(InStorage, InUIRowHandle);
-	}
 } // namespace UE::Editor::DataStorage
 
 #undef LOCTEXT_NAMESPACE
diff --git a/Engine/Plugins/Experimental/EditorDataStorageFeatures/Source/TedsTableViewer/Public/Widgets/STedsTableViewerRow.h b/Engine/Plugins/Experimental/EditorDataStorageFeatures/Source/TedsTableViewer/Public/Widgets/STedsTableViewerRow.h
--- a/Engine/Plugins/Experimental/EditorDataStorageFeatures/Source/TedsTableViewer/Public/Widgets/STedsTableViewerRow.h
+++ b/Engine/Plugins/Experimental/EditorDataStorageFeatures/Source/TedsTableViewer/Public/Widgets/STedsTableViewerRow.h
@@ -53,7 +53,7 @@ namespace UE::Editor::DataStorage
 	protected:
 
 		/** Add some common columns to the widget row for each cell widget created */
-		virtual void SetupWidgetColumns(ICoreProvider& InStorage, const RowHandle& InUIRowHandle);
+		void SetupWidgetColumns(ICoreProvider& InStorage, const RowHandle& InUIRowHandle);
 
 	private:
 		/** Get the current ItemHeight as an OptionalSize for the Box HeightOverride */
@@ -78,8 +78,5 @@ namespace UE::Editor::DataStorage
 	{
 	public:
 		virtual TSharedRef<SWidget> GenerateWidgetForColumn( const FName& ColumnName ) override;
-
-	protected:
-		virtual void SetupWidgetColumns(ICoreProvider& InStorage, const RowHandle& InUIRowHandle) override;
 	};
 } // namespace UE::Editor::DataStorage
