TedsSettings : change the default value of the cvar "TEDS.Feature.Settings.Enable" to true.

#Synced-CL 49087954

diff --git a/Engine/Plugins/Experimental/EditorDataStorageFeatures/Source/TedsSettings/Private/TedsSettingsEditorSubsystem.cpp b/Engine/Plugins/Experimental/EditorDataStorageFeatures/Source/TedsSettings/Private/TedsSettingsEditorSubsystem.cpp
--- a/Engine/Plugins/Experimental/EditorDataStorageFeatures/Source/TedsSettings/Private/TedsSettingsEditorSubsystem.cpp
+++ b/Engine/Plugins/Experimental/EditorDataStorageFeatures/Source/TedsSettings/Private/TedsSettingsEditorSubsystem.cpp
@@ -12,7 +12,7 @@ namespace UE::Editor::Settings::Private
 {
 	static TAutoConsoleVariable<bool> CVarTedsSettingsEnable(
 		TEXT("TEDS.Feature.Settings.Enable"),
-		false,
+		true,
 		TEXT("When true, settings objects from the ISettingsModule will be mirrored to rows in the editor data storage."));
 } // namespace UE::Editor::Settings::Private
 
