[TEDS Viewers] Allow custom mapping delegate in the viewers so that the Prefab Editor can use the GUID/ID Name to persist on rename and reinstance

#Synced-CL 49174465

diff --git a/Engine/Plugins/Experimental/EditorDataStorageFeatures/Source/TedsTableViewer/Private/TedsTableViewerUtils.cpp b/Engine/Plugins/Experimental/EditorDataStorageFeatures/Source/TedsTableViewer/Private/TedsTableViewerUtils.cpp
--- a/Engine/Plugins/Experimental/EditorDataStorageFeatures/Source/TedsTableViewer/Private/TedsTableViewerUtils.cpp
+++ b/Engine/Plugins/Experimental/EditorDataStorageFeatures/Source/TedsTableViewer/Private/TedsTableViewerUtils.cpp
@@ -187,7 +187,12 @@ namespace UE::Editor::DataStorage::TableViewerUtils
 
 	FString CreateRowMappingKey(const FName& HierarchyIdentifier, const TableViewerItemPtr Item)
 	{
-		return FString::Printf(TEXT("%s.%llu"), *HierarchyIdentifier.ToString(), Item.RowHandle);
+		return CreateRowMappingKey(HierarchyIdentifier, Item.RowHandle);
+	}
+	
+	FString CreateRowMappingKey(const FName& HierarchyIdentifier, const RowHandle Item)
+	{
+		return FString::Printf(TEXT("%s.%llu"), *HierarchyIdentifier.ToString(), Item);
 	}
 	
 	RowHandle GetTableViewerUiRow(const ICoreProvider* DataStorage, const RowHandle ChildWidgetRow)
diff --git a/Engine/Plugins/Experimental/EditorDataStorageFeatures/Source/TedsTableViewer/Private/Widgets/Composite/STedsCompositeHierarchyViewer.cpp b/Engine/Plugins/Experimental/EditorDataStorageFeatures/Source/TedsTableViewer/Private/Widgets/Composite/STedsCompositeHierarchyViewer.cpp
--- a/Engine/Plugins/Experimental/EditorDataStorageFeatures/Source/TedsTableViewer/Private/Widgets/Composite/STedsCompositeHierarchyViewer.cpp
+++ b/Engine/Plugins/Experimental/EditorDataStorageFeatures/Source/TedsTableViewer/Private/Widgets/Composite/STedsCompositeHierarchyViewer.cpp
@@ -114,6 +114,7 @@ namespace UE::Editor::DataStorage
 			.ColumnMetaData(HierarchyViewerArgs._ColumnMetaData)
 			.ListSelectionMode(HierarchyViewerArgs._ListSelectionMode)
 			.OnSelectionChanged(HierarchyViewerArgs._OnSelectionChanged)
+			.GetCustomTableViewerRowMapping(HierarchyViewerArgs._GetCustomTableViewerRowMapping)
 			.PrimaryColumn(HierarchyViewerArgs._PrimaryColumn)
 			.AllowInvisibleItemSelection(HierarchyViewerArgs._AllowInvisibleItemSelection);
 		
diff --git a/Engine/Plugins/Experimental/EditorDataStorageFeatures/Source/TedsTableViewer/Private/Widgets/STedsHierarchyViewer.cpp b/Engine/Plugins/Experimental/EditorDataStorageFeatures/Source/TedsTableViewer/Private/Widgets/STedsHierarchyViewer.cpp
--- a/Engine/Plugins/Experimental/EditorDataStorageFeatures/Source/TedsTableViewer/Private/Widgets/STedsHierarchyViewer.cpp
+++ b/Engine/Plugins/Experimental/EditorDataStorageFeatures/Source/TedsTableViewer/Private/Widgets/STedsHierarchyViewer.cpp
@@ -33,6 +33,7 @@ namespace UE::Editor::DataStorage
 	{
 		HierarchyInterface = MoveTemp(InHierarchyInterface);
 		OnSelectionChanged = InArgs._OnSelectionChanged;
+		GetCustomTableViewerRowMapping = InArgs._GetCustomTableViewerRowMapping;
 		EmptyRowsMessage = InArgs._EmptyRowsMessage;
 		
 		IUiProvider::FPurposeID CellWidgetPurpose = InArgs._CellWidgetPurpose;
@@ -277,7 +278,7 @@ namespace UE::Editor::DataStorage
 	{
 		if (ICoreProvider* Storage = Model->GetDataStorageInterface())
 		{
-			const FString TableRowId = TableViewerUtils::CreateRowMappingKey(TableViewerIdentifier, Item);
+			const FString TableRowId = GetTableViewerRowMapping(Item);
 			if (const RowHandle UIRow = Storage->LookupMappedRow(TableViewerUIRowMappingDomain, FMapKey(TableRowId)); Storage->IsRowAvailable(UIRow))
 			{
 				if (bExpanded)
@@ -436,12 +437,24 @@ namespace UE::Editor::DataStorage
 	{
 		return TreeView->IsItemVisible(InItem);
 	}
+	
+	FString SHierarchyViewer::GetTableViewerRowMapping(const TableViewerItemPtr InItem) const
+	{
+		if(GetCustomTableViewerRowMapping.IsBound())
+		{
+			if (FString RowMapping = GetCustomTableViewerRowMapping.Execute(TableViewerIdentifier, InItem.RowHandle); !RowMapping.IsEmpty())
+			{
+				return RowMapping;
+			}
+		}
+		return TableViewerUtils::CreateRowMappingKey(TableViewerIdentifier, InItem);
+	}
 
-	RowHandle SHierarchyViewer::AddOrFindTableViewerUIRow(TableViewerItemPtr InItem) const
+	RowHandle SHierarchyViewer::AddOrFindTableViewerUIRow(const TableViewerItemPtr InItem) const
 	{
 		if (ICoreProvider* Storage = Model->GetDataStorageInterface())
 		{
-			const FString MappingKey = TableViewerUtils::CreateRowMappingKey(TableViewerIdentifier, InItem);
+			const FString MappingKey = GetTableViewerRowMapping(InItem);
 			// If the row already existed but was removed from TreeView's SparseItemInfos (Searching, Filters, etc.), re-expand it
 			if (const RowHandle ExistingRow = Storage->LookupMappedRow(TableViewerUIRowMappingDomain, FMapKey(MappingKey)); Storage->IsRowAvailable(ExistingRow))
 			{
diff --git a/Engine/Plugins/Experimental/EditorDataStorageFeatures/Source/TedsTableViewer/Private/Widgets/STedsTableViewer.cpp b/Engine/Plugins/Experimental/EditorDataStorageFeatures/Source/TedsTableViewer/Private/Widgets/STedsTableViewer.cpp
--- a/Engine/Plugins/Experimental/EditorDataStorageFeatures/Source/TedsTableViewer/Private/Widgets/STedsTableViewer.cpp
+++ b/Engine/Plugins/Experimental/EditorDataStorageFeatures/Source/TedsTableViewer/Private/Widgets/STedsTableViewer.cpp
@@ -28,6 +28,7 @@ namespace UE::Editor::DataStorage
 	void STedsTableViewer::Construct(const FArguments& InArgs)
 	{
 		OnSelectionChanged = InArgs._OnSelectionChanged;
+		GetCustomTableViewerRowMapping = InArgs._GetCustomTableViewerRowMapping;
 		EmptyRowsMessage = InArgs._EmptyRowsMessage;
 		ItemHeight = InArgs._ItemHeight;
 		ItemPadding = InArgs._ItemPadding;
@@ -59,8 +60,7 @@ namespace UE::Editor::DataStorage
 			.ColumnMetaData = InArgs._ColumnMetaData,
 			.CellWidgetPurpose = CellWidgetPurpose,
 			.HeaderWidgetPurpose = HeaderWidgetPurpose,
-			.IsItemVisibleDelegate = FTedsTableViewerModel::FIsItemVisible::CreateSP(this, &STedsTableViewer::IsItemVisible),
-			.PrimaryColumn = InArgs._PrimaryColumn
+			.IsItemVisibleDelegate = FTedsTableViewerModel::FIsItemVisible::CreateSP(this, &STedsTableViewer::IsItemVisible)
 		});
 
 		InstanceCounter++;
@@ -295,12 +295,24 @@ namespace UE::Editor::DataStorage
 	{
 		return ListView->IsItemVisible(InItem);
 	}
+	
+	FString STedsTableViewer::GetTableViewerRowMapping(const TableViewerItemPtr InItem) const
+	{
+		if(GetCustomTableViewerRowMapping.IsBound())
+		{
+			if (FString RowMapping = GetCustomTableViewerRowMapping.Execute(TableViewerIdentifier, InItem.RowHandle); !RowMapping.IsEmpty())
+			{
+				return RowMapping;
+			}
+		}
+		return TableViewerUtils::CreateRowMappingKey(TableViewerIdentifier, InItem);
+	}
 
-	RowHandle STedsTableViewer::AddOrFindTableViewerUIRow(TableViewerItemPtr InItem) const
+	RowHandle STedsTableViewer::AddOrFindTableViewerUIRow(const TableViewerItemPtr InItem) const
 	{
 		if (ICoreProvider* Storage = Model->GetDataStorageInterface())
 		{
-			const FString MappingKey = TableViewerUtils::CreateRowMappingKey(TableViewerIdentifier, InItem);
+			const FString MappingKey = GetTableViewerRowMapping(InItem);
 			// If the row already existed but was removed from TreeView's SparseItemInfos (Searching, Filters, etc.), re-expand it
 			if (const RowHandle ExistingRow = Storage->LookupMappedRow(TableViewerUIRowMappingDomain, FMapKey(MappingKey)); Storage->IsRowAvailable(ExistingRow))
 			{
diff --git a/Engine/Plugins/Experimental/EditorDataStorageFeatures/Source/TedsTableViewer/Public/TedsTableViewerUtils.h b/Engine/Plugins/Experimental/EditorDataStorageFeatures/Source/TedsTableViewer/Public/TedsTableViewerUtils.h
--- a/Engine/Plugins/Experimental/EditorDataStorageFeatures/Source/TedsTableViewer/Public/TedsTableViewerUtils.h
+++ b/Engine/Plugins/Experimental/EditorDataStorageFeatures/Source/TedsTableViewer/Public/TedsTableViewerUtils.h
@@ -44,8 +44,11 @@ namespace UE::Editor::DataStorage::TableViewerUtils
 	// Scans the row for columns and tries to find the best matching icon.
 	TEDSTABLEVIEWER_API const FSlateBrush* GetIconForRow(ICoreProvider* DataStorage, RowHandle Row);
 
-	// Creates the Row Map Key for the Table Row UI Rows
+	// Creates the default Row Map Key for the Table Row UI Rows
 	TEDSTABLEVIEWER_API FString CreateRowMappingKey(const FName& HierarchyIdentifier, const TableViewerItemPtr Item);
+	
+	// Creates the default Row Map Key for the Table Row UI Rows from the RowHandle directly
+    TEDSTABLEVIEWER_API FString CreateRowMappingKey(const FName& HierarchyIdentifier, const RowHandle Item);
 
 	// Walks up the UI Widget Hierarchy until finding the row with the FTableViewerTag
 	TEDSTABLEVIEWER_API RowHandle GetTableViewerUiRow(const ICoreProvider* DataStorage, const RowHandle ChildWidgetRow);
diff --git a/Engine/Plugins/Experimental/EditorDataStorageFeatures/Source/TedsTableViewer/Public/Widgets/STedsHierarchyViewer.h b/Engine/Plugins/Experimental/EditorDataStorageFeatures/Source/TedsTableViewer/Public/Widgets/STedsHierarchyViewer.h
--- a/Engine/Plugins/Experimental/EditorDataStorageFeatures/Source/TedsTableViewer/Public/Widgets/STedsHierarchyViewer.h
+++ b/Engine/Plugins/Experimental/EditorDataStorageFeatures/Source/TedsTableViewer/Public/Widgets/STedsHierarchyViewer.h
@@ -54,6 +54,9 @@ namespace UE::Editor::DataStorage
 
 		// Delegate fired when the selection in the table viewer changes
 		DECLARE_DELEGATE_OneParam(FOnSelectionChanged, RowHandle)
+		
+		// Delegate fired to get a custom row mapping for the table viewer
+        DECLARE_DELEGATE_RetVal_TwoParams(FString, FGetTableViewerRowMapping, FName, RowHandle);
 
 		SLATE_BEGIN_ARGS(SHierarchyViewer)
 			: _CanSelectGeneratedColumn(true)
@@ -89,6 +92,10 @@ namespace UE::Editor::DataStorage
 		// The selection mode for the table viewer (single/multi etc)
 		SLATE_ARGUMENT(ESelectionMode::Type, ListSelectionMode)
 
+		// Delegate called to define a custom row mapping for the Table Viewer Row Widget Rows
+		// Used in cases where you need the UI Rows to persist across RowHandle changes
+		SLATE_EVENT(FGetTableViewerRowMapping, GetCustomTableViewerRowMapping)
+
 		// The message to show in place of the table viewer when there are no rows provided by the current query stack
 		// Empty = simply show the column headers instead of a message
 		SLATE_ATTRIBUTE(FText, EmptyRowsMessage)
@@ -154,7 +161,9 @@ namespace UE::Editor::DataStorage
 
 	protected:
 		
-		RowHandle AddOrFindTableViewerUIRow(TableViewerItemPtr InItem) const;
+		FString GetTableViewerRowMapping(const TableViewerItemPtr InItem) const;
+		
+		RowHandle AddOrFindTableViewerUIRow(const TableViewerItemPtr InItem) const;
 
 		void RemoveAllTableViewerUIRows() const;
 		
@@ -200,6 +209,9 @@ namespace UE::Editor::DataStorage
 
 		// Delegate fired when the selection changes
 		FOnSelectionChanged OnSelectionChanged;
+		
+		// Delegate fired when getting the row mapping, will use the default 'TableId.DataRowHandle' if nothing is bound
+        FGetTableViewerRowMapping GetCustomTableViewerRowMapping;
 
 		// Wrapper Teds Widget around our contents so we can use Teds columns to specify behavior
 		TSharedPtr<ITedsWidget> TedsWidget;
diff --git a/Engine/Plugins/Experimental/EditorDataStorageFeatures/Source/TedsTableViewer/Public/Widgets/STedsTableViewer.h b/Engine/Plugins/Experimental/EditorDataStorageFeatures/Source/TedsTableViewer/Public/Widgets/STedsTableViewer.h
--- a/Engine/Plugins/Experimental/EditorDataStorageFeatures/Source/TedsTableViewer/Public/Widgets/STedsTableViewer.h
+++ b/Engine/Plugins/Experimental/EditorDataStorageFeatures/Source/TedsTableViewer/Public/Widgets/STedsTableViewer.h
@@ -41,6 +41,9 @@ namespace UE::Editor::DataStorage
 		
 		// Delegate fired when the selection in the table viewer changes
 		DECLARE_DELEGATE_OneParam(FOnSelectionChanged, RowHandle)
+		
+		// Delegate fired to get a custom row mapping for the table viewer
+		DECLARE_DELEGATE_RetVal_TwoParams(FString, FGetTableViewerRowMapping, FName, RowHandle);
 
 		SLATE_BEGIN_ARGS(STedsTableViewer)
 			: _CanSelectGeneratedColumn(true)
@@ -78,6 +81,10 @@ namespace UE::Editor::DataStorage
 		// The selection mode for the table viewer (single/multi etc)
 		SLATE_ARGUMENT(ESelectionMode::Type, ListSelectionMode)
 
+		// Delegate called to define a custom row mapping for the Table Viewer Row Widget Rows
+		// Used in cases where you need the UI Rows to persist across RowHandle changes
+		SLATE_EVENT(FGetTableViewerRowMapping, GetCustomTableViewerRowMapping)
+
 		// The message to show in place of the table viewer when there are no rows provided by the current query stack
 		// Empty = simply show the column headers instead of a message
 		SLATE_ATTRIBUTE(FText, EmptyRowsMessage)
@@ -138,8 +145,10 @@ namespace UE::Editor::DataStorage
 		TEDSTABLEVIEWER_API virtual FReply OnDrop( const FGeometry& MyGeometry, const FDragDropEvent& DragDropEvent ) override;
 
 	protected:
+	
+		FString GetTableViewerRowMapping(const TableViewerItemPtr InItem) const;
 
-		RowHandle AddOrFindTableViewerUIRow(TableViewerItemPtr InItem) const;
+		RowHandle AddOrFindTableViewerUIRow(const TableViewerItemPtr InItem) const;
 
 		void RemoveAllTableViewerUIRows() const;
 		
@@ -179,6 +188,9 @@ namespace UE::Editor::DataStorage
 
 		// Delegate fired when the selection changes
 		FOnSelectionChanged OnSelectionChanged;
+		
+		// Delegate fired when getting the row mapping, will use the default 'TableId.DataRowHandle' if nothing is bound
+		FGetTableViewerRowMapping GetCustomTableViewerRowMapping;
 
 		// Wrapper Teds Widget around our contents so we can use Teds columns to specify behavior
 		TSharedPtr<ITedsWidget> TedsWidget;
