Mark SOutlinerEditableTextWidget and SFavoriteAssetTreeItem as invalidation compatible

#rnx
#Synced-CL 49206316

diff --git a/Engine/Plugins/Experimental/EditorDataStorageFeatures/Source/TedsOutliner/Private/Widgets/OutlinerTextWidget.cpp b/Engine/Plugins/Experimental/EditorDataStorageFeatures/Source/TedsOutliner/Private/Widgets/OutlinerTextWidget.cpp
--- a/Engine/Plugins/Experimental/EditorDataStorageFeatures/Source/TedsOutliner/Private/Widgets/OutlinerTextWidget.cpp
+++ b/Engine/Plugins/Experimental/EditorDataStorageFeatures/Source/TedsOutliner/Private/Widgets/OutlinerTextWidget.cpp
@@ -127,6 +127,12 @@ namespace UE::Editor::Outliner::TextWidget::Private
 	};
 }
 
+template<>
+struct TWidgetTypeTraits<UE::Editor::Outliner::TextWidget::Private::SOutlinerEditableTextWidget>
+{
+	static constexpr bool SupportsInvalidation() { return true; }
+};
+
 void UOutlinerTextWidgetFactory::RegisterWidgetConstructors(UE::Editor::DataStorage::ICoreProvider& DataStorage,
 	UE::Editor::DataStorage::IUiProvider& DataStorageUi) const
 {
diff --git a/Engine/Source/Editor/ContentBrowser/Private/SPathView.h b/Engine/Source/Editor/ContentBrowser/Private/SPathView.h
--- a/Engine/Source/Editor/ContentBrowser/Private/SPathView.h
+++ b/Engine/Source/Editor/ContentBrowser/Private/SPathView.h
@@ -602,6 +602,12 @@ private:
 	TSharedPtr<SExpandableArea> FavoritesArea;
 };
 
+template<>
+struct TWidgetTypeTraits<class SFavoriteAssetTreeItem>
+{
+	static constexpr bool SupportsInvalidation() { return true; }
+};
+
 /** AssetTreeItem for the Favorite view to override some behavior */
 class SFavoriteAssetTreeItem : public SAssetTreeItem
 {
