TedsTableViewerUtils : add GetTableViewerRowUiRow

#Synced-CL 49278429

diff --git a/Engine/Plugins/Experimental/EditorDataStorageFeatures/Source/TedsTableViewer/Private/TedsTableViewerUtils.cpp b/Engine/Plugins/Experimental/EditorDataStorageFeatures/Source/TedsTableViewer/Private/TedsTableViewerUtils.cpp
--- a/Engine/Plugins/Experimental/EditorDataStorageFeatures/Source/TedsTableViewer/Private/TedsTableViewerUtils.cpp
+++ b/Engine/Plugins/Experimental/EditorDataStorageFeatures/Source/TedsTableViewer/Private/TedsTableViewerUtils.cpp
@@ -195,6 +195,21 @@ namespace UE::Editor::DataStorage::TableViewerUtils
 		return FString::Printf(TEXT("%s.%llu"), *HierarchyIdentifier.ToString(), Item);
 	}
 	
+	RowHandle GetTableViewerRowUiRow(const ICoreProvider* DataStorage, const RowHandle ChildWidgetRow)
+	{
+		const FHierarchyHandle WidgetHierarchy = DataStorage->FindHierarchyByName(GetWidgetHierarchyName());
+		RowHandle ParentRow = ChildWidgetRow;
+		while (ParentRow != InvalidRowHandle)
+		{
+			ParentRow = DataStorage->GetParentRow(WidgetHierarchy, ParentRow);
+			if (DataStorage->HasColumns<FTableViewerRowTag>(ParentRow))
+			{
+				break;
+			}
+		}
+		return ParentRow;
+	}
+
 	RowHandle GetTableViewerUiRow(const ICoreProvider* DataStorage, const RowHandle ChildWidgetRow)
 	{
 		const FHierarchyHandle WidgetHierarchy = DataStorage->FindHierarchyByName(GetWidgetHierarchyName());
diff --git a/Engine/Plugins/Experimental/EditorDataStorageFeatures/Source/TedsTableViewer/Public/TedsTableViewerUtils.h b/Engine/Plugins/Experimental/EditorDataStorageFeatures/Source/TedsTableViewer/Public/TedsTableViewerUtils.h
--- a/Engine/Plugins/Experimental/EditorDataStorageFeatures/Source/TedsTableViewer/Public/TedsTableViewerUtils.h
+++ b/Engine/Plugins/Experimental/EditorDataStorageFeatures/Source/TedsTableViewer/Public/TedsTableViewerUtils.h
@@ -50,6 +50,9 @@ namespace UE::Editor::DataStorage::TableViewerUtils
 	// Creates the default Row Map Key for the Table Row UI Rows from the RowHandle directly
     TEDSTABLEVIEWER_API FString CreateRowMappingKey(const FName& HierarchyIdentifier, const RowHandle Item);
 
+	// Walks up the UI Widget Hierarchy until finding the row with the FTableViewerRowTag
+	TEDSTABLEVIEWER_API RowHandle GetTableViewerRowUiRow(const ICoreProvider* DataStorage, const RowHandle ChildWidgetRow);
+
 	// Walks up the UI Widget Hierarchy until finding the row with the FTableViewerTag
 	TEDSTABLEVIEWER_API RowHandle GetTableViewerUiRow(const ICoreProvider* DataStorage, const RowHandle ChildWidgetRow);
 } // namespace UE::Editor::DataStorage::TableViewerUtils
