[EntityFramework] Introduced Asset Reference Filter to new drag&drop workflow.

#Synced-CL 49202658

diff --git a/Engine/Plugins/Experimental/EditorDataStorageFeatures/Source/TedsOperations/Private/DragAndDrop/Widgets/WidgetDropHandler.cpp b/Engine/Plugins/Experimental/EditorDataStorageFeatures/Source/TedsOperations/Private/DragAndDrop/Widgets/WidgetDropHandler.cpp
--- a/Engine/Plugins/Experimental/EditorDataStorageFeatures/Source/TedsOperations/Private/DragAndDrop/Widgets/WidgetDropHandler.cpp
+++ b/Engine/Plugins/Experimental/EditorDataStorageFeatures/Source/TedsOperations/Private/DragAndDrop/Widgets/WidgetDropHandler.cpp
@@ -2,7 +2,6 @@
 
 #include "DragAndDrop/Widgets/WidgetDropHandler.h"
 
-#include "Containers/Ticker.h"
 #include "DataStorage/Features.h"
 #include "DragAndDrop/DecoratedDragDropOp.h"
 #include "DragAndDrop/DropOperationSystem.h"
@@ -11,14 +10,48 @@
 #include "Elements/Framework/EngineElementsLibrary.h"
 #include "Elements/Framework/TypedElementSelectionSet.h"
 #include "Elements/Interfaces/TypedElementDataStorageInterface.h"
+#include "Framework/Notifications/NotificationManager.h"
 #include "Input/DragAndDrop.h"
+#include "Logging/LogMacros.h"
 #include "ScopedTransaction.h"
 #include "TedsOperationInput.h"
+#include "Widgets/Notifications/SNotificationList.h"
 
 #define LOCTEXT_NAMESPACE "WidgetDropHandler"
 
+DEFINE_LOG_CATEGORY_STATIC(LogWidgetDropHandler, Log, All);
+
 namespace UE::Editor
 {
+void FWidgetDropHandler::ShowErrorNotifications(const DataStorage::ICoreProvider& Storage, DataStorage::FRowHandleArrayView InputRows)
+{
+	using namespace UE::Editor::DataStorage;
+
+	FNotificationInfo NotificationInfo(LOCTEXT("ErrorTitle", "Drop Operation Error"));
+	NotificationInfo.FadeInDuration = 2.0f;
+	NotificationInfo.FadeOutDuration = 2.0f;
+	NotificationInfo.ExpireDuration = 5.0f;
+	NotificationInfo.bUseSuccessFailIcons = false;
+	NotificationInfo.Image = FAppStyle::GetBrush("Icons.WarningWithColor");
+	
+	for (RowHandle InputRow : InputRows)
+	{
+		if (!Storage.HasColumns<Operations::FResultColumn>(InputRow))
+		{
+			if (const FText* Reason = Operations::Utilities::GetDescriptionPtr(Storage, InputRow); Reason && !Reason->IsEmpty())
+			{
+				UE_LOG(LogWidgetDropHandler, Warning, TEXT("Drop operation failed: %s"), *Reason->ToString());
+				
+				NotificationInfo.SubText = *Reason;
+				if (TSharedPtr<SNotificationItem> Notification = FSlateNotificationManager::Get().AddNotification(NotificationInfo))
+				{
+					Notification->SetCompletionState(SNotificationItem::CS_Fail);
+					Notification->ExpireAndFadeout();
+				}
+			}
+		}
+	}
+}
 
 FWidgetDropHandler::FWidgetDropHandler(TWeakObjectPtr<UTypedElementSelectionSet> SelectionSet)
 	: SelectionSet(MoveTemp(SelectionSet))
@@ -261,7 +294,10 @@ void FWidgetDropHandler::Drop(DataStorage::ICoreProvider& Storage, UDropOperatio
 		}
 		
 		FScopedTransaction Transaction(LOCTEXT("Drop_Transaction", "Drop on widget."));
-		(void)DropSystem.Apply(InputRows.GetRows());
+		if (DropSystem.Apply(InputRows.GetRows()) != InputRows.Num())
+		{
+			ShowErrorNotifications(Storage, InputRows.GetRows());
+		}
 
 		// Now we select the roots of all captured actors.
 		if (SelectionSetPtr)
diff --git a/Engine/Plugins/Experimental/EditorDataStorageFeatures/Source/TedsOperations/Public/DragAndDrop/Widgets/WidgetDropHandler.h b/Engine/Plugins/Experimental/EditorDataStorageFeatures/Source/TedsOperations/Public/DragAndDrop/Widgets/WidgetDropHandler.h
--- a/Engine/Plugins/Experimental/EditorDataStorageFeatures/Source/TedsOperations/Public/DragAndDrop/Widgets/WidgetDropHandler.h
+++ b/Engine/Plugins/Experimental/EditorDataStorageFeatures/Source/TedsOperations/Public/DragAndDrop/Widgets/WidgetDropHandler.h
@@ -43,6 +43,8 @@ protected:
 	};
 
 public:
+	UE_API static void ShowErrorNotifications(const DataStorage::ICoreProvider& Storage, DataStorage::FRowHandleArrayView InputRows);
+	
 	UE_API FWidgetDropHandler(TWeakObjectPtr<UTypedElementSelectionSet> SelectionSet);
 	UE_API virtual ~FWidgetDropHandler();
 
diff --git a/Engine/Plugins/Experimental/EditorDataStorageFeatures/Source/TedsOutliner/Private/TedsOutlinerMode.cpp b/Engine/Plugins/Experimental/EditorDataStorageFeatures/Source/TedsOutliner/Private/TedsOutlinerMode.cpp
--- a/Engine/Plugins/Experimental/EditorDataStorageFeatures/Source/TedsOutliner/Private/TedsOutlinerMode.cpp
+++ b/Engine/Plugins/Experimental/EditorDataStorageFeatures/Source/TedsOutliner/Private/TedsOutlinerMode.cpp
@@ -12,6 +12,7 @@
 #include "TedsOutlinerHierarchy.h"
 #include "TedsOutlinerItem.h"
 #include "ToolMenus.h"
+#include "DragAndDrop/Widgets/WidgetDropHandler.h"
 
 #include UE_INLINE_GENERATED_CPP_BY_NAME(TedsOutlinerMode)
 
@@ -218,7 +219,10 @@ void FTedsOutlinerMode::OnDrop(ISceneOutlinerTreeItem& DropTarget, const FSceneO
 				FRowHandleArray InputRows;
 				if (DropSystem->CreateInputRows(InputRows, DraggedRowHandles.GetRows(), TedsItem->GetRowHandle(), true))
 				{
-					(void)DropSystem->Apply(InputRows.GetRows());
+					if (DropSystem->Apply(InputRows.GetRows()) != InputRows.Num())
+					{
+						FWidgetDropHandler::ShowErrorNotifications(*Storage, InputRows.GetRows());
+					}
 					DropSystem->RemoveInputRows(InputRows.GetRows());
 				}
 			}
