TedsSettings : remove the workaround from 49071559 (SetParentRowOnceBothRowsAreAssigned)

#Synced-CL 49457093

diff --git a/Engine/Plugins/Experimental/EditorDataStorageFeatures/Source/TedsSettings/Private/TedsSettingsManager.cpp b/Engine/Plugins/Experimental/EditorDataStorageFeatures/Source/TedsSettings/Private/TedsSettingsManager.cpp
--- a/Engine/Plugins/Experimental/EditorDataStorageFeatures/Source/TedsSettings/Private/TedsSettingsManager.cpp
+++ b/Engine/Plugins/Experimental/EditorDataStorageFeatures/Source/TedsSettings/Private/TedsSettingsManager.cpp
@@ -421,8 +421,11 @@ void FTedsSettingsManager::UpdateSettingsCategory(TSharedPtr<ISettingsCategory>
 	ICompatibilityProvider* DataStorageCompatibility = GetMutableDataStorageFeature<ICompatibilityProvider>(CompatibilityFeatureName);
 	check(DataStorageCompatibility);
 
-	const FName& ContainerName = DataStorage->GetColumn<FSettingsNameColumn>(ContainerRow)->Name;
-	const FName& CategoryName = SettingsCategory->GetName();
+	const FSettingsNameColumn* ContainerNameColumn = DataStorage->GetColumn<FSettingsNameColumn>(ContainerRow);
+	checkf(ContainerNameColumn, TEXT("FSettingsNameColumn must exist on ContainerRow: %llu"), ContainerRow);
+
+	const FName ContainerName = ContainerNameColumn->Name;
+	const FName CategoryName = SettingsCategory->GetName();
 
 	UE_LOG(LogTedsSettings, Verbose, TEXT("Update Settings Category: '%s->%s'"), *ContainerName.ToString(), *CategoryName.ToString());
 
@@ -453,7 +456,7 @@ void FTedsSettingsManager::UpdateSettingsCategory(TSharedPtr<ISettingsCategory>
 		using namespace UE::Editor::DataStorage::Queries;
 
 		DataStorage->RunQuery(SelectAllSettingsQuery, CreateDirectQueryCallbackBinding(
-			[&OldRowHandles, &OldSectionNames, &ContainerName, &CategoryName](
+			[&OldRowHandles, &OldSectionNames, ContainerName, CategoryName](
 				IDirectQueryContext& Context,
 				const FSettingsContainerReferenceColumn* ContainerColumns,
 				const FSettingsCategoryReferenceColumn* CategoryColumns,
@@ -537,12 +540,12 @@ void FTedsSettingsManager::UpdateSettingsCategory(TSharedPtr<ISettingsCategory>
 			UE_CLOG(OldSectionRow != NewSectionRow, LogTedsSettings, Verbose, TEXT("Added Settings Section : '%s'"), *SectionName.ToString());
 			UE_CLOG(OldSectionRow == NewSectionRow, LogTedsSettings, Verbose, TEXT("Updated Settings Section : '%s'"), *SectionName.ToString());
 
-			SetParentRowOnceBothRowsAreAssigned(SettingsHierarchy, NewSectionRow, CategoryRow);
-			
-			DataStorage->MapRow(Settings::MappingDomain, MoveTemp(SectionIndexKey), NewSectionRow);
-
 			if (OldSectionRow != NewSectionRow)
 			{
+				DataStorage->SetParentRow(SettingsHierarchy, NewSectionRow, CategoryRow);
+
+				DataStorage->MapRow(Settings::MappingDomain, MoveTemp(SectionIndexKey), NewSectionRow);
+
 				for (FProperty* Property = SettingsObjectPtr->GetClass()->PropertyLink; Property; Property = Property->PropertyLinkNext)
 				{
 					if (Property->HasAnyPropertyFlags(CPF_Transient) || !Property->HasAnyPropertyFlags(CPF_Edit))
@@ -601,41 +604,3 @@ void FTedsSettingsManager::UpdateSettingsCategory(TSharedPtr<ISettingsCategory>
 		UE_LOG(LogTedsSettings, Verbose, TEXT("Removed Settings Section : '%s'"), *OldSectionName.ToString());
 	}
 }
-
-void FTedsSettingsManager::SetParentRowOnceBothRowsAreAssigned(
-	UE::Editor::DataStorage::FHierarchyHandle Hierarchy,
-	UE::Editor::DataStorage::RowHandle Target,
-	UE::Editor::DataStorage::RowHandle Parent)
-{
-	using namespace UE::Editor::DataStorage;
-
-	TRACE_CPUPROFILER_EVENT_SCOPE(TedsSettingsManager.SetParentRowOnceBothRowsAreAssigned);
-
-	ICoreProvider* DataStorage = GetMutableDataStorageFeature<ICoreProvider>(StorageFeatureName);
-
-	// If the hierarchy is no longer valid or either row is no longer available, then we abort the attempt to establish the relationship.
-	if (!DataStorage ||
-		!DataStorage->IsValidHierarchyHandle(Hierarchy) ||
-		!DataStorage->IsRowAvailable(Target) ||
-		!DataStorage->IsRowAvailable(Parent))
-	{
-		return;
-	}
-
-	// If both rows are assigned, then it is now safe to establish the relationship.
-	if (DataStorage->IsRowAssigned(Target) && DataStorage->IsRowAssigned(Parent))
-	{
-		DataStorage->SetParentRow(Hierarchy, Target, Parent);
-	}
-	// Else one of the rows must be available but not yet assigned, so defer another tick and try again.
-	else
-	{
-		ExecuteOnGameThread(UE_SOURCE_LOCATION, [WeakThis = SharedThis(this).ToWeakPtr(), Hierarchy, Target, Parent]()
-		{
-			if (TSharedPtr<FTedsSettingsManager> This = WeakThis.Pin())
-			{
-				This->SetParentRowOnceBothRowsAreAssigned(Hierarchy, Target, Parent);
-			}
-		});
-	}
-}
diff --git a/Engine/Plugins/Experimental/EditorDataStorageFeatures/Source/TedsSettings/Private/TedsSettingsManager.h b/Engine/Plugins/Experimental/EditorDataStorageFeatures/Source/TedsSettings/Private/TedsSettingsManager.h
--- a/Engine/Plugins/Experimental/EditorDataStorageFeatures/Source/TedsSettings/Private/TedsSettingsManager.h
+++ b/Engine/Plugins/Experimental/EditorDataStorageFeatures/Source/TedsSettings/Private/TedsSettingsManager.h
@@ -46,13 +46,6 @@ private:
 
 	void UpdateSettingsCategory(TSharedPtr<ISettingsCategory> SettingsCategory, UE::Editor::DataStorage::RowHandle ContainerRow, const bool bQueryExistingRows = true);
 
-	// This is a workaround for an issue where calling SetParentRow with a not-yet-assigned Target (child) row
-	// sometimes fails to establish the relationship.
-	void SetParentRowOnceBothRowsAreAssigned(
-		UE::Editor::DataStorage::FHierarchyHandle Hierarchy,
-		UE::Editor::DataStorage::RowHandle Target,
-		UE::Editor::DataStorage::RowHandle Parent);
-
 	bool bIsInitialized;
 	UE::Editor::DataStorage::QueryHandle SelectAllSettingsQuery;
 	UE::Editor::DataStorage::TableHandle SettingsContainerTable;
