TEDS-Asset data:
- Add processors to sync favorites back from TEDS to the CB
- Add an FUnresolvedFavoriteTag used to evaluate favorites that come in before TEDS has processed the path (e.g from an .ini read)
- Add sync tags to asset rows to track changes to favorites (without using FUpdatedPathTag which would cause a re-evaluation of the path by TEDS asset data)
- this is all still disabled behind a cvar until perf tested more

#Synced-CL 49591093

diff --git a/Engine/Plugins/Experimental/EditorDataStorageFeatures/Source/TedsAssetData/Private/Processors/AssetProcessors.cpp b/Engine/Plugins/Experimental/EditorDataStorageFeatures/Source/TedsAssetData/Private/Processors/AssetProcessors.cpp
--- a/Engine/Plugins/Experimental/EditorDataStorageFeatures/Source/TedsAssetData/Private/Processors/AssetProcessors.cpp
+++ b/Engine/Plugins/Experimental/EditorDataStorageFeatures/Source/TedsAssetData/Private/Processors/AssetProcessors.cpp
@@ -27,68 +27,15 @@
 namespace UE::TedsAssetDataFactory::Private
 {
 	TAutoConsoleVariable<bool> CVarTedsAssetDataFactory(TEXT("TEDS.TedsAssetDataFactory"),
-		false,
+		true,
 		TEXT("When true this will enable some experimental features that are not optimized to work at scale yet. Note: The value need to be set a boot time to see the effect of this cvar."));
 
-	// @return true if any of the texture dimensions are not a power of 2
-	bool IsTextureNonSquare(FStringView Dimensions)
-	{
-		// Texture dimensions are of the form XxY (2D) or XxYxZ (3D)
-
-		// Extract the X Dimension first (everything to the left of the first x)
-		int32 FirstXPos;
-		if (!Dimensions.FindChar('x', FirstXPos))
-		{
-			return false; // Failsafe, in case the dimension isn't in the correct format somehow.
-		}
-		
-		FStringView XDimension(Dimensions.Left(FirstXPos));
-		int32 X;
-
-		// Convert the X Dimension to an int
-		LexFromString(X, XDimension);
-
-		// If it isn't a power of 2, early out since we already know this is a non square without needing to check the rest
-		if (!FMath::IsPowerOfTwo(X))
-		{
-			return true;
-		}
-
-		// Extract out the YDimension, or YxZ in case of a 3D texture
-		FStringView YDimension(Dimensions.RightChop(FirstXPos + 1));
-		int32 SecondXPos;
-
-		// If we find another x, this is a 3D texture
-		if (YDimension.FindChar('x', SecondXPos))
-		{
-			// Extract out the Z Dimension first, which is everything to the right of the second X. And then convert it to an int
-			FStringView ZDimension = YDimension.RightChop(SecondXPos + 1);
-			int32 Z;
-			LexFromString(Z, ZDimension);
-
-			// Early out if the Z Dimension is not a power of 2 so we don't need to check Y Dimension
-			if (!FMath::IsPowerOfTwo(Z))
-			{
-				return true;
-			}
-
-			// The Y Dimension is everything between the two x's for a 3D texture
-			YDimension = YDimension.Left(SecondXPos);
-		}
-
-		// Regardless of 2D/3D, once we are here YDimension should be the correct value so check it finally
-		int32 Y;
-		LexFromString(Y, YDimension);
-
-		if (!FMath::IsPowerOfTwo(Y))
-		{
-			return true;
-		}
-
-		// If we are here, all dimensions are a power of 2 and the texture is square.
-		return false;
+	static FName UnresolvedFavoritesTable(TEXT("UnresolvedFavorites"));
+}
 
-	}
+void UTedsAssetDataFactory::RegisterTables(UE::Editor::DataStorage::ICoreProvider& DataStorage)
+{
+	DataStorage.RegisterTable<FUnresolvedFavoriteColumn>(UE::TedsAssetDataFactory::Private::UnresolvedFavoritesTable);
 }
 
 void UTedsAssetDataFactory::RegisterQueries(UE::Editor::DataStorage::ICoreProvider& DataStorage)
@@ -121,7 +68,7 @@ void UTedsAssetDataFactory::RegisterQueries(UE::Editor::DataStorage::ICoreProvid
 			}
 		)
 		.Where()
-			.All<FFolderTag, FUpdatedPathTag, FVirtualPathColumn_Experimental>()
+			.All<FFolderTag, FTypedElementSyncFromWorldTag, FVirtualPathColumn_Experimental>()
 		.Compile()
 		);
 	
@@ -167,37 +114,72 @@ void UTedsAssetDataFactory::RegisterQueries(UE::Editor::DataStorage::ICoreProvid
 		.Compile()
 		);
 
+	struct FSetFolderFavoriteCommand
+	{
+		void operator()()
+		{
+			UE::Editor::ContentBrowser::ExtensionUtils::SetFolderFavorite(FolderPath, bFavorite);
+		}
+		FString FolderPath;
+		bool bFavorite;
+	};
+
 	DataStorage.RegisterQuery(
 		Select(
-			TEXT("TedsAssetDataFactory: Add/Remove non-square texture warning"),
-			FProcessor(EQueryTickPhase::PostPhysics, DataStorage.GetQueryTickGroupName(EQueryTickGroups::SyncExternalToDataStorage)),
-			[](IQueryContext& Context, const RowHandle* Rows)
+			TEXT("TedsAssetDataFactory: Sync favorite folder back to world"),
+			FProcessor(EQueryTickPhase::PrePhysics, DataStorage.GetQueryTickGroupName(EQueryTickGroups::SyncDataStorageToExternal))
+			.SetExecutionMode(EExecutionMode::GameThread),
+			[](IQueryContext& Context, RowHandle Row, const FAssetPathColumn_Experimental& AssetPathColumn)
 			{
-				const int32 NumRows = Context.GetRowCount();
-				static const FName AlertName = "NonSquareTextureAlert";
+				// Only do work if the state in the world doesn't match TEDS already
+				if (!UE::Editor::ContentBrowser::ExtensionUtils::IsFolderFavorite(AssetPathColumn.Path.ToString()))
+				{
+					Context.PushCommand(FSetFolderFavoriteCommand{.FolderPath = AssetPathColumn.Path.ToString(), .bFavorite = true});
+				}
+			}
+		)
+		.Where()
+			.All<FFolderTag, FTypedElementSyncBackToWorldTag, FFavoriteTag>()
+		.Compile()
+		);
 
-				TConstArrayView<FItemStringAttributeColumn_Experimental> ColumnView = MakeConstArrayView(
-								Context.GetColumn<FItemStringAttributeColumn_Experimental>("Dimensions"),
-								NumRows);
-				
-				for (int32 Index = 0; Index < NumRows; ++Index)
+		DataStorage.RegisterQuery(
+			Select(
+				TEXT("TedsAssetDataFactory: Sync non-favorite folder back to world"),
+				FProcessor(EQueryTickPhase::PrePhysics, DataStorage.GetQueryTickGroupName(EQueryTickGroups::SyncDataStorageToExternal))
+				.SetExecutionMode(EExecutionMode::GameThread),
+				[](IQueryContext& Context, RowHandle Row, const FAssetPathColumn_Experimental& AssetPathColumn)
 				{
-					if (UE::TedsAssetDataFactory::Private::IsTextureNonSquare(ColumnView[Index].Value))
-					{
-						Alerts::AddAlert(Context, Rows[Index], AlertName,
-							LOCTEXT("NonSquareTextureAlert", "Texture has a non-square aspect ratio."),
-							FAlertColumnType::Error);
-					}
-					else
+					// Only do work if the state in the world doesn't match TEDS already
+					if (UE::Editor::ContentBrowser::ExtensionUtils::IsFolderFavorite(AssetPathColumn.Path.ToString()))
 					{
-						Alerts::RemoveAlert(Context, Rows[Index], AlertName);
+						Context.PushCommand(FSetFolderFavoriteCommand{.FolderPath = AssetPathColumn.Path.ToString(), .bFavorite = false});
 					}
 				}
+			)
+			.Where()
+				.All<FFolderTag, FTypedElementSyncBackToWorldTag>()
+				.None<FFavoriteTag>()
+			.Compile()
+			);
+
+	// Query to resolve any favorites that come in before TEDS asset data has a chance to process them
+	// TODO: Make Activable after the processing is complete
+	DataStorage.RegisterQuery(
+		Select(
+			TEXT("Resolve Favorites"),
+			FProcessor(EQueryTickPhase::PostPhysics, DataStorage.GetQueryTickGroupName(EQueryTickGroups::SyncExternalToDataStorage)),
+			[](IQueryContext& Context, const RowHandle Row, const FUnresolvedFavoriteColumn& UnresolvedFavoriteColumn)
+			{
+				RowHandle AssetRow = Context.LookupMappedRow(UE::Editor::AssetData::MappingDomain, UnresolvedFavoriteColumn.MappingIndex);
+				if (Context.IsRowAvailable(AssetRow))
+				{
+					Context.AddColumns<FTypedElementSyncFromWorldTag>(AssetRow);
+					Context.AddColumns<FFavoriteTag>(AssetRow);
+					Context.RemoveRow(Row);
+				}
 			}
 		)
-		.ReadOnly<FItemStringAttributeColumn_Experimental>("Dimensions")
-		.Where()
-			.All<FAssetTag, FUpdatedAssetDataTag>()
 		.Compile()
 		);
 }
@@ -244,7 +226,7 @@ void UTedsAssetDataFactory::OnSetFolderColor(const FString& Path, UE::Editor::Da
 
 	if (DataStorage->IsRowAvailable(Row))
 	{
-		DataStorage->AddColumn<FUpdatedPathTag>(Row);
+		DataStorage->AddColumn<FTypedElementSyncFromWorldTag>(Row);
 	}
 }
 
@@ -257,7 +239,34 @@ void UTedsAssetDataFactory::OnFavoritesChanged(const FContentBrowserItemPath& It
 
 		if (DataStorage->IsRowAvailable(Row))
 		{
-			bAdded ? DataStorage->AddColumn<FFavoriteTag>(Row) : DataStorage->RemoveColumn<FFavoriteTag>(Row);
+			// Since we are adding the sync tag which could end up doing other expensive processing, only add it when the data actually changed
+			// from what TEDS has
+			if (bAdded)
+			{
+				if (!DataStorage->HasColumns<FFavoriteTag>(Row))
+				{
+					DataStorage->AddColumn<FTypedElementSyncFromWorldTag>(Row);
+					DataStorage->AddColumn<FFavoriteTag>(Row);
+				}
+			}
+			else
+			{
+				if (DataStorage->HasColumns<FFavoriteTag>(Row))
+				{
+					DataStorage->AddColumn<FTypedElementSyncFromWorldTag>(Row);
+					DataStorage->RemoveColumn<FFavoriteTag>(Row);
+				}
+			}
+		}
+		else if (bAdded) // if it's being removed and doesn't exist we don't have anything to do
+		{
+			UE::Editor::DataStorage::RowHandle TempRow = DataStorage->AddRow(
+				DataStorage->FindTable(UE::TedsAssetDataFactory::Private::UnresolvedFavoritesTable));
+
+			if (FUnresolvedFavoriteColumn* Column = DataStorage->GetColumn<FUnresolvedFavoriteColumn>(TempRow))
+			{
+				Column->MappingIndex = PathKey;
+			}
 		}
 	}
 	
diff --git a/Engine/Plugins/Experimental/EditorDataStorageFeatures/Source/TedsAssetData/Private/Processors/AssetProcessors.h b/Engine/Plugins/Experimental/EditorDataStorageFeatures/Source/TedsAssetData/Private/Processors/AssetProcessors.h
--- a/Engine/Plugins/Experimental/EditorDataStorageFeatures/Source/TedsAssetData/Private/Processors/AssetProcessors.h
+++ b/Engine/Plugins/Experimental/EditorDataStorageFeatures/Source/TedsAssetData/Private/Processors/AssetProcessors.h
@@ -2,6 +2,8 @@
 
 #pragma once
 
+#include "DataStorage/CommonTypes.h"
+#include "DataStorage/MapKey.h"
 #include "Elements/Interfaces/TypedElementDataStorageFactory.h"
 
 #include "AssetProcessors.generated.h"
@@ -13,6 +15,17 @@ namespace UE::Editor::DataStorage
 	class ICoreProvider;
 } // namespace UE::Editor::DataStorage
 
+/**
+ * Temp tag added to a row when it is a favorite before TEDS asset data has processed it
+ */
+USTRUCT(meta = (DisplayName = "Unresolved Favorite"))
+struct FUnresolvedFavoriteColumn final : public FEditorDataStorageColumn
+{
+	GENERATED_BODY()
+
+	UE::Editor::DataStorage::FMapKey MappingIndex;
+};
+
 UCLASS()
 class UTedsAssetDataFactory : public UEditorDataStorageFactory
 {
@@ -21,6 +34,7 @@ class UTedsAssetDataFactory : public UEditorDataStorageFactory
 public:
 	~UTedsAssetDataFactory() override = default;
 
+	virtual void RegisterTables(UE::Editor::DataStorage::ICoreProvider& DataStorage) override;
 	void RegisterQueries(UE::Editor::DataStorage::ICoreProvider& DataStorage) override;
 	virtual void PreRegister(UE::Editor::DataStorage::ICoreProvider& DataStorage) override;
 	virtual void PreShutdown(UE::Editor::DataStorage::ICoreProvider& DataStorage) override;
diff --git a/Engine/Source/Editor/ContentBrowser/Private/Experimental/ContentBrowserExtensionUtils.cpp b/Engine/Source/Editor/ContentBrowser/Private/Experimental/ContentBrowserExtensionUtils.cpp
--- a/Engine/Source/Editor/ContentBrowser/Private/Experimental/ContentBrowserExtensionUtils.cpp
+++ b/Engine/Source/Editor/ContentBrowser/Private/Experimental/ContentBrowserExtensionUtils.cpp
@@ -61,6 +61,19 @@ namespace UE::Editor::ContentBrowser::ExtensionUtils
 
 	bool IsFolderFavorite(const FString& FolderPath)
 	{
-		return ContentBrowserUtils::IsFavoriteFolder(FContentBrowserItemPath(FolderPath, EContentBrowserPathType::Virtual));
+		return ContentBrowserUtils::IsFavoriteFolder(FContentBrowserItemPath(FolderPath, EContentBrowserPathType::Internal));
+	}
+
+	void SetFolderFavorite(const FString& FolderPath, bool bIsFavorite)
+	{
+		if (bIsFavorite)
+		{
+			ContentBrowserUtils::AddFavoriteFolder(FContentBrowserItemPath(FolderPath, EContentBrowserPathType::Internal));
+		}
+		else
+		{
+			ContentBrowserUtils::RemoveFavoriteFolder(FContentBrowserItemPath(FolderPath, EContentBrowserPathType::Internal));
+		}
+		
 	}
 }
diff --git a/Engine/Source/Editor/ContentBrowser/Public/Experimental/ContentBrowserExtensionUtils.h b/Engine/Source/Editor/ContentBrowser/Public/Experimental/ContentBrowserExtensionUtils.h
--- a/Engine/Source/Editor/ContentBrowser/Public/Experimental/ContentBrowserExtensionUtils.h
+++ b/Engine/Source/Editor/ContentBrowser/Public/Experimental/ContentBrowserExtensionUtils.h
@@ -11,17 +11,27 @@ struct FContentBrowserItemPath;
 namespace UE::Editor::ContentBrowser::ExtensionUtils
 {
 	// Get the custom color of the given folder (if any). Prefer using this over raw FName variants.
+	UE_EXPERIMENTAL(5.8, "The Content Browser extension utils are experimental hooks that may be removed in the future")
 	CONTENTBROWSER_API TOptional<FLinearColor> GetFolderColor(const FContentBrowserItem& FolderItem);
 
 	// Get the custom color of the given folder path (if any)
+	UE_EXPERIMENTAL(5.8, "The Content Browser extension utils are experimental hooks that may be removed in the future")
 	CONTENTBROWSER_API TOptional<FLinearColor> GetFolderColor(const FContentBrowserItemPath& FolderPath);
 
 	// Get the custom color of the given internal folder path (if any)
+	UE_EXPERIMENTAL(5.8, "The Content Browser extension utils are experimental hooks that may be removed in the future")
 	CONTENTBROWSER_API TOptional<FLinearColor> GetFolderColor(const FName& FolderPath);
 
 	// Set a custom color for the given folder path
+	UE_EXPERIMENTAL(5.8, "The Content Browser extension utils are experimental hooks that may be removed in the future")
 	CONTENTBROWSER_API void SetFolderColor(const FName& FolderPath, const FLinearColor& FolderColor);
 
 	// Returns if this folder has been marked as a favorite folder
+	UE_EXPERIMENTAL(5.8, "The Content Browser extension utils are experimental hooks that may be removed in the future")
 	CONTENTBROWSER_API bool IsFolderFavorite(const FString& FolderPath);
+
+	// Set if a folder is favorite
+	UE_EXPERIMENTAL(5.8, "The Content Browser extension utils are experimental hooks that may be removed in the future")
+	CONTENTBROWSER_API void SetFolderFavorite(const FString& FolderPath, bool bIsFavorite);
+	
 }
