New drag&drop workflow now fails if not all drop data can be converted to teds rows.

#Synced-CL 49647438

diff --git a/Engine/Plugins/Experimental/EditorDataStorageFeatures/Source/TedsOperations/Private/DragAndDrop/TedsDragDropOpUtility.cpp b/Engine/Plugins/Experimental/EditorDataStorageFeatures/Source/TedsOperations/Private/DragAndDrop/TedsDragDropOpUtility.cpp
--- a/Engine/Plugins/Experimental/EditorDataStorageFeatures/Source/TedsOperations/Private/DragAndDrop/TedsDragDropOpUtility.cpp
+++ b/Engine/Plugins/Experimental/EditorDataStorageFeatures/Source/TedsOperations/Private/DragAndDrop/TedsDragDropOpUtility.cpp
@@ -14,46 +14,73 @@
 namespace UE::Editor::DataStorage::DragAndDrop
 {
 
-void GetRowsFromData(FRowHandleArray& OutRows, const FDragDropOperation& InData)
+bool GetRowsFromData(FRowHandleArray& OutRows, const FDragDropOperation& InData)
 {
 	if (InData.IsOfType<Widgets::FTedsDragDropOp>())
 	{
 		OutRows = static_cast<const Widgets::FTedsDragDropOp&>(InData).GetRowArray();
+		return true;
 	}
 	else if (InData.IsOfType<FAssetDragDropOp>())
 	{
 		if (const ICoreProvider* Storage = GetDataStorageFeature<ICoreProvider>(StorageFeatureName))
 		{
+			bool bResult = true;
 			for (const FAssetData& Element : static_cast<const FAssetDragDropOp&>(InData).GetAssets())
 			{
-				if (Element.IsValid())
+				if (!Element.IsValid())
 				{
-					RowHandle Row = Storage->LookupMappedRow(/*UE::Editor::AssetData::MappingDomain*/"Asset", FMapKey(Element.GetSoftObjectPath()));
-					if (Row != InvalidRowHandle)
-					{
-						OutRows.Add(Row);
-					}
+					bResult = false;
+					continue;
 				}
+				
+				RowHandle Row = Storage->LookupMappedRow(/*UE::Editor::AssetData::MappingDomain*/"Asset", FMapKey(Element.GetSoftObjectPath()));
+				if (Row == InvalidRowHandle)
+				{
+					bResult = false;
+					continue;
+				}
+				
+				OutRows.Add(Row);
+			}
+			
+			bool bHasFactory = static_cast<const FAssetDragDropOp&>(InData).GetAssetFactory() != nullptr;
+			if (bHasFactory)
+			{
+				// If an asset factory is defined, we "lose" data by ignoring it, so return false.
+				return false;
 			}
+			
+			return bResult;
 		}
 	}
 	else if (InData.IsOfType<FClassDragDropOp>())
 	{
 		if (const ICoreProvider* Storage = GetDataStorageFeature<ICoreProvider>(StorageFeatureName))
 		{
+			bool bResult = true;
 			for (const TWeakObjectPtr<UClass>& ClassPtr : static_cast<const FClassDragDropOp&>(InData).ClassesToDrop)
 			{
-				if (const UClass* Class = ClassPtr.Get())
+				const UClass* Class = ClassPtr.Get();
+				if (!Class)
 				{
-					RowHandle Row = Storage->LookupMappedRow(/*UE::Editor::AssetData::MappingDomain*/"Asset", FMapKey(Class->GetPathName()));
-					if (Row != InvalidRowHandle)
-					{
-						OutRows.Add(Row);
-					}
+					bResult = false;
+					continue;
 				}
+				
+				RowHandle Row = Storage->LookupMappedRow(/*UE::Editor::AssetData::MappingDomain*/"Asset", FMapKey(Class->GetPathName()));
+				if (Row == InvalidRowHandle)
+				{
+					bResult = false;
+					continue;
+				}
+				
+				OutRows.Add(Row);
 			}
+			return bResult;
 		}
 	}
+	return false;
 }
 
-}
\ No newline at end of file
+}
diff --git a/Engine/Plugins/Experimental/EditorDataStorageFeatures/Source/TedsOperations/Private/DragAndDrop/Widgets/WidgetDropHandler.cpp b/Engine/Plugins/Experimental/EditorDataStorageFeatures/Source/TedsOperations/Private/DragAndDrop/Widgets/WidgetDropHandler.cpp
--- a/Engine/Plugins/Experimental/EditorDataStorageFeatures/Source/TedsOperations/Private/DragAndDrop/Widgets/WidgetDropHandler.cpp
+++ b/Engine/Plugins/Experimental/EditorDataStorageFeatures/Source/TedsOperations/Private/DragAndDrop/Widgets/WidgetDropHandler.cpp
@@ -74,9 +74,11 @@ bool FWidgetDropHandler::OnDragEnter(const FGeometry& Geometry, const FDragDropE
 		{
 			if (UDropOperationSystem* DropSystem = Storage->FindFactory<UDropOperationSystem>())
 			{
-				FDropResult Result = Start(*Storage, *DropSystem, Geometry, InputEvent, *Operation);
-				UpdateOperationDecorator(*Operation, Result);
-				return true;
+				if (TOptional<FDropResult> Result = Start(*Storage, *DropSystem, Geometry, InputEvent, *Operation))
+				{
+					UpdateOperationDecorator(*Operation, Result);
+					return true;
+				}
 			}
 		}
 	}
@@ -157,8 +159,10 @@ bool FWidgetDropHandler::PrepareInput(DataStorage::ICoreProvider& Storage, UDrop
 	using namespace UE::Editor::DataStorage;
 	
 	FRowHandleArray SourceRows;
-	DragAndDrop::GetRowsFromData(SourceRows, Operation);
-	SourceRows.Remove(InvalidRowHandle);
+	if (!DragAndDrop::GetRowsFromData(SourceRows, Operation))
+	{
+		return false; // If we lose data by converting to rows, we don't want to continue.
+	}
 	
 	return PrepareInput(Storage, DropSystem, SourceRows.GetRows());
 }
@@ -214,7 +218,7 @@ void FWidgetDropHandler::SetRowsFromParameters(DataStorage::ICoreProvider& Stora
 {
 }
 
-FWidgetDropHandler::FDropResult FWidgetDropHandler::Start(DataStorage::ICoreProvider& Storage, UDropOperationSystem& DropSystem,
+TOptional<FWidgetDropHandler::FDropResult> FWidgetDropHandler::Start(DataStorage::ICoreProvider& Storage, UDropOperationSystem& DropSystem,
                                                           const FGeometry& Geometry, const FDragDropEvent& InputEvent, FDragDropOperation& Operation)
 {
 	using namespace UE::Editor::DataStorage;
@@ -229,7 +233,7 @@ FWidgetDropHandler::FDropResult FWidgetDropHandler::Start(DataStorage::ICoreProv
 	if (!PrepareInput(Storage, DropSystem, Operation))
 	{
 		Stop(Storage, DropSystem);
-		return { false, LOCTEXT("Drop_Unsupported", "Drop operation is unsupported.") };
+		return { };
 	}
 	(void)UpdateParameters(Storage, Geometry, InputEvent);
 	SetRowsFromParameters(Storage, InputRows.GetRows());
diff --git a/Engine/Plugins/Experimental/EditorDataStorageFeatures/Source/TedsOperations/Public/DragAndDrop/TedsDragDropOpUtility.h b/Engine/Plugins/Experimental/EditorDataStorageFeatures/Source/TedsOperations/Public/DragAndDrop/TedsDragDropOpUtility.h
--- a/Engine/Plugins/Experimental/EditorDataStorageFeatures/Source/TedsOperations/Public/DragAndDrop/TedsDragDropOpUtility.h
+++ b/Engine/Plugins/Experimental/EditorDataStorageFeatures/Source/TedsOperations/Public/DragAndDrop/TedsDragDropOpUtility.h
@@ -15,7 +15,8 @@ class FRowHandleArray;
 namespace UE::Editor::DataStorage::DragAndDrop
 {
 
-UE_API void GetRowsFromData(FRowHandleArray& OutRows, const FDragDropOperation& InData);
+/** Tries to convert the dropped data into rows. Returns true if all data was successfully converted. */
+UE_API bool GetRowsFromData(FRowHandleArray& OutRows, const FDragDropOperation& InData);
 
 }
 
diff --git a/Engine/Plugins/Experimental/EditorDataStorageFeatures/Source/TedsOperations/Public/DragAndDrop/Widgets/WidgetDropHandler.h b/Engine/Plugins/Experimental/EditorDataStorageFeatures/Source/TedsOperations/Public/DragAndDrop/Widgets/WidgetDropHandler.h
--- a/Engine/Plugins/Experimental/EditorDataStorageFeatures/Source/TedsOperations/Public/DragAndDrop/Widgets/WidgetDropHandler.h
+++ b/Engine/Plugins/Experimental/EditorDataStorageFeatures/Source/TedsOperations/Public/DragAndDrop/Widgets/WidgetDropHandler.h
@@ -67,11 +67,15 @@ protected:
 	/** Method to assign cached parameters to the given rows. */
 	UE_API virtual void SetRowsFromParameters(DataStorage::ICoreProvider& Storage, DataStorage::FRowHandleArrayView Rows) const;
 	
-	UE_API virtual FDropResult Start(DataStorage::ICoreProvider& Storage, UDropOperationSystem& DropSystem, const FGeometry& Geometry,
+	/** Starts a drag-operation internally. Returns an unset value if the operation is ignored, a result with True if the drop operation would be valid or False if invalid. */
+	UE_API virtual TOptional<FDropResult> Start(DataStorage::ICoreProvider& Storage, UDropOperationSystem& DropSystem, const FGeometry& Geometry,
 		const FDragDropEvent& InputEvent, FDragDropOperation& Operation);
+	/** Updates the current drag-operation internally. Returns an updated result or unset optional if nothing has changed. */
 	UE_API virtual TOptional<FDropResult> Update(DataStorage::ICoreProvider& Storage, UDropOperationSystem& DropSystem, const FGeometry& Geometry,
 		const FDragDropEvent& InputEvent);
+	/** Stops the current drag-operation. */
 	UE_API virtual void Stop(DataStorage::ICoreProvider& Storage, UDropOperationSystem& DropSystem);
+	/** Stops and executes the drop for the current drag-operation. */
 	UE_API virtual void Drop(DataStorage::ICoreProvider& Storage, UDropOperationSystem& DropSystem, const FGeometry& Geometry,
 		const FDragDropEvent& InputEvent);
 
diff --git a/Engine/Plugins/Experimental/EditorDataStorageFeatures/Source/TedsOutliner/Private/TedsOutlinerMode.cpp b/Engine/Plugins/Experimental/EditorDataStorageFeatures/Source/TedsOutliner/Private/TedsOutlinerMode.cpp
--- a/Engine/Plugins/Experimental/EditorDataStorageFeatures/Source/TedsOutliner/Private/TedsOutlinerMode.cpp
+++ b/Engine/Plugins/Experimental/EditorDataStorageFeatures/Source/TedsOutliner/Private/TedsOutlinerMode.cpp
@@ -9,6 +9,7 @@
 #include "DragAndDrop/TedsDragDropOpUtility.h"
 #include "DragAndDrop/DropOperationInput.h"
 #include "DragAndDrop/DropOperationSystem.h"
+#include "ScopedTransaction.h"
 #include "TedsOutlinerHierarchy.h"
 #include "TedsOutlinerItem.h"
 #include "ToolMenus.h"
@@ -23,6 +24,20 @@ namespace UE::Editor::Outliner
 	namespace Private
 	{
 		static FName ContextMenuName("TedsOutlinerContextMenu");
+		
+		static bool GetDropRowsFromPayload(FRowHandleArray& OutRows, const FSceneOutlinerDragDropPayload& Payload)
+		{
+			if (Payload.DraggedItems.IsEmpty()) // External drop
+			{
+				return DragAndDrop::GetRowsFromData(OutRows, Payload.SourceOperation);
+			}
+	
+			Payload.ForEachItem<FTedsOutlinerTreeItem>([&OutRows](const FTedsOutlinerTreeItem& TedsItem)
+			{
+				OutRows.Add(TedsItem.GetRowHandle());
+			});
+			return true;
+		}
 	} // namespace Private
 
 FTedsOutlinerMode::FTedsOutlinerMode(const FTedsOutlinerParams& InParams)
@@ -136,94 +151,82 @@ bool FTedsOutlinerMode::ParseDragDrop(FSceneOutlinerDragDropPayload& OutPayload,
 	}
 	return true; // Always return true since we support external drops as well
 }
-
+	
 FSceneOutlinerDragValidationInfo FTedsOutlinerMode::ValidateDrop(const ISceneOutlinerTreeItem& DropTarget, const FSceneOutlinerDragDropPayload& Payload) const
 {
-	ICoreProvider* Storage = TedsOutlinerImpl->GetStorage();
-	FRowHandleArray DraggedRowHandles;
-
-	// External drop
-	if (Payload.DraggedItems.IsEmpty())
-	{
-		DragAndDrop::GetRowsFromData(DraggedRowHandles, Payload.SourceOperation);
-	}
-	else
-	{
-		Payload.ForEachItem<FTedsOutlinerTreeItem>([&DraggedRowHandles](FTedsOutlinerTreeItem& TedsItem)
-		{
-			DraggedRowHandles.Add(TedsItem.GetRowHandle());
-		});
-	}
-
-	// TODO: The Outliner considers a drop onto empty area as a drop onto an empty FFolderTreeItem, we need to handle that case by providing
-	// that info to the operations somehow
-	if (const FTedsOutlinerTreeItem* TedsDropTarget = DropTarget.CastTo<FTedsOutlinerTreeItem>())
+	if (ICoreProvider* Storage = TedsOutlinerImpl->GetStorage())
 	{
 		if (UDropOperationSystem* DropSystem = Storage->FindFactory<UDropOperationSystem>())
 		{
-			FRowHandleArray InputRows;
-			if (DropSystem->CreateInputRows(InputRows, DraggedRowHandles.GetRows(), TedsDropTarget->GetRowHandle(), true))
+			FRowHandleArray DraggedRowHandles;
+			if (Private::GetDropRowsFromPayload(DraggedRowHandles, Payload) && !DraggedRowHandles.IsEmpty())
 			{
-				bool bResult = DropSystem->Test(InputRows.GetRows()) == InputRows.Num();
-
-				FText Description;
-				for (RowHandle InputRow : InputRows.GetRows())
+				// TODO: The Outliner considers a drop onto empty area as a drop onto an empty FFolderTreeItem, we need to handle that case by providing
+				// that info to the operations somehow
+				if (const FTedsOutlinerTreeItem* TedsDropTarget = DropTarget.CastTo<FTedsOutlinerTreeItem>())
 				{
-					// Find an error text when we failed and a success text if we succeeded.
-					if (bResult != Storage->HasColumns<Operations::FTestResultTag>(InputRow))
+					FRowHandleArray InputRows;
+					if (DropSystem->CreateInputRows(InputRows, DraggedRowHandles.GetRows(), TedsDropTarget->GetRowHandle(), true))
 					{
-						continue;
-					}
+						bool bResult = DropSystem->Test(InputRows.GetRows()) == InputRows.Num();
+
+						FText Description;
+						for (RowHandle InputRow : InputRows.GetRows())
+						{
+							// Find an error text when we failed and a success text if we succeeded.
+							if (bResult != Storage->HasColumns<Operations::FTestResultTag>(InputRow))
+							{
+								continue;
+							}
 		
-					Description = Operations::Utilities::GetDescription(*Storage, InputRow);
-					if (!Description.IsEmpty())
-					{
-						break;
-					}
-				}
+							Description = Operations::Utilities::GetDescription(*Storage, InputRow);
+							if (!Description.IsEmpty())
+							{
+								break;
+							}
+						}
 
-				DropSystem->RemoveInputRows(InputRows.GetRows());
+						DropSystem->RemoveInputRows(InputRows.GetRows());
 
-				return { bResult ? ESceneOutlinerDropCompatibility::CompatibleGeneric : ESceneOutlinerDropCompatibility::IncompatibleGeneric, MoveTemp(Description) };
+						return { bResult ? ESceneOutlinerDropCompatibility::CompatibleGeneric : ESceneOutlinerDropCompatibility::IncompatibleGeneric, MoveTemp(Description) };
+					}
+				}
 			}
 		}
 	}
 
-	return FSceneOutlinerDragValidationInfo(ESceneOutlinerDropCompatibility::IncompatibleGeneric, LOCTEXT("InvalidDrop", "Invalid Drop target"));
+	return { ESceneOutlinerDropCompatibility::IncompatibleGeneric, LOCTEXT("DropUnsupported", "Drop operation is unsupported.") };
 }
 
 void FTedsOutlinerMode::OnDrop(ISceneOutlinerTreeItem& DropTarget, const FSceneOutlinerDragDropPayload& Payload, const FSceneOutlinerDragValidationInfo& ValidationInfo) const
 {
-	ICoreProvider* Storage = TedsOutlinerImpl->GetStorage();
-
-	FRowHandleArray DraggedRowHandles;
-
-	if (Payload.DraggedItems.IsEmpty())
-	{
-		DragAndDrop::GetRowsFromData(DraggedRowHandles, Payload.SourceOperation);
-	}
-	else
+	if (ValidationInfo.CompatibilityType != ESceneOutlinerDropCompatibility::CompatibleGeneric)
 	{
-		Payload.ForEachItem<FTedsOutlinerTreeItem>([&DraggedRowHandles](FTedsOutlinerTreeItem& TedsItem)
-		{
-			DraggedRowHandles.Add(TedsItem.GetRowHandle());
-		});
+		return;
 	}
-
-	if (const FTedsOutlinerTreeItem* TedsItem = DropTarget.CastTo<FTedsOutlinerTreeItem>())
+		
+	if (ICoreProvider* Storage = TedsOutlinerImpl->GetStorage())
 	{
-		if (ValidationInfo.CompatibilityType == ESceneOutlinerDropCompatibility::CompatibleGeneric)
+		if (UDropOperationSystem* DropSystem = Storage->FindFactory<UDropOperationSystem>())
 		{
-			if (UDropOperationSystem* DropSystem = Storage->FindFactory<UDropOperationSystem>())
+			FRowHandleArray DraggedRowHandles;
+			if (Private::GetDropRowsFromPayload(DraggedRowHandles, Payload) && !DraggedRowHandles.IsEmpty())
 			{
-				FRowHandleArray InputRows;
-				if (DropSystem->CreateInputRows(InputRows, DraggedRowHandles.GetRows(), TedsItem->GetRowHandle(), true))
+				if (const FTedsOutlinerTreeItem* TedsItem = DropTarget.CastTo<FTedsOutlinerTreeItem>())
 				{
-					if (DropSystem->Apply(InputRows.GetRows()) != InputRows.Num())
+					FRowHandleArray InputRows;
+					if (DropSystem->CreateInputRows(InputRows, DraggedRowHandles.GetRows(), TedsItem->GetRowHandle(), true))
 					{
+						ON_SCOPE_EXIT{ DropSystem->RemoveInputRows(InputRows.GetRows()); };
+						{
+							FScopedTransaction Transaction(LOCTEXT("Transaction_Drop", "Drop on Outliner."));
+							if (DropSystem->Apply(InputRows.GetRows()) != InputRows.Num())
+							{
+								FWidgetDropHandler::ShowErrorNotifications(*Storage, InputRows.GetRows());
+							}
+						}
 						FWidgetDropHandler::ShowErrorNotifications(*Storage, InputRows.GetRows());
 					}
-					DropSystem->RemoveInputRows(InputRows.GetRows());
 				}
 			}
 		}
