[Outliner][Bug] Fix Package Short Name column not appearing with the TEDS Hybrid Outliner enabled (Level Column was overwritting it since the legacy merges them into one column)

#Synced-CL 49649433

diff --git a/Engine/Plugins/Experimental/EditorDataStorageFeatures/Source/TedsOutliner/Private/Compatibility/SceneOutlinerTedsBridge.cpp b/Engine/Plugins/Experimental/EditorDataStorageFeatures/Source/TedsOutliner/Private/Compatibility/SceneOutlinerTedsBridge.cpp
--- a/Engine/Plugins/Experimental/EditorDataStorageFeatures/Source/TedsOutliner/Private/Compatibility/SceneOutlinerTedsBridge.cpp
+++ b/Engine/Plugins/Experimental/EditorDataStorageFeatures/Source/TedsOutliner/Private/Compatibility/SceneOutlinerTedsBridge.cpp
@@ -84,4 +84,14 @@ FName USceneOutlinerTedsBridgeFactory::FindOutlinerColumnFromTedsColumns(TConstA
 	return NAME_None;
 }
 
+void USceneOutlinerTedsBridgeFactory::RemoveTedsColumnMappingToOutlinerColumn(TWeakObjectPtr<const UScriptStruct> TEDSColumnToRemove)
+{
+	TEDSToOutlinerDefaultColumnMapping.Remove(TEDSColumnToRemove);
+}
+
+void USceneOutlinerTedsBridgeFactory::AddTedsColumnMappingToOutlinerColumn(TWeakObjectPtr<const UScriptStruct> TEDSColumnToAdd, const FName& OutlinerColumnId)
+{
+	TEDSToOutlinerDefaultColumnMapping.Add(TEDSColumnToAdd, OutlinerColumnId);
+}
+
 #undef LOCTEXT_NAMESPACE
diff --git a/Engine/Plugins/Experimental/EditorDataStorageFeatures/Source/TedsOutliner/Private/TedsOutlinerHelpers.cpp b/Engine/Plugins/Experimental/EditorDataStorageFeatures/Source/TedsOutliner/Private/TedsOutlinerHelpers.cpp
--- a/Engine/Plugins/Experimental/EditorDataStorageFeatures/Source/TedsOutliner/Private/TedsOutlinerHelpers.cpp
+++ b/Engine/Plugins/Experimental/EditorDataStorageFeatures/Source/TedsOutliner/Private/TedsOutlinerHelpers.cpp
@@ -84,6 +84,22 @@ namespace UE::Editor::Outliner::Helpers
 
 		return NAME_None;
 	}
+
+	void RemoveTedsColumnMappingToOutlinerColumn(DataStorage::ICoreProvider* Storage, TWeakObjectPtr<const UScriptStruct> TEDSColumnToRemove)
+	{
+		if(USceneOutlinerTedsBridgeFactory* Factory = Storage->FindFactory<USceneOutlinerTedsBridgeFactory>())
+		{
+			Factory->RemoveTedsColumnMappingToOutlinerColumn(TEDSColumnToRemove);
+		}
+	}
+
+	void AddTedsColumnMappingToOutlinerColumn(DataStorage::ICoreProvider* Storage, TWeakObjectPtr<const UScriptStruct> TEDSColumnToAdd, const FName& OutlinerColumnId)
+	{
+		if(USceneOutlinerTedsBridgeFactory* Factory = Storage->FindFactory<USceneOutlinerTedsBridgeFactory>())
+		{
+			Factory->AddTedsColumnMappingToOutlinerColumn(TEDSColumnToAdd, OutlinerColumnId);
+		}
+	}
 	
     bool CheckValidFilterQueryHandle(const DataStorage::QueryHandle& InQueryHandle)
     {
diff --git a/Engine/Plugins/Experimental/EditorDataStorageFeatures/Source/TedsOutliner/Private/TedsOutlinerImpl.cpp b/Engine/Plugins/Experimental/EditorDataStorageFeatures/Source/TedsOutliner/Private/TedsOutlinerImpl.cpp
--- a/Engine/Plugins/Experimental/EditorDataStorageFeatures/Source/TedsOutliner/Private/TedsOutlinerImpl.cpp
+++ b/Engine/Plugins/Experimental/EditorDataStorageFeatures/Source/TedsOutliner/Private/TedsOutlinerImpl.cpp
@@ -223,6 +223,23 @@ FTedsOutlinerImpl::FTedsOutlinerImpl(const FTedsOutlinerParams& InParams, IScene
 	StorageUi = GetMutableDataStorageFeature<DataStorage::IUiProvider>(UiFeatureName);
 	StorageCompatibility = GetMutableDataStorageFeature<DataStorage::ICompatibilityProvider>(CompatibilityFeatureName);
 
+	// Only do this on init to match parity since the legacy outliner only checks for this when opened
+	if (GEditor) // No need to check if ->GetEditorWorldContext is valid as it should always return a valid UWorld reference
+	{
+		if(const UWorld* World = GEditor->GetEditorWorldContext().World())
+		{
+			if (World->PersistentLevel && World->PersistentLevel->IsUsingExternalActors())
+			{
+				// TODO: When the Package Short Name data is available, have it reassign to that column instead of removing it
+				Helpers::RemoveTedsColumnMappingToOutlinerColumn(Storage, FLevelColumn::StaticStruct());
+			}
+			else
+			{
+				Helpers::AddTedsColumnMappingToOutlinerColumn(Storage, FLevelColumn::StaticStruct(), FSceneOutlinerBuiltInColumnTypes::Level());
+			}
+		}
+	}
+	
 	bUsingQueryConditionsSyntax = InitialQueryDescription.Conditions && !InitialQueryDescription.Conditions->IsEmpty();
 }
 
diff --git a/Engine/Plugins/Experimental/EditorDataStorageFeatures/Source/TedsOutliner/Public/Compatibility/SceneOutlinerTedsBridge.h b/Engine/Plugins/Experimental/EditorDataStorageFeatures/Source/TedsOutliner/Public/Compatibility/SceneOutlinerTedsBridge.h
--- a/Engine/Plugins/Experimental/EditorDataStorageFeatures/Source/TedsOutliner/Public/Compatibility/SceneOutlinerTedsBridge.h
+++ b/Engine/Plugins/Experimental/EditorDataStorageFeatures/Source/TedsOutliner/Public/Compatibility/SceneOutlinerTedsBridge.h
@@ -34,6 +34,8 @@ public:
 	virtual void RegisterWidgetPurposes(UE::Editor::DataStorage::IUiProvider& DataStorageUi) const override;
 
 	FName FindOutlinerColumnFromTedsColumns(TConstArrayView<TWeakObjectPtr<const UScriptStruct>> TEDSColumns) const;
+	void RemoveTedsColumnMappingToOutlinerColumn(TWeakObjectPtr<const UScriptStruct> TEDSColumnToRemove);
+	void AddTedsColumnMappingToOutlinerColumn(TWeakObjectPtr<const UScriptStruct> TEDSColumnToAdd, const FName& OutlinerColumnId);
 protected:
 
 	TMap<TWeakObjectPtr<const UScriptStruct>, FName> TEDSToOutlinerDefaultColumnMapping;
diff --git a/Engine/Plugins/Experimental/EditorDataStorageFeatures/Source/TedsOutliner/Public/TedsOutlinerHelpers.h b/Engine/Plugins/Experimental/EditorDataStorageFeatures/Source/TedsOutliner/Public/TedsOutlinerHelpers.h
--- a/Engine/Plugins/Experimental/EditorDataStorageFeatures/Source/TedsOutliner/Public/TedsOutlinerHelpers.h
+++ b/Engine/Plugins/Experimental/EditorDataStorageFeatures/Source/TedsOutliner/Public/TedsOutlinerHelpers.h
@@ -25,6 +25,11 @@ namespace UE::Editor::Outliner::Helpers
 	
 	// Get the name of the Outliner column corresponding to the given TEDS column (if any)
 	FName FindOutlinerColumnFromTedsColumns(const DataStorage::ICoreProvider* Storage, TConstArrayView<TWeakObjectPtr<const UScriptStruct>> TEDSColumns);
+
+	// Remove a Teds Column from the TEDSToOutlinerDefaultColumnMapping
+	void RemoveTedsColumnMappingToOutlinerColumn(DataStorage::ICoreProvider* Storage, TWeakObjectPtr<const UScriptStruct> TEDSColumnToRemove);
+	// Add a Teds Column to the TEDSToOutlinerDefaultColumnMapping
+	void AddTedsColumnMappingToOutlinerColumn(DataStorage::ICoreProvider* Storage, TWeakObjectPtr<const UScriptStruct> TEDSColumnToAdd, const FName& OutlinerColumnId);
 	
 	// Helper function to check if the given QueryHandle is valid to be used in a filter (is not an observer).
     bool CheckValidFilterQueryHandle(const DataStorage::QueryHandle& InQueryHandle);
