[TEDS Debugger] Add Discovery Tab

#Synced-CL 49652531

diff --git a/Engine/Plugins/Experimental/EditorDataStorageFeatures/Source/TedsDebugger/Private/STedsDebugger.cpp b/Engine/Plugins/Experimental/EditorDataStorageFeatures/Source/TedsDebugger/Private/STedsDebugger.cpp
--- a/Engine/Plugins/Experimental/EditorDataStorageFeatures/Source/TedsDebugger/Private/STedsDebugger.cpp
+++ b/Engine/Plugins/Experimental/EditorDataStorageFeatures/Source/TedsDebugger/Private/STedsDebugger.cpp
@@ -14,6 +14,7 @@
 #include "QueryEditor/TedsQueryEditor.h"
 #include "TedsOutlinerModule.h"
 #include "TedsOutlinerMode.h"
+#include "Widgets/STedsDiscoveryDebugger.h"
 
 #define LOCTEXT_NAMESPACE "STedsDebugger"
 
@@ -21,7 +22,8 @@ namespace UE::Editor::DataStorage::Debug
 {
 	namespace Private
 	{
-		FName ToolbarTabName = TEXT("TEDS Debugger Toolbar");
+		const FName ToolbarTabName = TEXT("TEDS Debugger Toolbar");
+		const FName DiscoveryTabName = TEXT("Discover");
 	}
 
 enum class EQueryEditorTabs
@@ -52,7 +54,7 @@ void STedsDebugger::Construct(const FArguments& InArgs, const TSharedRef<SDockTa
 	RegisterTabSpawners();
 
 	// Setup the default layout
-	TSharedRef<FTabManager::FLayout> Layout = FTabManager::NewLayout("TedsDebuggerLayout_v1.0")
+	TSharedRef<FTabManager::FLayout> Layout = FTabManager::NewLayout("TedsDebuggerLayout_v1.1")
 	->AddArea
 	(
 		FTabManager::NewPrimaryArea()
@@ -74,6 +76,7 @@ void STedsDebugger::Construct(const FArguments& InArgs, const TSharedRef<SDockTa
 					->AddTab(GetQueryTabName(EQueryEditorTabs::QueryEditorTwo), ETabState::ClosedTab)
 					->AddTab(GetQueryTabName(EQueryEditorTabs::QueryEditorThree), ETabState::ClosedTab)
 					->AddTab(GetQueryTabName(EQueryEditorTabs::QueryEditorFour), ETabState::ClosedTab)
+					->AddTab(Private::DiscoveryTabName, ETabState::OpenedTab)
 					->SetHideTabWell(false)
 					->SetForegroundTab(GetQueryTabName(EQueryEditorTabs::QueryEditorOne))
 			)
@@ -181,6 +184,16 @@ TSharedRef<SDockTab> STedsDebugger::SpawnQueryEditorTab(const FSpawnTabArgs& Arg
 	return DockTab;
 }
 
+TSharedRef<SDockTab> STedsDebugger::SpawnDiscoverTab(const FSpawnTabArgs& Args)
+{
+	TSharedRef<SDockTab> DockTab = SNew(SDockTab)
+		.TabRole(ETabRole::PanelTab)
+		.Label(LOCTEXT("TedsDebugger_DiscoverDisplayName", "Discover"));
+
+	DockTab->SetContent(SNew(SDiscoveryDebuggerTab));
+	return DockTab;
+}
+
 void STedsDebugger::RegisterTabSpawners()
 {
 	const TSharedRef<FWorkspaceItem> AppMenuGroup =
@@ -225,6 +238,14 @@ void STedsDebugger::RegisterTabSpawners()
 		.SetDisplayName(LOCTEXT("TedsDebugger_QueryEditorDisplayNameFour", "Query Editor 4"))
 		.SetTooltipText(QueryEditorTooltip)
 		.SetIcon(FSlateIcon(FAppStyle::GetAppStyleSetName(), "Debug"));
+
+	TabManager->RegisterTabSpawner(
+		Private::DiscoveryTabName,
+		FOnSpawnTab::CreateRaw(this, &STedsDebugger::SpawnDiscoverTab))
+		.SetGroup(AppMenuGroup)
+		.SetDisplayName(LOCTEXT("TedsDebugger_DiscoverDisplayName", "Discover"))
+		.SetTooltipText(LOCTEXT("TedsDebugger_DiscoverToolTip", "Opens the Discover tab. Contains special searching functionality to explore TEDS data without a query."))
+		.SetIcon(FSlateIcon(FAppStyle::GetAppStyleSetName(), "Icons.Search"));
 }
 
 } // namespace UE::Editor::DataStorage::Debug
diff --git a/Engine/Plugins/Experimental/EditorDataStorageFeatures/Source/TedsDebugger/Private/STedsDebugger.h b/Engine/Plugins/Experimental/EditorDataStorageFeatures/Source/TedsDebugger/Private/STedsDebugger.h
--- a/Engine/Plugins/Experimental/EditorDataStorageFeatures/Source/TedsDebugger/Private/STedsDebugger.h
+++ b/Engine/Plugins/Experimental/EditorDataStorageFeatures/Source/TedsDebugger/Private/STedsDebugger.h
@@ -40,6 +40,7 @@ namespace UE::Editor::DataStorage
 			static FName GetQueryTabName(const EQueryEditorTabs QueryTabIdentifier);
 			TSharedRef<SDockTab> SpawnToolbar(const FSpawnTabArgs& Args);
 			TSharedRef<SDockTab> SpawnQueryEditorTab(const FSpawnTabArgs& Args, const EQueryEditorTabs QueryTabIdentifier);
+			TSharedRef<SDockTab> SpawnDiscoverTab(const FSpawnTabArgs& Args);
 			void FillWindowMenu( FMenuBuilder& MenuBuilder);
 
 			void RegisterTabSpawners();
diff --git a/Engine/Plugins/Experimental/EditorDataStorageFeatures/Source/TedsDebugger/Private/Widgets/STedsDiscoveryDebugger.cpp b/Engine/Plugins/Experimental/EditorDataStorageFeatures/Source/TedsDebugger/Private/Widgets/STedsDiscoveryDebugger.cpp
new file mode 100644
--- /dev/null
+++ b/Engine/Plugins/Experimental/EditorDataStorageFeatures/Source/TedsDebugger/Private/Widgets/STedsDiscoveryDebugger.cpp
@@ -0,0 +1,279 @@
+// Copyright Epic Games, Inc. All Rights Reserved.
+
+#include "STedsDiscoveryDebugger.h"
+
+#include "Framework/MultiBox/MultiBoxBuilder.h"
+#include "TedsTableViewerColumn.h"
+#include "TedsRowArrayNode.h"
+#include "Widgets/Input/SSearchBox.h"
+#include "Widgets/Layout/SBorder.h"
+#include "Widgets/STedsTableViewer.h"
+#include "Widgets/SRowDetails.h"
+
+#define LOCTEXT_NAMESPACE "TedsDiscoveryDebuggerTab"
+
+namespace UE::Editor::DataStorage::Debug
+{
+	void SDiscoveryDebuggerTab::Construct(const FArguments& InArgs)
+	{
+		Storage = GetMutableDataStorageFeature<ICoreProvider>(StorageFeatureName);
+		UiProvider = GetMutableDataStorageFeature<IUiProvider>(UiFeatureName);
+        checkf(Storage && UiProvider, TEXT("Cannot create the Discovery Tab without TEDS Core or TEDS UI"));
+
+		SetRowQueryStack();
+		
+		ChildSlot
+		[
+			SNew(SBorder)
+			.BorderImage(FAppStyle::GetBrush("Brushes.Background"))
+			[
+				SNew(SVerticalBox)
+				+SVerticalBox::Slot()
+				.AutoHeight()
+				[
+					CreateToolbar()
+				]
+				+SVerticalBox::Slot()
+				.FillHeight(1.f)
+				.Padding(0.f, 4.0f)
+				[
+					SNew(SSplitter)
+					+SSplitter::Slot()
+					.Value(0.5f)
+					[
+						SNew(SVerticalBox)
+						+SVerticalBox::Slot()
+						.AutoHeight()
+						[
+							CreateSearchBox()
+						]
+						+SVerticalBox::Slot()
+						.FillContentHeight(1.0f)
+						[
+							CreateTableViewer()
+						]
+					]
+					+SSplitter::Slot()
+					.Value(0.5f)
+					[
+						SAssignNew(RowDetailsWidget, SRowDetails)
+					]
+				]
+			]
+		];
+	}
+
+	TSharedRef<SWidget> SDiscoveryDebuggerTab::CreateToolbar()
+	{
+		FSlimHorizontalToolBarBuilder ToolBarBuilder(nullptr, FMultiBoxCustomization::None);
+
+		ToolBarBuilder.AddComboButton( 
+			FUIAction(),
+			FOnGetContent::CreateSP(this, &SDiscoveryDebuggerTab::GetSearchModesMenuWidget),
+			TAttribute<FText>::CreateSP(this, &SDiscoveryDebuggerTab::GetCurrentSearchModeAsText),
+			FText(),
+			FSlateIcon(FAppStyle::GetAppStyleSetName(), "Icons.Search"),
+			false);
+
+		return ToolBarBuilder.MakeWidget();
+	}
+
+	TSharedRef<SWidget> SDiscoveryDebuggerTab::GetSearchModesMenuWidget()
+	{
+		FMenuBuilder MenuBuilder(true, nullptr);
+
+		MenuBuilder.AddMenuEntry(
+			GetSearchModeAsText(ESearchMode::RowHandle),
+			LOCTEXT("RowHandleSearchTooltip", "Search for a specific Row Handle and observe its properties"),
+			FSlateIcon(), 
+			FUIAction(FExecuteAction::CreateLambda([this]()
+			{
+				SearchMode = ESearchMode::RowHandle;
+				if (ColumnSearchNode)
+				{
+					ColumnSearchNode->ClearSearch();
+				}
+				SetRowQueryStack();
+			}))
+		);
+
+		MenuBuilder.AddMenuEntry(
+			GetSearchModeAsText(ESearchMode::Column),
+			LOCTEXT("ColumnSearchTooltip", "Search on all rows and their columns tagged as searchable"),
+			FSlateIcon(), 
+			FUIAction(FExecuteAction::CreateLambda([this]()
+			{
+				SearchMode = ESearchMode::Column;
+				SetRowQueryStack();
+			}))
+		);
+
+		return MenuBuilder.MakeWidget();
+	}
+	
+	FText SDiscoveryDebuggerTab::GetCurrentSearchModeAsText() const
+	{
+		return GetSearchModeAsText(SearchMode);
+	}
+
+	FText SDiscoveryDebuggerTab::GetSearchModeAsText(const ESearchMode InSearchMode)
+	{
+		if (InSearchMode == ESearchMode::RowHandle)
+		{
+			return LOCTEXT("RowHandleComboBoxText", "Row Handle Search");
+		}
+		else if (InSearchMode == ESearchMode::Column)
+		{
+			return LOCTEXT("ColumnSearchComboBoxText", "Column Search");
+		}
+		return FText();
+	}
+	
+	void SDiscoveryDebuggerTab::InitColumnSearch()
+	{
+		// No need to initialize this twice, we want to search on all available columns always
+		if (!ColumnSearchNode)
+		{
+			// Uses the same node tech as the Universal Search but avoids making groups and searches without conditions
+			// Use ESyncActions::Refresh since this is a debugger, and we prefer accuracy to performance
+			ColumnSearchNode = MakeShared<QueryStack::FColumnsSearchNode>(*Storage, QueryStack::FColumnsSearchNode::ESyncActions::Refresh);
+
+			// Retrieve searchable columns.
+			QueryStack::Searching::ListSearchableColumns([this](const UScriptStruct* Column)
+			{
+				ColumnSearchNode->RegisterColumn(Column);
+			});
+		}
+	}
+	
+	void SDiscoveryDebuggerTab::SetRowQueryStack()
+	{
+		RowHandle SearchedRowHandle = InvalidRowHandle;
+		const FString SearchString = SearchText.ToString();
+		if (SearchMode == ESearchMode::RowHandle)
+		{
+			FRowHandleArray Rows;
+			LexFromString(SearchedRowHandle, *SearchString);
+			if(Storage->IsRowAvailable(SearchedRowHandle))
+			{
+				Rows.Add(SearchedRowHandle);
+			}
+
+			RowQueryStack = MakeShared<QueryStack::FRowArrayNode>(Rows);
+		}
+		else if (SearchMode == ESearchMode::Column)
+		{
+			// Init the ColumnSearchNode if not done already (it checks if already done within the function)
+			InitColumnSearch();
+			if (ColumnSearchNode.IsValid())
+			{
+				RowQueryStack = ColumnSearchNode;
+				if (SearchString.IsEmpty())
+				{
+					ColumnSearchNode->ClearSearch();
+				}
+				else
+				{
+					ColumnSearchNode->StartSearch(*SearchString);
+				}
+			}
+		}
+
+		if (TableViewer)
+		{
+			TableViewer->SetQueryStack(RowQueryStack);
+			TableViewer->SetSelection(SearchedRowHandle, true, ESelectInfo::Direct);
+			if(SearchedRowHandle == InvalidRowHandle && RowDetailsWidget)
+			{
+				RowDetailsWidget->ClearRow();
+			}
+		}
+	}
+
+	TSharedRef<SWidget> SDiscoveryDebuggerTab::CreateTableViewer()
+	{
+		CreateRowHandleColumn();
+		
+		TableViewer = SNew(STedsTableViewer)
+			.QueryStack(RowQueryStack)
+			.EmptyRowsMessage(LOCTEXT("EmptyRowsMessage", "No results for the current search."))
+			.OnSelectionChanged(STedsTableViewer::FOnSelectionChanged::CreateLambda(
+				[this](RowHandle SelectedRow)
+					{
+						if(RowDetailsWidget)
+						{
+							if(SelectedRow != InvalidRowHandle)
+							{
+								RowDetailsWidget->SetRow(SelectedRow);
+							}
+							else
+							{
+								RowDetailsWidget->ClearRow();
+							}
+						}
+					}));
+
+		if(RowHandleColumn)
+		{
+			TableViewer->AddCustomRowWidget(RowHandleColumn.ToSharedRef());
+		}
+		
+		return TableViewer->AsWidget();
+	}
+
+	void SDiscoveryDebuggerTab::CreateRowHandleColumn()
+	{
+		auto AssignWidgetToColumn = [this](TUniquePtr<FTypedElementWidgetConstructor> Constructor, TConstArrayView<TWeakObjectPtr<const UScriptStruct>>)
+			{
+				TSharedPtr<FTypedElementWidgetConstructor> WidgetConstructor(Constructor.Release());
+				RowHandleColumn = MakeShared<FTedsTableViewerColumn>(TEXT("Row Handle"), WidgetConstructor);
+				return false;
+			};
+
+		UiProvider->CreateWidgetConstructors(
+			UiProvider->FindPurpose(IUiProvider::FPurposeInfo("General", "Cell", "RowHandle").GeneratePurposeID()),
+			FMetaDataView(), AssignWidgetToColumn);
+	}
+
+	TSharedRef<SWidget> SDiscoveryDebuggerTab::CreateSearchBox()
+	{
+		SearchBox = SNew(SSearchBox)
+			.SelectAllTextWhenFocused(true)
+			.DelayChangeNotificationsWhileTyping(false)
+			.OnTextChanged_Lambda([this](const FText& InText)
+			{
+				SearchText = InText;
+
+				if (SearchText.IsEmpty())
+				{
+					if (TableViewer)
+					{
+						TableViewer->ClearSelection();
+					}
+					if (RowDetailsWidget)
+					{
+						RowDetailsWidget->ClearRow();
+					}
+					if (ColumnSearchNode)
+					{
+						ColumnSearchNode->ClearSearch();
+					}
+				}
+
+				if (SearchMode == ESearchMode::RowHandle)
+				{
+					SetRowQueryStack();
+				}
+				else if (SearchMode == ESearchMode::Column && ColumnSearchNode && !SearchText.IsEmpty())
+				{
+					ColumnSearchNode->StartSearch(*SearchText.ToString());
+				}
+			})
+			.HintText(LOCTEXT("SearchBox_Hint", "Search for TEDS data..."));
+			
+		return SearchBox.ToSharedRef();
+	}
+	
+} // namespace UE::Editor::DataStorage::Debug
+
+#undef LOCTEXT_NAMESPACE
diff --git a/Engine/Plugins/Experimental/EditorDataStorageFeatures/Source/TedsDebugger/Private/Widgets/STedsDiscoveryDebugger.h b/Engine/Plugins/Experimental/EditorDataStorageFeatures/Source/TedsDebugger/Private/Widgets/STedsDiscoveryDebugger.h
new file mode 100644
--- /dev/null
+++ b/Engine/Plugins/Experimental/EditorDataStorageFeatures/Source/TedsDebugger/Private/Widgets/STedsDiscoveryDebugger.h
@@ -0,0 +1,63 @@
+// Copyright Epic Games, Inc. All Rights Reserved.
+#pragma once
+
+#include "DataStorage/Features.h"
+#include "QueryEditor/TedsQueryEditorResultsView.h"
+#include "TedsColumnsSearchNode.h"
+#include "Widgets/SCompoundWidget.h"
+
+class SRowDetails;
+class SSearchBox;
+class STedsTableViewer;
+class FTedsTableViewerColumn;
+
+namespace UE::Editor::DataStorage::Debug
+{
+	enum class ESearchMode : uint8
+	{
+		// Search for a specific RowHandle for specific debugging
+		RowHandle,
+		// Search on all rows on columns with Searchable properties
+		Column
+	};
+	
+	class SDiscoveryDebuggerTab final : public SCompoundWidget
+	{
+	public:
+
+		SLATE_BEGIN_ARGS(SDiscoveryDebuggerTab) 
+		{}
+		SLATE_END_ARGS()
+
+		void Construct(const FArguments& InArgs);
+
+	private:
+
+		TSharedRef<SWidget> CreateToolbar();
+		TSharedRef<SWidget> GetSearchModesMenuWidget();
+		FText GetCurrentSearchModeAsText() const;
+		static FText GetSearchModeAsText(const ESearchMode InSearchMode);
+
+		void SetRowQueryStack();
+		void InitColumnSearch();
+		TSharedRef<SWidget> CreateTableViewer();
+		void CreateRowHandleColumn();
+		TSharedRef<SWidget> CreateSearchBox();
+
+		TSharedPtr<STedsTableViewer> TableViewer;
+		TSharedPtr<SSearchBox> SearchBox;
+		// Widget that displays details of a row
+		TSharedPtr<SRowDetails> RowDetailsWidget;
+		// Custom column for the table viewer to display row handles
+		TSharedPtr<FTedsTableViewerColumn> RowHandleColumn;
+		
+		TSharedPtr<QueryStack::IRowNode> RowQueryStack;
+		// Hold onto the ColumnSearchNode so it doesn't need to be re-initialized when switching back to column search
+		TSharedPtr<QueryStack::FColumnsSearchNode> ColumnSearchNode;
+
+		ICoreProvider* Storage = nullptr;
+		IUiProvider* UiProvider = nullptr;
+		FText SearchText;
+		ESearchMode SearchMode = ESearchMode::RowHandle;
+	};
+}
diff --git a/Engine/Plugins/Experimental/EditorDataStorageFeatures/Source/TedsQueryStack/Private/Searching/SearchUtils.cpp b/Engine/Plugins/Experimental/EditorDataStorageFeatures/Source/TedsQueryStack/Private/Searching/SearchUtils.cpp
--- a/Engine/Plugins/Experimental/EditorDataStorageFeatures/Source/TedsQueryStack/Private/Searching/SearchUtils.cpp
+++ b/Engine/Plugins/Experimental/EditorDataStorageFeatures/Source/TedsQueryStack/Private/Searching/SearchUtils.cpp
@@ -80,6 +80,30 @@ namespace UE::Editor::DataStorage::QueryStack::Searching
 	{
 		return Search(String.ToString(), Context, TempString);
 	}
+	
+	void ListSearchableColumns(TFunctionRef<void(const UScriptStruct*)> Callback)
+	{
+		static const FName SearchableMetaDataTag = FName("Searchable");
+		
+		UScriptStruct* DataColumnType = FEditorDataStorageColumn::StaticStruct();
+		ForEachObjectOfClass(UScriptStruct::StaticClass(),
+			[Callback, DataColumnType](UObject* Object)
+			{
+				UScriptStruct* Struct = static_cast<UScriptStruct*>(Object);
+				if (DataColumnType != Struct && Struct->IsChildOf(DataColumnType))
+				{
+					// Don't use ListSearchableProperties since we can't early out and only need to find one
+					for (TFieldIterator<FProperty> PropIt(Struct); PropIt; ++PropIt)
+					{
+						if ((*PropIt)->HasMetaData(SearchableMetaDataTag))
+						{
+							Callback(Struct);
+							break;
+						}
+					}
+				}
+			});
+	}
 
 	void ListSearchableProperties(const UScriptStruct* ColumnType, TFunctionRef<void(const FProperty*)> Callback)
 	{
diff --git a/Engine/Plugins/Experimental/EditorDataStorageFeatures/Source/TedsQueryStack/Public/Searching/SearchUtils.h b/Engine/Plugins/Experimental/EditorDataStorageFeatures/Source/TedsQueryStack/Public/Searching/SearchUtils.h
--- a/Engine/Plugins/Experimental/EditorDataStorageFeatures/Source/TedsQueryStack/Public/Searching/SearchUtils.h
+++ b/Engine/Plugins/Experimental/EditorDataStorageFeatures/Source/TedsQueryStack/Public/Searching/SearchUtils.h
@@ -66,6 +66,7 @@ namespace UE::Editor::DataStorage::QueryStack::Searching
 	TEDSQUERYSTACK_API bool Search(const FName& String, const FSearchContext& Context, FString& TempString);
 	TEDSQUERYSTACK_API bool Search(const FText& String, const FSearchContext& Context, FString& TempString);
 
+	TEDSQUERYSTACK_API void ListSearchableColumns(TFunctionRef<void(const UScriptStruct*)> Callback);
 	TEDSQUERYSTACK_API void ListSearchableProperties(const UScriptStruct* ColumnType, TFunctionRef<void(const FProperty*)> Callback);
 } // namespace UE::Editor::DataStorage::QueryStack::Searching
 
