[TedsTableViewer] DropHandler are now added externally. TableViewers now remove the DropInfoColumn instead of the DropHandler.

#Synced-CL 49611723

diff --git a/Engine/Plugins/Experimental/EditorDataStorageFeatures/Source/TedsTableViewer/Private/Widgets/Composite/STedsCompositeHierarchyViewer.cpp b/Engine/Plugins/Experimental/EditorDataStorageFeatures/Source/TedsTableViewer/Private/Widgets/Composite/STedsCompositeHierarchyViewer.cpp
--- a/Engine/Plugins/Experimental/EditorDataStorageFeatures/Source/TedsTableViewer/Private/Widgets/Composite/STedsCompositeHierarchyViewer.cpp
+++ b/Engine/Plugins/Experimental/EditorDataStorageFeatures/Source/TedsTableViewer/Private/Widgets/Composite/STedsCompositeHierarchyViewer.cpp
@@ -2,9 +2,9 @@
 
 #include "Widgets/Composite/STedsCompositeHierarchyViewer.h"
 
+#include "DragAndDrop/Widgets/WidgetDropHandler.h"
 #include "TedsQueryStackInterfaces.h"
 #include "Widgets/STedsFilterBar.h"
-#include "TedsTableViewerModel.h"
 #include "Widgets/STedsSearchBox.h"
 
 namespace UE::Editor::DataStorage
@@ -252,4 +252,8 @@ namespace UE::Editor::DataStorage
 		HierarchyViewer->ExpandWithParents(Row);
 	}
 
+	void STedsCompositeHierarchyViewer::SetDropHandler(TUniquePtr<FWidgetDropHandler> InDropHandler)
+	{
+		HierarchyViewer->SetDropHandler(MoveTemp(InDropHandler));
+	}
 } // namespace UE::Editor::DataStorage
diff --git a/Engine/Plugins/Experimental/EditorDataStorageFeatures/Source/TedsTableViewer/Private/Widgets/STedsHierarchyViewer.cpp b/Engine/Plugins/Experimental/EditorDataStorageFeatures/Source/TedsTableViewer/Private/Widgets/STedsHierarchyViewer.cpp
--- a/Engine/Plugins/Experimental/EditorDataStorageFeatures/Source/TedsTableViewer/Private/Widgets/STedsHierarchyViewer.cpp
+++ b/Engine/Plugins/Experimental/EditorDataStorageFeatures/Source/TedsTableViewer/Private/Widgets/STedsHierarchyViewer.cpp
@@ -2,26 +2,29 @@
 
 #include "Widgets/STedsHierarchyViewer.h"
 
-#include "TableViewerWidgetDropHandler.h"
-#include "TedsTableViewerColumn.h"
-#include "TedsTableViewerUtils.h"
-#include "TedsTableViewerWidgetColumns.h"
 #include "Columns/SlateDelegateColumns.h"
 #include "DataStorage/Features.h"
+#include "DragAndDrop/Widgets/WidgetDropHandler.h"
 #include "Elements/Columns/TypedElementHiearchyColumns.h"
-#include "Elements/Columns/TypedElementSlateWidgetColumns.h"
 #include "Elements/Columns/TypedElementUIColumns.h"
 #include "Elements/Framework/TypedElementDataStorageWidget.h"
 #include "QueryStack/TopLevelRowsNode.h"
+#include "TedsTableViewerColumn.h"
+#include "TedsTableViewerUtils.h"
+#include "TedsTableViewerWidgetColumns.h"
 #include "Widgets/STedsTableViewerRow.h"
-#include "Widgets/Views/SListView.h"
-#include "Widgets/Text/STextBlock.h"
 #include "Widgets/STedsTreeView.h"
+#include "Widgets/Text/STextBlock.h"
+#include "Widgets/Views/SListView.h"
 
 #define LOCTEXT_NAMESPACE "SHierarchyViewer"
 
 namespace UE::Editor::DataStorage
 {
+	SHierarchyViewer::SHierarchyViewer()
+	{
+	}
+
 	SHierarchyViewer::~SHierarchyViewer()
 	{
 		RemoveAllTableViewerUIRows();
@@ -106,8 +109,6 @@ namespace UE::Editor::DataStorage
 
 		// Whenever the model changes, refresh the list to update the UI
 		Model->GetOnModelChanged().AddRaw(this, &SHierarchyViewer::OnModelChanged);
-
-		DropHandler = MakePimpl<Widgets::FTableViewerWidgetDropHandler>(SharedThis(this));
 	}
 
 	void SHierarchyViewer::OnModelChanged()
@@ -186,11 +187,21 @@ namespace UE::Editor::DataStorage
 		{
 			DropHandler->OnDragLeave(DragDropEvent);
 		}
+		
+		if (ICoreProvider* DataStorage = Model->GetDataStorageInterface())
+		{
+			DataStorage->RemoveColumn<FTableViewerDropInfoColumn>(GetWidgetRowHandle());
+		}
 	}
 
 	FReply SHierarchyViewer::OnDrop(const FGeometry& MyGeometry, const FDragDropEvent& DragDropEvent)
 	{
-		return DropHandler ? DropHandler->OnDrop(MyGeometry, DragDropEvent) : FReply::Unhandled();
+		FReply Reply = DropHandler ? DropHandler->OnDrop(MyGeometry, DragDropEvent) : FReply::Unhandled();
+		if (ICoreProvider* DataStorage = Model->GetDataStorageInterface())
+		{
+			DataStorage->RemoveColumn<FTableViewerDropInfoColumn>(GetWidgetRowHandle());
+		}
+		return Reply;
 	}
 
 	void SHierarchyViewer::AddWidgetColumns()
@@ -433,6 +444,11 @@ namespace UE::Editor::DataStorage
 		}
 	}
 
+	void SHierarchyViewer::SetDropHandler(TUniquePtr<FWidgetDropHandler> InDropHandler)
+	{
+		DropHandler = MoveTemp(InDropHandler);
+	}
+
 	bool SHierarchyViewer::IsItemVisible(TableViewerItemPtr InItem) const
 	{
 		return TreeView->IsItemVisible(InItem);
diff --git a/Engine/Plugins/Experimental/EditorDataStorageFeatures/Source/TedsTableViewer/Private/Widgets/STedsTableViewer.cpp b/Engine/Plugins/Experimental/EditorDataStorageFeatures/Source/TedsTableViewer/Private/Widgets/STedsTableViewer.cpp
--- a/Engine/Plugins/Experimental/EditorDataStorageFeatures/Source/TedsTableViewer/Private/Widgets/STedsTableViewer.cpp
+++ b/Engine/Plugins/Experimental/EditorDataStorageFeatures/Source/TedsTableViewer/Private/Widgets/STedsTableViewer.cpp
@@ -2,24 +2,27 @@
 
 #include "Widgets/STedsTableViewer.h"
 
-#include "TableViewerWidgetDropHandler.h"
-#include "TedsTableViewerColumn.h"
-#include "TedsTableViewerUtils.h"
-#include "TedsTableViewerWidgetColumns.h"
 #include "Columns/SlateDelegateColumns.h"
 #include "DataStorage/Features.h"
-#include "Elements/Columns/TypedElementSlateWidgetColumns.h"
+#include "DragAndDrop/Widgets/WidgetDropHandler.h"
 #include "Elements/Columns/TypedElementUIColumns.h"
 #include "Elements/Framework/TypedElementAttributeBinding.h"
 #include "Elements/Framework/TypedElementDataStorageWidget.h"
+#include "TedsTableViewerColumn.h"
+#include "TedsTableViewerUtils.h"
+#include "TedsTableViewerWidgetColumns.h"
 #include "Widgets/STedsTableViewerRow.h"
-#include "Widgets/Views/SListView.h"
 #include "Widgets/Text/STextBlock.h"
+#include "Widgets/Views/SListView.h"
 
 #define LOCTEXT_NAMESPACE "STedsTableViewer"
 
 namespace UE::Editor::DataStorage
 {
+	STedsTableViewer::STedsTableViewer()
+	{
+	}
+
 	STedsTableViewer::~STedsTableViewer()
 	{
 		RemoveAllTableViewerUIRows();
@@ -105,8 +108,6 @@ namespace UE::Editor::DataStorage
 			ListView->RequestListRefresh();
 			CreateInternalWidget();
 		});
-
-		DropHandler = MakePimpl<Widgets::FTableViewerWidgetDropHandler>(SharedThis(this));
 	}
 
 	void STedsTableViewer::OnDragEnter(const FGeometry& MyGeometry, const FDragDropEvent& DragDropEvent)
@@ -128,11 +129,26 @@ namespace UE::Editor::DataStorage
 		{
 			DropHandler->OnDragLeave(DragDropEvent);
 		}
+		
+		if (ICoreProvider* DataStorage = Model->GetDataStorageInterface())
+		{
+			DataStorage->RemoveColumn<FTableViewerDropInfoColumn>(GetWidgetRowHandle());
+		}
 	}
 
 	FReply STedsTableViewer::OnDrop(const FGeometry& MyGeometry, const FDragDropEvent& DragDropEvent)
 	{
-		return DropHandler ? DropHandler->OnDrop(MyGeometry, DragDropEvent) : FReply::Unhandled();
+		FReply Reply = DropHandler ? DropHandler->OnDrop(MyGeometry, DragDropEvent) : FReply::Unhandled();
+		if (ICoreProvider* DataStorage = Model->GetDataStorageInterface())
+		{
+			DataStorage->RemoveColumn<FTableViewerDropInfoColumn>(GetWidgetRowHandle());
+		}
+		return Reply;
+	}
+
+	void STedsTableViewer::SetDropHandler(TUniquePtr<FWidgetDropHandler> InDropHandler)
+	{
+		DropHandler = MoveTemp(InDropHandler);
 	}
 
 	void STedsTableViewer::AddWidgetColumns()
diff --git a/Engine/Plugins/Experimental/EditorDataStorageFeatures/Source/TedsTableViewer/Private/Widgets/TableViewerWidgetDropHandler.cpp b/Engine/Plugins/Experimental/EditorDataStorageFeatures/Source/TedsTableViewer/Private/Widgets/TableViewerWidgetDropHandler.cpp
--- a/Engine/Plugins/Experimental/EditorDataStorageFeatures/Source/TedsTableViewer/Private/Widgets/TableViewerWidgetDropHandler.cpp
+++ b/Engine/Plugins/Experimental/EditorDataStorageFeatures/Source/TedsTableViewer/Private/Widgets/TableViewerWidgetDropHandler.cpp
@@ -1,6 +1,6 @@
 // Copyright Epic Games, Inc. All Rights Reserved.
 
-#include "TableViewerWidgetDropHandler.h"
+#include "Widgets/TableViewerWidgetDropHandler.h"
 
 #include "DragAndDrop/DropOperationInput.h"
 #include "ITedsTableViewer.h"
@@ -8,12 +8,16 @@
 
 namespace UE::Editor::DataStorage::Widgets
 {
-	FTableViewerWidgetDropHandler::FTableViewerWidgetDropHandler(TWeakPtr<ITableViewer> InTableViewer)
-		: FWidgetDropHandler(nullptr) // The table viewer does not operate on selection sets
+	FTableViewerWidgetDropHandler::FTableViewerWidgetDropHandler(TWeakPtr<ITableViewer> InTableViewer, TWeakObjectPtr<UTypedElementSelectionSet> InSelectionSet)
+		: FWidgetDropHandler(MoveTemp(InSelectionSet))
 		, TableViewer(MoveTemp(InTableViewer))
 	{
 	}
 
+	FTableViewerWidgetDropHandler::~FTableViewerWidgetDropHandler()
+	{
+	}
+
 	FWidgetDropHandler::EUpdateParameters FTableViewerWidgetDropHandler::UpdateParameters(ICoreProvider& Storage, const FGeometry& Geometry,
 		const FDragDropEvent& InputEvent)
 	{
@@ -65,11 +69,6 @@ namespace UE::Editor::DataStorage::Widgets
 	{
 		LastTargetRow = InvalidRowHandle;
 		
-		if (TSharedPtr<ITableViewer> TableViewerPtr = TableViewer.Pin())
-		{
-			Storage.RemoveColumn<FTableViewerDropInfoColumn>(TableViewerPtr->GetWidgetRowHandle());
-		}
-
 		FWidgetDropHandler::Stop(Storage, DropSystem);
 	}
 }
diff --git a/Engine/Plugins/Experimental/EditorDataStorageFeatures/Source/TedsTableViewer/Private/Widgets/TableViewerWidgetDropHandler.h b/Engine/Plugins/Experimental/EditorDataStorageFeatures/Source/TedsTableViewer/Private/Widgets/TableViewerWidgetDropHandler.h
deleted file mode 100644
diff --git a/Engine/Plugins/Experimental/EditorDataStorageFeatures/Source/TedsTableViewer/Public/Widgets/Composite/STedsCompositeHierarchyViewer.h b/Engine/Plugins/Experimental/EditorDataStorageFeatures/Source/TedsTableViewer/Public/Widgets/Composite/STedsCompositeHierarchyViewer.h
--- a/Engine/Plugins/Experimental/EditorDataStorageFeatures/Source/TedsTableViewer/Public/Widgets/Composite/STedsCompositeHierarchyViewer.h
+++ b/Engine/Plugins/Experimental/EditorDataStorageFeatures/Source/TedsTableViewer/Public/Widgets/Composite/STedsCompositeHierarchyViewer.h
@@ -122,7 +122,10 @@ namespace UE::Editor::DataStorage
 
 		// Expand this row, and all its parents
 		TEDSTABLEVIEWER_API void ExpandWithParents(RowHandle Row);
-
+		
+		// Set a handler that this viewer will forward its drop operations to.
+		TEDSTABLEVIEWER_API void SetDropHandler(TUniquePtr<FWidgetDropHandler> InDropHandler);
+		
 	private:
 		// Row Node passed to the SearchBox that receives the searched result
 		TSharedPtr<QueryStack::IRowNode> SearchNode;
diff --git a/Engine/Plugins/Experimental/EditorDataStorageFeatures/Source/TedsTableViewer/Public/Widgets/STedsHierarchyViewer.h b/Engine/Plugins/Experimental/EditorDataStorageFeatures/Source/TedsTableViewer/Public/Widgets/STedsHierarchyViewer.h
--- a/Engine/Plugins/Experimental/EditorDataStorageFeatures/Source/TedsTableViewer/Public/Widgets/STedsHierarchyViewer.h
+++ b/Engine/Plugins/Experimental/EditorDataStorageFeatures/Source/TedsTableViewer/Public/Widgets/STedsHierarchyViewer.h
@@ -17,14 +17,19 @@ class ITableRow;
 class STableViewBase;
 class SHeaderRow;
 
-namespace UE::Editor::DataStorage
+namespace UE::Editor
 {
-	class ITedsWidget;
-	class STedsTreeView;
-
-	namespace Widgets
+	class FWidgetDropHandler;
+	
+	namespace DataStorage
 	{
-		class FTableViewerWidgetDropHandler;
+		class ITedsWidget;
+		class STedsTreeView;
+
+		namespace Widgets
+		{
+			class FTableViewerWidgetDropHandler;
+		}
 	}
 }
 
@@ -109,6 +114,7 @@ namespace UE::Editor::DataStorage
 		SLATE_END_ARGS()
 
 	public:
+		TEDSTABLEVIEWER_API SHierarchyViewer();
 		TEDSTABLEVIEWER_API virtual ~SHierarchyViewer() override;
 		
 		TEDSTABLEVIEWER_API void Construct(
@@ -159,6 +165,9 @@ namespace UE::Editor::DataStorage
 		// Expand this row and all its parents
 		TEDSTABLEVIEWER_API void ExpandWithParents(RowHandle Row);
 
+		// Set a handler that this viewer will forward its drop operations to.
+		TEDSTABLEVIEWER_API void SetDropHandler(TUniquePtr<FWidgetDropHandler> InDropHandler);
+		
 	protected:
 		
 		FString GetTableViewerRowMapping(const TableViewerItemPtr InItem) const;
@@ -239,6 +248,6 @@ namespace UE::Editor::DataStorage
 		mutable TOptional<bool> bExpansionStateChangeRequested;
 
 		// Drag and Drop Handler
-		TPimplPtr<Widgets::FTableViewerWidgetDropHandler> DropHandler;
+		TUniquePtr<FWidgetDropHandler> DropHandler;
 	};
 }
\ No newline at end of file
diff --git a/Engine/Plugins/Experimental/EditorDataStorageFeatures/Source/TedsTableViewer/Public/Widgets/STedsTableViewer.h b/Engine/Plugins/Experimental/EditorDataStorageFeatures/Source/TedsTableViewer/Public/Widgets/STedsTableViewer.h
--- a/Engine/Plugins/Experimental/EditorDataStorageFeatures/Source/TedsTableViewer/Public/Widgets/STedsTableViewer.h
+++ b/Engine/Plugins/Experimental/EditorDataStorageFeatures/Source/TedsTableViewer/Public/Widgets/STedsTableViewer.h
@@ -13,12 +13,17 @@ class ITableRow;
 class STableViewBase;
 class SHeaderRow;
 
-namespace UE::Editor::DataStorage
+namespace UE::Editor
 {
-	class ITedsWidget;
-	namespace Widgets
+	class FWidgetDropHandler;
+	
+	namespace DataStorage
 	{
-		class FTableViewerWidgetDropHandler;
+		class ITedsWidget;
+		namespace Widgets
+		{
+			class FTableViewerWidgetDropHandler;
+		}
 	}
 }
 
@@ -101,6 +106,7 @@ namespace UE::Editor::DataStorage
 		SLATE_END_ARGS()
 
 	public:
+		TEDSTABLEVIEWER_API STedsTableViewer();
 		TEDSTABLEVIEWER_API virtual ~STedsTableViewer() override;
 		
 		TEDSTABLEVIEWER_API void Construct(const FArguments& InArgs);
@@ -143,6 +149,9 @@ namespace UE::Editor::DataStorage
 		TEDSTABLEVIEWER_API virtual FReply OnDragOver(const FGeometry& MyGeometry, const FDragDropEvent& DragDropEvent) override;
 		TEDSTABLEVIEWER_API virtual void OnDragLeave(const FDragDropEvent& DragDropEvent) override;
 		TEDSTABLEVIEWER_API virtual FReply OnDrop( const FGeometry& MyGeometry, const FDragDropEvent& DragDropEvent ) override;
+		
+		// Set a handler that this viewer will forward its drop operations to.
+		TEDSTABLEVIEWER_API void SetDropHandler(TUniquePtr<FWidgetDropHandler> InDropHandler);
 
 	protected:
 	
@@ -199,6 +208,6 @@ namespace UE::Editor::DataStorage
 		TAttribute<FText> EmptyRowsMessage;
 
 		// Drag and Drop Handler
-		TPimplPtr<Widgets::FTableViewerWidgetDropHandler> DropHandler;
+		TUniquePtr<FWidgetDropHandler> DropHandler;
 	};
 }
\ No newline at end of file
diff --git a/Engine/Plugins/Experimental/EditorDataStorageFeatures/Source/TedsTableViewer/Public/Widgets/TableViewerWidgetDropHandler.h b/Engine/Plugins/Experimental/EditorDataStorageFeatures/Source/TedsTableViewer/Public/Widgets/TableViewerWidgetDropHandler.h
new file mode 100644
--- /dev/null
+++ b/Engine/Plugins/Experimental/EditorDataStorageFeatures/Source/TedsTableViewer/Public/Widgets/TableViewerWidgetDropHandler.h
@@ -0,0 +1,34 @@
+// Copyright Epic Games, Inc. All Rights Reserved.
+
+#pragma once
+
+#include "DragAndDrop/Widgets/WidgetDropHandler.h"
+
+namespace UE::Editor::DataStorage
+{
+	class ITableViewer;
+}
+
+#define UE_API TEDSTABLEVIEWER_API
+
+namespace UE::Editor::DataStorage::Widgets
+{
+	class FTableViewerWidgetDropHandler : public FWidgetDropHandler
+	{
+	public:
+		UE_API FTableViewerWidgetDropHandler(TWeakPtr<ITableViewer> InTableViewer, TWeakObjectPtr<UTypedElementSelectionSet> InSelectionSet);
+		UE_API virtual ~FTableViewerWidgetDropHandler() override;
+
+	protected:
+		virtual EUpdateParameters UpdateParameters(ICoreProvider& Storage, const FGeometry& Geometry, const FDragDropEvent& InputEvent) override;
+		virtual void SetRowsFromParameters(ICoreProvider& Storage, FRowHandleArrayView Rows) const override;
+		virtual void Stop(ICoreProvider& Storage, UDropOperationSystem& DropSystem) override;
+		
+		RowHandle GetDropTarget(ICoreProvider& Storage) const;
+	private:
+		TWeakPtr<ITableViewer> TableViewer;
+		RowHandle LastTargetRow = InvalidRowHandle;
+	};
+}
+
+#undef UE_API
