TedsOperations WidgetDropHandler now has virtual sub-methods for drops & previews. The LevelEditor's viewport drop handler now triggers FEditorDelegate actor signals through those new methods.

#Synced-CL 49726319

diff --git a/Engine/Plugins/Experimental/EditorDataStorageFeatures/Source/TedsOperations/Private/DragAndDrop/Widgets/PreviewWidgetDropHandler.cpp b/Engine/Plugins/Experimental/EditorDataStorageFeatures/Source/TedsOperations/Private/DragAndDrop/Widgets/PreviewWidgetDropHandler.cpp
--- a/Engine/Plugins/Experimental/EditorDataStorageFeatures/Source/TedsOperations/Private/DragAndDrop/Widgets/PreviewWidgetDropHandler.cpp
+++ b/Engine/Plugins/Experimental/EditorDataStorageFeatures/Source/TedsOperations/Private/DragAndDrop/Widgets/PreviewWidgetDropHandler.cpp
@@ -147,7 +147,7 @@ void FPreviewWidgetDropHandler::CreatePreviews(DataStorage::ICoreProvider& Stora
 	}
 
 	// Create the previews from our input.
-	(void)DropSystem.Apply(InputRows.GetRows());
+	ExecutePreviewDrop(Storage, DropSystem, InputRows.GetRows());
 
 	// Gather the results.
 	FRowHandleArray Created;
@@ -225,6 +225,29 @@ void FPreviewWidgetDropHandler::GetPreviewRows(DataStorage::FRowHandleArray& Out
 		}
 	}
 }
+
+int32 FPreviewWidgetDropHandler::ExecutePreviewDrop(DataStorage::ICoreProvider& Storage, UDropOperationSystem& DropSystem,
+	DataStorage::FRowHandleArrayView Input)
+{	
+	using namespace UE::Editor::DataStorage;
+	
+	int32 Result = 0;
+	for (RowHandle InputRow : Input)
+	{
+		// Executing in a batch would be better, but right now we cannot do this due to legacy reasons:
+		// The Level Viewport handler needs to execute the FEditorDelegates::OnNewActorsDropped/Placed signals which requires precise knowledge of
+		// which dropped UObject has created which AActor. As the required columns for this are deferred, we need to do this without rows.
+		Result += ExecutePreviewDropSingle(Storage, DropSystem, InputRow);
+	}
+	return Result;
+}
+
+bool FPreviewWidgetDropHandler::ExecutePreviewDropSingle(DataStorage::ICoreProvider& Storage, UDropOperationSystem& DropSystem,
+	DataStorage::RowHandle InputRow)
+{
+	using namespace UE::Editor::DataStorage;
+	return DropSystem.Apply(FRowHandleArrayView(&InputRow, 1, FRowHandleArray::EFlags::IsUnique | FRowHandleArray::EFlags::IsSorted)) != 0;
+}
 	
 }
 
diff --git a/Engine/Plugins/Experimental/EditorDataStorageFeatures/Source/TedsOperations/Private/DragAndDrop/Widgets/WidgetDropHandler.cpp b/Engine/Plugins/Experimental/EditorDataStorageFeatures/Source/TedsOperations/Private/DragAndDrop/Widgets/WidgetDropHandler.cpp
--- a/Engine/Plugins/Experimental/EditorDataStorageFeatures/Source/TedsOperations/Private/DragAndDrop/Widgets/WidgetDropHandler.cpp
+++ b/Engine/Plugins/Experimental/EditorDataStorageFeatures/Source/TedsOperations/Private/DragAndDrop/Widgets/WidgetDropHandler.cpp
@@ -277,53 +277,11 @@ void FWidgetDropHandler::Drop(DataStorage::ICoreProvider& Storage, UDropOperatio
 	// Update one last time.
 	(void)UpdateParameters(Storage, Geometry, InputEvent);
 	SetRowsFromParameters(Storage, InputRows.GetRows());
-	{
-		// Bit of a hack to select our newly placed objects.
-		// Usually we would like use teds rows for this, but since the synchronisation is deferred we cannot access TEv1 elements from our rows yet.
-		// We also cannot access the rows next tick, because we might be within a transaction scope which would not capture the selection.
-		// To solve this we collect all actors that are placed while the drop process is ongoing.
-		UTypedElementSelectionSet* SelectionSetPtr = SelectionSet.Get();	
-		TSet<TObjectPtr<AActor>> PlacedActors;	
-		FDelegateHandle DelegateHandle;
-		ON_SCOPE_EXIT{ GEditor->OnLevelActorAdded().Remove(DelegateHandle); };
-		if (SelectionSetPtr)
-		{
-			DelegateHandle = GEditor->OnLevelActorAdded().AddLambda([&PlacedActors](AActor* Actor)
-				{
-					if (Actor)
-					{
-						PlacedActors.Add(Actor);
-					}
-				});
-		}
-		
-		FScopedTransaction Transaction(LOCTEXT("Drop_Transaction", "Drop on widget."));
-		if (DropSystem.Apply(InputRows.GetRows()) != InputRows.Num())
-		{
-			ShowErrorNotifications(Storage, InputRows.GetRows());
-		}
 
-		// Now we select the roots of all captured actors.
-		if (SelectionSetPtr)
-		{
-			TArray<FTypedElementHandle> RootElements;
-			for (const TObjectPtr<AActor>& ActorPtr : PlacedActors)
-			{
-				if (AActor* Actor = ActorPtr.Get())
-				{
-					AActor* Parent = Actor->GetParentActor();
-					if (!Parent || !PlacedActors.Contains(Parent))
-					{
-						constexpr bool bAllowCreate = true;
-						RootElements.Add(UEngineElementsLibrary::AcquireEditorActorElementHandle(Actor, bAllowCreate));
-					}
-				}
-			}
-			
-			SelectionSetPtr->SetSelection(RootElements, FTypedElementSelectionOptions());
-		}
-	}
+	// Drop all data.
+	ExecuteDrop(Storage, DropSystem, InputRows.GetRows());
 	
+	// Stop the current operation.
 	Stop(Storage, DropSystem);
 }
 
@@ -353,7 +311,74 @@ void FWidgetDropHandler::UpdateOperationDecorator(FDragDropOperation& Operation,
 		Operation.SetCursorOverride({});
 	}
 }
+
+int32 FWidgetDropHandler::ExecuteDrop(DataStorage::ICoreProvider& Storage, UDropOperationSystem& DropSystem, DataStorage::FRowHandleArrayView Input)
+{
+	using namespace UE::Editor::DataStorage;
+	
+	// Bit of a hack to select our newly placed objects.
+	// Usually we would like use teds rows for this, but since the synchronization is deferred we cannot access TEv1 elements from our rows yet.
+	// We also cannot access the rows next tick, because we might be within a transaction scope which would not capture the selection.
+	// To solve this we collect all actors that are placed while the drop process is ongoing.
+	TSet<TObjectPtr<AActor>> PlacedActors;	
+	FDelegateHandle DelegateHandle;
+	if (SelectionSet.IsValid())
+	{
+		DelegateHandle = GEditor->OnLevelActorAdded().AddLambda([&PlacedActors](AActor* Actor)
+			{
+				PlacedActors.Add(Actor);
+			});
+	}
 	
+	int32 Result = 0;
+	FScopedTransaction Transaction(LOCTEXT("Drop_Transaction", "Drop on widget."));
+	for (RowHandle InputRow : Input)
+	{
+		// Executing in a batch would be better, but right now we cannot do this due to legacy reasons:
+		// The Level Viewport handler needs to execute the FEditorDelegates::OnNewActorsDropped/Placed signals which requires precise knowledge of
+		// which dropped UObject has created which AActor. As the required columns for this are deferred, we need to do this without rows.
+		Result += ExecuteDropSingle(Storage, DropSystem, InputRow);
+	}
+	
+	// Now we select the roots of all captured actors.
+	if (UTypedElementSelectionSet* SelectionSetPtr = SelectionSet.Get())
+	{
+		GEditor->OnLevelActorAdded().Remove(DelegateHandle);
+		
+		TArray<FTypedElementHandle> RootElements;
+		for (const TObjectPtr<AActor>& ActorPtr : PlacedActors)
+		{
+			if (AActor* Actor = ActorPtr.Get(); IsValid(Actor))
+			{
+				AActor* Parent = Actor->GetParentActor();
+				if (!IsValid(Parent) || !PlacedActors.Contains(Parent))
+				{
+					constexpr bool bAllowCreate = true;
+					if (FTypedElementHandle Element = UEngineElementsLibrary::AcquireEditorActorElementHandle(Actor, bAllowCreate))
+					{
+						RootElements.Add(Element);
+					}
+				}
+			}
+		}
+			
+		SelectionSetPtr->SetSelection(RootElements, FTypedElementSelectionOptions());
+	}
+	
+	if (Result != Input.Num())
+	{
+		ShowErrorNotifications(Storage, Input);
+	}
+	
+	return Result;
+}
+
+bool FWidgetDropHandler::ExecuteDropSingle(DataStorage::ICoreProvider& Storage, UDropOperationSystem& DropSystem, DataStorage::RowHandle InputRow)
+{
+	using namespace UE::Editor::DataStorage;
+	return DropSystem.Apply(FRowHandleArrayView(&InputRow, 1, FRowHandleArray::EFlags::IsUnique | FRowHandleArray::EFlags::IsSorted)) != 0;
+}
+
 }
 
 #undef LOCTEXT_NAMESPACE
diff --git a/Engine/Plugins/Experimental/EditorDataStorageFeatures/Source/TedsOperations/Public/DragAndDrop/Widgets/PreviewWidgetDropHandler.h b/Engine/Plugins/Experimental/EditorDataStorageFeatures/Source/TedsOperations/Public/DragAndDrop/Widgets/PreviewWidgetDropHandler.h
--- a/Engine/Plugins/Experimental/EditorDataStorageFeatures/Source/TedsOperations/Public/DragAndDrop/Widgets/PreviewWidgetDropHandler.h
+++ b/Engine/Plugins/Experimental/EditorDataStorageFeatures/Source/TedsOperations/Public/DragAndDrop/Widgets/PreviewWidgetDropHandler.h
@@ -34,6 +34,9 @@ protected:
 
 	UE_API void GetPreviewRows(DataStorage::FRowHandleArray& OutRows, DataStorage::ICoreProvider& Storage) const;
 
+	UE_API virtual int32 ExecutePreviewDrop(DataStorage::ICoreProvider& Storage, UDropOperationSystem& DropSystem, DataStorage::FRowHandleArrayView InputRows);
+	UE_API virtual bool ExecutePreviewDropSingle(DataStorage::ICoreProvider& Storage, UDropOperationSystem& DropSystem, DataStorage::RowHandle InputRow);
+	
 protected:
 	DataStorage::FRowHandleArray PreviewInputRows;
 };
diff --git a/Engine/Plugins/Experimental/EditorDataStorageFeatures/Source/TedsOperations/Public/DragAndDrop/Widgets/WidgetDropHandler.h b/Engine/Plugins/Experimental/EditorDataStorageFeatures/Source/TedsOperations/Public/DragAndDrop/Widgets/WidgetDropHandler.h
--- a/Engine/Plugins/Experimental/EditorDataStorageFeatures/Source/TedsOperations/Public/DragAndDrop/Widgets/WidgetDropHandler.h
+++ b/Engine/Plugins/Experimental/EditorDataStorageFeatures/Source/TedsOperations/Public/DragAndDrop/Widgets/WidgetDropHandler.h
@@ -81,6 +81,9 @@ protected:
 
 	UE_API virtual void UpdateOperationDecorator(FDragDropOperation& Operation, TOptional<FDropResult> Result) const;
 
+	UE_API virtual int32 ExecuteDrop(DataStorage::ICoreProvider& Storage, UDropOperationSystem& DropSystem, DataStorage::FRowHandleArrayView InputRows);
+	UE_API virtual bool ExecuteDropSingle(DataStorage::ICoreProvider& Storage, UDropOperationSystem& DropSystem, DataStorage::RowHandle InputRow);
+
 protected:
 	TWeakObjectPtr<UTypedElementSelectionSet> SelectionSet;
 
