Fixed Spaces' CB Favorites integration not adding favorites on launch
[Added] Post Register to AssetProcessor
[Added] OnAssetRegistryStorageInit is called if asset registery is already enabled.

#Synced-CL 49621724

diff --git a/Engine/Plugins/Experimental/EditorDataStorageFeatures/Source/TedsAssetData/Private/Processors/AssetProcessors.cpp b/Engine/Plugins/Experimental/EditorDataStorageFeatures/Source/TedsAssetData/Private/Processors/AssetProcessors.cpp
--- a/Engine/Plugins/Experimental/EditorDataStorageFeatures/Source/TedsAssetData/Private/Processors/AssetProcessors.cpp
+++ b/Engine/Plugins/Experimental/EditorDataStorageFeatures/Source/TedsAssetData/Private/Processors/AssetProcessors.cpp
@@ -190,11 +190,6 @@ void UTedsAssetDataFactory::PreRegister(UE::Editor::DataStorage::ICoreProvider&
 	{
 		return;
 	}
-	
-	if (UE::Editor::AssetData::FTedsAssetDataModule* TedsAssetDataModule = FModuleManager::Get().GetModulePtr<UE::Editor::AssetData::FTedsAssetDataModule>("TedsAssetData"))
-	{
-		TedsAssetDataModule->OnAssetRegistryStorageInit().AddUObject(this, &UTedsAssetDataFactory::OnAssetDataStorageEnabled, &DataStorage);
-	}
 
 	if (FContentBrowserModule* ContentBrowserModule = FModuleManager::Get().GetModulePtr<FContentBrowserModule>("ContentBrowser"))
 	{
@@ -205,6 +200,21 @@ void UTedsAssetDataFactory::PreRegister(UE::Editor::DataStorage::ICoreProvider&
 		FOnFavoritesChanged::FDelegate::CreateUObject(this, &UTedsAssetDataFactory::OnFavoritesChanged, &DataStorage));
 }
 
+void UTedsAssetDataFactory::PostRegister(UE::Editor::DataStorage::ICoreProvider& DataStorage)
+{
+	if (UE::Editor::AssetData::FTedsAssetDataModule* TedsAssetDataModule = FModuleManager::Get().GetModulePtr<UE::Editor::AssetData::FTedsAssetDataModule>("TedsAssetData"))
+	{
+		if (TedsAssetDataModule->IsTedsAssetRegistryStorageEnabled())
+		{
+			OnAssetDataStorageEnabled(&DataStorage);
+		}
+		else
+		{
+			TedsAssetDataModule->OnAssetRegistryStorageInit().AddUObject(this, &UTedsAssetDataFactory::OnAssetDataStorageEnabled, &DataStorage);
+		}
+	}
+}
+
 void UTedsAssetDataFactory::PreShutdown(UE::Editor::DataStorage::ICoreProvider& DataStorage)
 {
 	if (FContentBrowserModule* ContentBrowserModule = FModuleManager::Get().GetModulePtr<FContentBrowserModule>("ContentBrowser"))
diff --git a/Engine/Plugins/Experimental/EditorDataStorageFeatures/Source/TedsAssetData/Private/Processors/AssetProcessors.h b/Engine/Plugins/Experimental/EditorDataStorageFeatures/Source/TedsAssetData/Private/Processors/AssetProcessors.h
--- a/Engine/Plugins/Experimental/EditorDataStorageFeatures/Source/TedsAssetData/Private/Processors/AssetProcessors.h
+++ b/Engine/Plugins/Experimental/EditorDataStorageFeatures/Source/TedsAssetData/Private/Processors/AssetProcessors.h
@@ -37,6 +37,7 @@ public:
 	virtual void RegisterTables(UE::Editor::DataStorage::ICoreProvider& DataStorage) override;
 	void RegisterQueries(UE::Editor::DataStorage::ICoreProvider& DataStorage) override;
 	virtual void PreRegister(UE::Editor::DataStorage::ICoreProvider& DataStorage) override;
+	virtual void PostRegister(UE::Editor::DataStorage::ICoreProvider& DataStorage) override;
 	virtual void PreShutdown(UE::Editor::DataStorage::ICoreProvider& DataStorage) override;
 
 protected:
