WidgetDropHandler now ignores the drop operation if no operations support it.

diff --git a/Engine/Plugins/Experimental/EditorDataStorageFeatures/Source/TedsOperations/Private/DragAndDrop/Widgets/WidgetDropHandler.cpp b/Engine/Plugins/Experimental/EditorDataStorageFeatures/Source/TedsOperations/Private/DragAndDrop/Widgets/WidgetDropHandler.cpp
--- a/Engine/Plugins/Experimental/EditorDataStorageFeatures/Source/TedsOperations/Private/DragAndDrop/Widgets/WidgetDropHandler.cpp
+++ b/Engine/Plugins/Experimental/EditorDataStorageFeatures/Source/TedsOperations/Private/DragAndDrop/Widgets/WidgetDropHandler.cpp
@@ -237,6 +237,20 @@ TOptional<FWidgetDropHandler::FDropResult> FWidgetDropHandler::Start(DataStorage
 	}
 	(void)UpdateParameters(Storage, Geometry, InputEvent);
 	SetRowsFromParameters(Storage, InputRows.GetRows());
+	
+	// We ignore the drop if there are no operations for it.
+	TArray<RowHandle> Operations;
+	for (RowHandle InputRow : InputRows.GetRows())
+	{
+		Operations.Reset();
+		DropSystem.GetOperations(Operations, InputRow);
+		if (Operations.IsEmpty())
+		{
+			Stop(Storage, DropSystem);
+			return {};
+		}
+	}
+	
 	return ResetOperations(Storage, DropSystem);
 }
 
