[UE-355612] Deprecate UI capabilities from TEDS UI

[Deprecate] ITypedElementUiStyleOverrideCapability

#Synced-CL 49823132

#Synced-CL 49823132

diff --git a/Engine/Plugins/Experimental/EditorDataStorage/Source/TedsUI/Private/Widgets/LabelWidget.cpp b/Engine/Plugins/Experimental/EditorDataStorage/Source/TedsUI/Private/Widgets/LabelWidget.cpp
--- a/Engine/Plugins/Experimental/EditorDataStorage/Source/TedsUI/Private/Widgets/LabelWidget.cpp
+++ b/Engine/Plugins/Experimental/EditorDataStorage/Source/TedsUI/Private/Widgets/LabelWidget.cpp
@@ -11,9 +11,6 @@
 #include "Elements/Framework/TypedElementAttributeBinding.h"
 #include "Elements/Framework/TypedElementQueryBuilder.h"
 #include "Elements/Interfaces/TypedElementDataStorageInterface.h"
-#include "Elements/Interfaces/Capabilities/TypedElementUiTextCapability.h"
-#include "Elements/Interfaces/Capabilities/TypedElementUiTooltipCapability.h"
-#include "Elements/Interfaces/Capabilities/TypedElementUiStyleOverrideCapability.h"
 #include "Widgets/Text/STextBlock.h"
 #include "Widgets/Text/SInlineEditableTextBlock.h"
 
diff --git a/Engine/Plugins/Experimental/EditorDataStorageFeatures/Source/TedsActorCompatibility/Private/LevelInstanceFactory.cpp b/Engine/Plugins/Experimental/EditorDataStorageFeatures/Source/TedsActorCompatibility/Private/LevelInstanceFactory.cpp
new file mode 100644
--- /dev/null
+++ b/Engine/Plugins/Experimental/EditorDataStorageFeatures/Source/TedsActorCompatibility/Private/LevelInstanceFactory.cpp
@@ -0,0 +1,55 @@
+// Copyright Epic Games, Inc. All Rights Reserved.
+
+
+#include "LevelInstanceFactory.h"
+
+#include "Columns/TedsLevelInstanceColumns.h"
+#include "DataStorage/Features.h"
+#include "DataStorage/Handles.h"
+#include "GameFramework/Actor.h"
+#include "Modules/ModuleManager.h"
+#include "LevelInstance/LevelInstanceInterface.h"
+
+void UTedsLevelInstanceFactory::PreRegister(UE::Editor::DataStorage::ICoreProvider& InDataStorage)
+{
+	if (ILevelInstanceEditorModule* EditorModule = FModuleManager::GetModulePtr<ILevelInstanceEditorModule>("LevelInstanceEditor"))
+	{
+		EditorModule->OnLevelInstanceEditModeChanged().AddUObject(this, &UTedsLevelInstanceFactory::LevelInstanceEditModeChanged);
+	}
+	DataStorage = &InDataStorage;
+
+	using namespace UE::Editor::DataStorage;
+	CompatibilityProvider = GetMutableDataStorageFeature<ICompatibilityProvider>(CompatibilityFeatureName);
+}
+
+void UTedsLevelInstanceFactory::PreShutdown(UE::Editor::DataStorage::ICoreProvider& InDataStorage)
+{
+	if (ILevelInstanceEditorModule* EditorModule = FModuleManager::GetModulePtr<ILevelInstanceEditorModule>("LevelInstanceEditor"))
+	{
+		EditorModule->OnLevelInstanceEditModeChanged().RemoveAll(this);
+	}
+}
+
+void UTedsLevelInstanceFactory::LevelInstanceEditModeChanged(ILevelInstanceInterface* LevelInstanceInterface,
+	ILevelInstanceEditorModule::ELevelInstanceEditMode LevelInstanceEditMode) const
+{
+	if (AActor* LevelInstanceActor = Cast<AActor>(LevelInstanceInterface))
+	{
+		using namespace UE::Editor::DataStorage;
+
+		if (CompatibilityProvider)
+		{
+			if (const RowHandle Row = CompatibilityProvider->FindRowWithCompatibleObject(LevelInstanceActor); Row != InvalidRowHandle)
+			{
+				if (LevelInstanceEditMode == ILevelInstanceEditorModule::ELevelInstanceEditMode::None)
+				{
+					DataStorage->RemoveColumn<FLevelInstanceEditingColumn>(Row);
+				}
+				else
+				{
+					DataStorage->AddColumn<FLevelInstanceEditingColumn>(Row, { .EditMode = LevelInstanceEditMode });
+				}
+			}
+		}
+	}
+}
diff --git a/Engine/Plugins/Experimental/EditorDataStorageFeatures/Source/TedsActorCompatibility/Private/LevelInstanceFactory.h b/Engine/Plugins/Experimental/EditorDataStorageFeatures/Source/TedsActorCompatibility/Private/LevelInstanceFactory.h
new file mode 100644
--- /dev/null
+++ b/Engine/Plugins/Experimental/EditorDataStorageFeatures/Source/TedsActorCompatibility/Private/LevelInstanceFactory.h
@@ -0,0 +1,28 @@
+// Copyright Epic Games, Inc. All Rights Reserved.
+
+#pragma once
+
+#include "CoreMinimal.h"
+#include "Elements/Interfaces/TypedElementDataStorageFactory.h"
+#include "LevelInstance/ILevelInstanceEditorModule.h"
+#include "Elements/Interfaces/TypedElementDataStorageCompatibilityInterface.h"
+#include "LevelInstanceFactory.generated.h"
+
+class IEditorDataStorageProvider;
+
+
+UCLASS()
+class UTedsLevelInstanceFactory : public UEditorDataStorageFactory
+{
+	GENERATED_BODY()
+public:
+	virtual ~UTedsLevelInstanceFactory() override = default;
+	virtual void PreRegister(UE::Editor::DataStorage::ICoreProvider& InDataStorage) override;
+	virtual void PreShutdown(UE::Editor::DataStorage::ICoreProvider& DataStorage) override;
+
+private:
+	void LevelInstanceEditModeChanged(ILevelInstanceInterface* LevelInstanceInterface, ILevelInstanceEditorModule::ELevelInstanceEditMode LevelInstanceEditMode) const;
+
+	UE::Editor::DataStorage::ICoreProvider* DataStorage;
+	UE::Editor::DataStorage::ICompatibilityProvider* CompatibilityProvider;
+};
diff --git a/Engine/Plugins/Experimental/EditorDataStorageFeatures/Source/TedsActorCompatibility/Public/Columns/TedsLevelInstanceColumns.h b/Engine/Plugins/Experimental/EditorDataStorageFeatures/Source/TedsActorCompatibility/Public/Columns/TedsLevelInstanceColumns.h
new file mode 100644
--- /dev/null
+++ b/Engine/Plugins/Experimental/EditorDataStorageFeatures/Source/TedsActorCompatibility/Public/Columns/TedsLevelInstanceColumns.h
@@ -0,0 +1,18 @@
+// Copyright Epic Games, Inc. All Rights Reserved.
+
+#pragma once
+#include "DataStorage/CommonTypes.h"
+#include "LevelInstance/ILevelInstanceEditorModule.h"
+
+#include "TedsLevelInstanceColumns.generated.h"
+
+/**
+ * Tag to identify a row with an LevelInstance actor that is being edited.
+ */
+USTRUCT(meta = (DisplayName = "Editing Level Instance"))
+struct FLevelInstanceEditingColumn final : public FEditorDataStorageColumn
+{
+	GENERATED_BODY()
+
+	ILevelInstanceEditorModule::ELevelInstanceEditMode EditMode;
+};
\ No newline at end of file
diff --git a/Engine/Plugins/Experimental/EditorDataStorageFeatures/Source/TedsActorCompatibility/TedsActorCompatibility.Build.cs b/Engine/Plugins/Experimental/EditorDataStorageFeatures/Source/TedsActorCompatibility/TedsActorCompatibility.Build.cs
--- a/Engine/Plugins/Experimental/EditorDataStorageFeatures/Source/TedsActorCompatibility/TedsActorCompatibility.Build.cs
+++ b/Engine/Plugins/Experimental/EditorDataStorageFeatures/Source/TedsActorCompatibility/TedsActorCompatibility.Build.cs
@@ -30,6 +30,7 @@ public class TedsActorCompatibility : ModuleRules
 					"TedsTableViewer",
 					"SlateCore",
 					"Slate",
+					"LevelInstanceEditor"
 				});
 
 			DynamicallyLoadedModuleNames.AddRange(new string[] { });
diff --git a/Engine/Plugins/Experimental/EditorDataStorageFeatures/Source/TedsOutliner/Private/Widgets/OutlinerTextWidget.cpp b/Engine/Plugins/Experimental/EditorDataStorageFeatures/Source/TedsOutliner/Private/Widgets/OutlinerTextWidget.cpp
--- a/Engine/Plugins/Experimental/EditorDataStorageFeatures/Source/TedsOutliner/Private/Widgets/OutlinerTextWidget.cpp
+++ b/Engine/Plugins/Experimental/EditorDataStorageFeatures/Source/TedsOutliner/Private/Widgets/OutlinerTextWidget.cpp
@@ -3,16 +3,19 @@
 #include "Widgets/OutlinerTextWidget.h"
 
 #include "ActorEditorUtils.h"
+#include "ISceneOutliner.h"
 #include "ITedsTableViewer.h"
+#include "SceneOutlinerStandaloneTypes.h"
+#include "TedsOutlinerHelpers.h"
 #include "TedsTableViewerUtils.h"
 #include "TedsTableViewerWidgetColumns.h"
-#include "TedsTableViewerWidgetColumns.h"
 #include "Columns/SlateDelegateColumns.h"
-#include "Elements/Columns/TypedElementHiearchyColumns.h"
+#include "Columns/TedsLevelInstanceColumns.h"
+#include "DataStorage/Features.h"
 #include "Elements/Columns/TypedElementLabelColumns.h"
 #include "Elements/Columns/TypedElementMiscColumns.h"
+#include "Elements/Columns/TypedElementSlateWidgetColumns.h"
 #include "Elements/Framework/TypedElementAttributeBinding.h"
-#include "Elements/Interfaces/Capabilities/TypedElementUiStyleOverrideCapability.h"
 #include "Settings/EditorStyleSettings.h"
 #include "Widgets/Text/SInlineEditableTextBlock.h"
 #include "Widgets/Text/STextBlock.h"
@@ -23,6 +26,45 @@
 
 namespace UE::Editor::Outliner::TextWidget::Private
 {
+
+	FSlateColor GetForegroundColor(
+		TNotNull<DataStorage::ICoreProvider*> DataStorage, 
+		DataStorage::RowHandle DataRow, 
+		DataStorage::RowHandle WidgetRow
+	)
+	{
+		if (FLevelInstanceEditingColumn* LevelInstanceEditingColumn = DataStorage->GetColumn<FLevelInstanceEditingColumn>(DataRow))
+		{
+			switch (LevelInstanceEditingColumn->EditMode)
+			{
+			case ILevelInstanceEditorModule::ELevelInstanceEditMode::None:
+				break;
+			case ILevelInstanceEditorModule::ELevelInstanceEditMode::Original:
+				return FAppStyle::Get().GetSlateColor("Colors.AccentGreen").GetSpecifiedColor();
+			case ILevelInstanceEditorModule::ELevelInstanceEditMode::Override:
+				return FAppStyle::Get().GetSlateColor("Colors.AccentBlue").GetSpecifiedColor();
+			default:
+				break;
+			}
+		}
+
+		if (const FTedsOutlinerColumn* OutlinerColumn = DataStorage->GetColumn<FTedsOutlinerColumn>(WidgetRow))
+		{
+			if (const TSharedPtr<ISceneOutliner> Outliner = OutlinerColumn->Outliner.Pin())
+			{
+				if (const FSceneOutlinerTreeItemPtr TreeItemPtr = Helpers::GetTreeItemFromRowHandle(DataStorage, Outliner.ToSharedRef(), DataRow))
+				{
+					if (TOptional<FLinearColor> ForegroundColor = FSceneOutlinerCommonLabelData::GetForegroundColor(Outliner, *TreeItemPtr))
+					{
+						return ForegroundColor.GetValue();
+					}
+				}
+			}
+		}
+
+		return FSlateColor::UseForeground();
+	}
+
 	// We are using an inheritance based model with a custom SWidget class just so we can manage lifetimes of delegates easier using BindSP
 	class SOutlinerEditableTextWidget : public SInlineEditableTextBlock
 	{
@@ -57,11 +99,14 @@ namespace UE::Editor::Outliner::TextWidget::Private
 				.Text(DataRowBinder.BindText(&FTypedElementLabelColumn::Label))
 				.ToolTipText(DataRowBinder.BindText(&FTypedElementLabelColumn::Label))
 				.OverflowPolicy(bShouldUseMiddleEllipsis ? ETextOverflowPolicy::MiddleEllipsis : TOptional<ETextOverflowPolicy>())
-				.IsSelected(WidgetRowBinder.BindEvent(&FExternalWidgetSelectionColumn::IsSelected)));
+				.IsSelected(WidgetRowBinder.BindEvent(&FExternalWidgetSelectionColumn::IsSelected))
+				.ColorAndOpacity(TAttribute<FSlateColor>::CreateLambda([this]() -> FSlateColor
+					{
+						return Private::GetForegroundColor(DataStorage, DataRow, WidgetRow);
+					}))
+				);
 
 			BindTableViewerRenameDelegate();
-
-			AddMetadata(MakeShared<TTypedElementUiStyleOverrideCapability<SInlineEditableTextBlock>>(*this));
 		}
 		
 		virtual ~SOutlinerEditableTextWidget() override = default;
@@ -188,7 +233,6 @@ TSharedPtr<SWidget> FOutlinerTextWidgetConstructor::CreateWidget(UE::Editor::Dat
 		Result = CreateNonEditableWidget(DataStorage, DataStorageUi, TargetRow, WidgetRow, Arguments);
 	}
 	
-	
 	return Result;
 }
 
@@ -260,9 +304,12 @@ TSharedPtr<STextBlock> FOutlinerTextWidgetConstructor::CreateNonEditableTextBloc
 	FAttributeBinder TargetRowBinder(TargetRow, DataStorage);
 
 	TSharedPtr<STextBlock> TextBlock = SNew(STextBlock)
-		.IsEnabled(false)
 		.Text(TargetRowBinder.BindText(&FTypedElementLabelColumn::Label))
-		.ToolTipText(TargetRowBinder.BindText(&FTypedElementLabelColumn::Label));
+		.ToolTipText(TargetRowBinder.BindText(&FTypedElementLabelColumn::Label))
+		.ColorAndOpacity(TAttribute<FSlateColor>::CreateLambda([DataStorage, TargetRow, WidgetRow]() -> FSlateColor
+			{
+				return UE::Editor::Outliner::TextWidget::Private::GetForegroundColor(DataStorage, TargetRow, WidgetRow);
+			}));
 
 	return TextBlock;
 }
diff --git a/Engine/Plugins/Experimental/EditorDataStorageFeatures/Source/TedsOutliner/TedsOutliner.Build.cs b/Engine/Plugins/Experimental/EditorDataStorageFeatures/Source/TedsOutliner/TedsOutliner.Build.cs
--- a/Engine/Plugins/Experimental/EditorDataStorageFeatures/Source/TedsOutliner/TedsOutliner.Build.cs
+++ b/Engine/Plugins/Experimental/EditorDataStorageFeatures/Source/TedsOutliner/TedsOutliner.Build.cs
@@ -42,7 +42,8 @@ public class TedsOutliner : ModuleRules
 					"LevelEditor",
 					"TedsQueryStack",
 					"TedsActorCompatibility",
-					"TedsOperations"
+					"TedsOperations",
+					"LevelInstanceEditor"
 				});
 
 			DynamicallyLoadedModuleNames.AddRange(new string[] {});
diff --git a/Engine/Source/Editor/SceneOutliner/Private/SceneOutlinerStandaloneTypes.cpp b/Engine/Source/Editor/SceneOutliner/Private/SceneOutlinerStandaloneTypes.cpp
--- a/Engine/Source/Editor/SceneOutliner/Private/SceneOutlinerStandaloneTypes.cpp
+++ b/Engine/Source/Editor/SceneOutliner/Private/SceneOutlinerStandaloneTypes.cpp
@@ -17,7 +17,7 @@ const FSceneOutlinerTreeItemType ISceneOutlinerTreeItem::Type;
 
 const FLinearColor FSceneOutlinerCommonLabelData::DarkColor(0.15f, 0.15f, 0.15f);
 
-TOptional<FLinearColor> FSceneOutlinerCommonLabelData::GetForegroundColor(const ISceneOutlinerTreeItem& TreeItem) const
+TOptional<FLinearColor> FSceneOutlinerCommonLabelData::GetForegroundColor(TWeakPtr<ISceneOutliner> WeakSceneOutliner, const ISceneOutlinerTreeItem& TreeItem)
 {
 	if (!TreeItem.IsValid())
 	{
@@ -30,8 +30,15 @@ TOptional<FLinearColor> FSceneOutlinerCommonLabelData::GetForegroundColor(const
 		TSharedPtr<FDragDropOperation> DragDropOp = FSlateApplication::Get().GetDragDroppingContent();
 
 		FSceneOutlinerDragDropPayload DraggedObjects(*DragDropOp);
-		const auto Outliner = WeakSceneOutliner.Pin();
-		if (Outliner->GetMode()->ParseDragDrop(DraggedObjects, *DragDropOp) && !Outliner->GetMode()->ValidateDrop(TreeItem, DraggedObjects).IsValid())
+
+		if (const TSharedPtr<ISceneOutliner> Outliner = WeakSceneOutliner.Pin())
+		{
+			if (Outliner->GetMode()->ParseDragDrop(DraggedObjects, *DragDropOp) && !Outliner->GetMode()->ValidateDrop(TreeItem, DraggedObjects).IsValid())
+			{
+				return DarkColor;
+			}
+		}
+		else
 		{
 			return DarkColor;
 		}
@@ -70,6 +77,11 @@ TOptional<FLinearColor> FSceneOutlinerCommonLabelData::GetForegroundColor(const
 	return TOptional<FLinearColor>();
 }
 
+TOptional<FLinearColor> FSceneOutlinerCommonLabelData::GetForegroundColor(const ISceneOutlinerTreeItem& TreeItem) const
+{
+	return GetForegroundColor(WeakSceneOutliner, TreeItem);
+}
+
 bool FSceneOutlinerCommonLabelData::CanExecuteRenameRequest(const ISceneOutlinerTreeItem& Item) const
 {
 	if (const ISceneOutliner* SceneOutliner = WeakSceneOutliner.Pin().Get())
diff --git a/Engine/Source/Editor/SceneOutliner/Public/SceneOutlinerStandaloneTypes.h b/Engine/Source/Editor/SceneOutliner/Public/SceneOutlinerStandaloneTypes.h
--- a/Engine/Source/Editor/SceneOutliner/Public/SceneOutlinerStandaloneTypes.h
+++ b/Engine/Source/Editor/SceneOutliner/Public/SceneOutlinerStandaloneTypes.h
@@ -212,6 +212,8 @@ struct FSceneOutlinerCommonLabelData
 	TWeakPtr<ISceneOutliner> WeakSceneOutliner;
 	static UE_API const FLinearColor DarkColor;
 
+	static UE_API TOptional<FLinearColor> GetForegroundColor(TWeakPtr<ISceneOutliner> WeakSceneOutliner, const ISceneOutlinerTreeItem& TreeItem);
+
 	UE_API TOptional<FLinearColor> GetForegroundColor(const ISceneOutlinerTreeItem& TreeItem) const;
 
 	UE_API bool CanExecuteRenameRequest(const ISceneOutlinerTreeItem& Item) const;
diff --git a/Engine/Source/Runtime/TypedElementFramework/Public/Elements/Interfaces/Capabilities/TypedElementUiStyleOverrideCapability.h b/Engine/Source/Runtime/TypedElementFramework/Public/Elements/Interfaces/Capabilities/TypedElementUiStyleOverrideCapability.h
--- a/Engine/Source/Runtime/TypedElementFramework/Public/Elements/Interfaces/Capabilities/TypedElementUiStyleOverrideCapability.h
+++ b/Engine/Source/Runtime/TypedElementFramework/Public/Elements/Interfaces/Capabilities/TypedElementUiStyleOverrideCapability.h
@@ -4,14 +4,14 @@
 
 #include "Elements/Interfaces/TypedElementUiCapabilities.h"
 #include "Misc/Attribute.h"
-#include "Math/Color.h"
 
 struct FSlateColor;
 
 /**
  * Interface to provide access to widgets that support working with style overrides.
  */
-class ITypedElementUiStyleOverrideCapability : public ITypedElementUiCapability
+class UE_DEPRECATED(5.8, "Please use FTedsOutlinerColumn to find the foreground color using Helpers::GetTreeItemFromRowHandle and FSceneOutlinerCommonLabelData::GetForegroundColor")
+	ITypedElementUiStyleOverrideCapability : public ITypedElementUiCapability
 {
 public:
 	SLATE_METADATA_TYPE(ITypedElementUiStyleOverrideCapability, ITypedElementUiCapability)
@@ -22,7 +22,7 @@ public:
 };
 
 template<typename WidgetType>
-class TTypedElementUiStyleOverrideCapability : public ITypedElementUiStyleOverrideCapability
+class UE_DEPRECATED(5.8, "Please use FTedsOutlinerColumn to find the foreground color using Helpers::GetTreeItemFromRowHandle and FSceneOutlinerCommonLabelData::GetForegroundColor") TTypedElementUiStyleOverrideCapability : public ITypedElementUiStyleOverrideCapability
 {
 public:
 	explicit TTypedElementUiStyleOverrideCapability(WidgetType& InWidget) : Widget(InWidget){}
