[UE-355612] Deprecate UI capabilities from TEDS UI

[Deprecate] ITypedElementUiTextCapability

#Synced-CL 50045910

diff --git a/Engine/Plugins/Experimental/EditorDataStorageFeatures/Source/TedsOutliner/Private/TedsOutlinerColumn.cpp b/Engine/Plugins/Experimental/EditorDataStorageFeatures/Source/TedsOutliner/Private/TedsOutlinerColumn.cpp
--- a/Engine/Plugins/Experimental/EditorDataStorageFeatures/Source/TedsOutliner/Private/TedsOutlinerColumn.cpp
+++ b/Engine/Plugins/Experimental/EditorDataStorageFeatures/Source/TedsOutliner/Private/TedsOutlinerColumn.cpp
@@ -3,11 +3,12 @@
 #include "TedsOutlinerColumn.h"
 
 #include "ActorTreeItem.h"
-#include "Elements/Interfaces/Capabilities/TypedElementUiTextCapability.h"
 #include "Modules/ModuleManager.h"
 #include "SceneOutlinerModule.h"
 #include "TedsOutlinerHelpers.h"
 #include "TedsOutlinerItem.h"
+#include "TedsTableViewerWidgetColumns.h"
+#include "Elements/Columns/TypedElementSlateWidgetColumns.h"
 
 namespace UE::Editor::Outliner
 {
@@ -122,29 +123,6 @@ namespace UE::Editor::Outliner
 	{
 		return !bHybridMode && TedsOutlinerImpl.IsValid() ? !TedsOutlinerImpl.Pin()->IsSorting() : true;
 	}
-
-	void FTedsOutlinerUiColumn::SetHighlightText(SWidget& Widget)
-	{
-		TSharedPtr<ISceneOutliner> OutlinerPinned = OwningOutliner.Pin();
-
-		if (!OutlinerPinned)
-		{
-			return;
-		}
-
-		if (TSharedPtr<ITypedElementUiTextCapability> TextCapability = Widget.GetMetaData<ITypedElementUiTextCapability>())
-		{
-			TextCapability->SetHighlightText(OutlinerPinned->GetFilterHighlightText());
-		}
-	
-		if (FChildren* ChildWidgets = Widget.GetChildren())
-		{
-			ChildWidgets->ForEachWidget([this](SWidget& ChildWidget)
-				{
-					SetHighlightText(ChildWidget);
-				});
-		}
-	}
 	
 	const TSharedRef<SWidget> FTedsOutlinerUiColumn::ConstructRowWidget(FSceneOutlinerTreeItemRef TreeItem, const STableRow<FSceneOutlinerTreeItemPtr>& Row)
 	{
@@ -182,7 +160,6 @@ namespace UE::Editor::Outliner
 
 		if (RowWidget)
 		{
-			SetHighlightText(*RowWidget);
 			return RowWidget.ToSharedRef();
 		}
 
diff --git a/Engine/Plugins/Experimental/EditorDataStorageFeatures/Source/TedsOutliner/Private/TedsOutlinerColumn.h b/Engine/Plugins/Experimental/EditorDataStorageFeatures/Source/TedsOutliner/Private/TedsOutlinerColumn.h
--- a/Engine/Plugins/Experimental/EditorDataStorageFeatures/Source/TedsOutliner/Private/TedsOutlinerColumn.h
+++ b/Engine/Plugins/Experimental/EditorDataStorageFeatures/Source/TedsOutliner/Private/TedsOutlinerColumn.h
@@ -63,7 +63,6 @@ namespace UE::Editor::Outliner
 		const TSharedRef<SWidget> ConstructRowWidget(FSceneOutlinerTreeItemRef TreeItem, const STableRow<FSceneOutlinerTreeItemPtr>& Row) override;
 		virtual void PopulateSearchStrings( const ISceneOutlinerTreeItem& Item, TArray< FString >& OutSearchStrings ) const override;
 		
-		void SetHighlightText(SWidget& Widget);
 		bool IsRowVisible(const DataStorage::RowHandle InRowHandle) const;
 
 	protected:
diff --git a/Engine/Plugins/Experimental/EditorDataStorageFeatures/Source/TedsOutliner/Private/TedsOutlinerImpl.cpp b/Engine/Plugins/Experimental/EditorDataStorageFeatures/Source/TedsOutliner/Private/TedsOutlinerImpl.cpp
--- a/Engine/Plugins/Experimental/EditorDataStorageFeatures/Source/TedsOutliner/Private/TedsOutlinerImpl.cpp
+++ b/Engine/Plugins/Experimental/EditorDataStorageFeatures/Source/TedsOutliner/Private/TedsOutlinerImpl.cpp
@@ -31,6 +31,7 @@
 #include "TedsRowOrderInversionNode.h"
 #include "TedsRowQueryResultsNode.h"
 #include "TedsRowSortNode.h"
+#include "TedsTableViewerWidgetColumns.h"
 #include "Columns/TedsOutlinerColumns.h"
 #include "Compatibility/SceneOutlinerTedsBridge.h"
 #include "Elements/Columns/TypedElementCompatibilityColumns.h"
@@ -148,7 +149,12 @@ namespace LabelWidgetUtils
 		static const FQueryDescription MetaDataQueryRead = Select().ReadOnly<FTypedElementLabelColumn>().Where().Compile();
 	
 		RowHandle WidgetRow = Storage.AddRow(Storage.FindTable(TableViewerUtils::GetWidgetTableName()));
+
 		checkf(WidgetRow != InvalidRowHandle, TEXT("Expected valid row handle!"));
+
+		FHierarchyHandle Hierarchy = Storage.FindHierarchyByName(TableViewerUtils::GetWidgetHierarchyName());
+		Storage.SetParentRow(Hierarchy, WidgetRow, TargetRow);
+
 		if (FTypedElementRowReferenceColumn* RowReference = Storage.GetColumn<FTypedElementRowReferenceColumn>(WidgetRow))
 		{
 			RowReference->Row = TargetRow;
@@ -371,6 +377,9 @@ FTedsOutlinerImpl::~FTedsOutlinerImpl()
 		Storage->RemoveRow(OutlinerRowHandle);
 		Storage->OnUpdateCompleted().RemoveAll(this);
 	}
+
+	SceneOutliner->OnTextFilterChanged().Remove(OnTextFilteredChangedHandle);
+
 	UnregisterQueries();
 	FTSTicker::RemoveTicker(TickerHandle);
 }
@@ -505,7 +514,18 @@ void FTedsOutlinerImpl::RegisterTedsOutliner()
 		{
 			TedsOutlinerColumn->Outliner = StaticCastSharedRef<ISceneOutliner>(SceneOutliner->AsShared());
 		}
-		
+
+		Storage->AddColumn<FTableViewerTag>(OutlinerRowHandle);
+		Storage->AddColumn(OutlinerRowHandle, FHighlightTextColumn{});
+
+		OnTextFilteredChangedHandle = SceneOutliner->OnTextFilterChanged().AddLambda([this]()
+			{
+				if (FHighlightTextColumn* HighlightTextColumn = Storage->GetColumn<FHighlightTextColumn>(OutlinerRowHandle))
+				{
+					HighlightTextColumn->HighlightText = SceneOutliner->GetFilterHighlightText().Get();
+				}
+			});
+
 		if (FTedsOutlinerColumnQueryColumn* ColumnQueryColumn = Storage->GetColumn<FTedsOutlinerColumnQueryColumn>(OutlinerRowHandle))
 		{
 			ColumnQueryColumn->ColumnQueryHandle = CreationParams.ColumnQuery;
diff --git a/Engine/Plugins/Experimental/EditorDataStorageFeatures/Source/TedsOutliner/Private/TedsOutlinerItem.cpp b/Engine/Plugins/Experimental/EditorDataStorageFeatures/Source/TedsOutliner/Private/TedsOutlinerItem.cpp
--- a/Engine/Plugins/Experimental/EditorDataStorageFeatures/Source/TedsOutliner/Private/TedsOutlinerItem.cpp
+++ b/Engine/Plugins/Experimental/EditorDataStorageFeatures/Source/TedsOutliner/Private/TedsOutlinerItem.cpp
@@ -6,6 +6,7 @@
 #include "Elements/Columns/TypedElementMiscColumns.h"
 #include "HAL/PlatformApplicationMisc.h"
 #include "TedsOutlinerImpl.h"
+#include "TedsTableViewerUtils.h"
 
 #include "UObject/PropertyBagRepository.h"
 
@@ -70,7 +71,15 @@ bool FTedsOutlinerTreeItem::CanInteract() const
 TSharedRef<SWidget> FTedsOutlinerTreeItem::GenerateLabelWidget(ISceneOutliner& Outliner,
 	const STableRow<FSceneOutlinerTreeItemPtr>& InRow)
 {
-	return TedsOutlinerImpl.IsValid() ? TedsOutlinerImpl.Pin()->CreateLabelWidgetForItem(RowHandle, *this, InRow, CanInteract()) : SNullWidget::NullWidget;
+	if (TSharedPtr TedsOutlinerImplPin = TedsOutlinerImpl.Pin())
+	{
+		ICoreProvider* Storage = TedsOutlinerImplPin->GetStorage();
+		FHierarchyHandle Hierarchy = Storage->FindHierarchyByName(TableViewerUtils::GetWidgetHierarchyName());
+		Storage->SetParentRow(Hierarchy, RowHandle, TedsOutlinerImplPin->GetOutlinerRowHandle());
+
+		return TedsOutlinerImplPin->CreateLabelWidgetForItem(RowHandle, *this, InRow, CanInteract());
+	}
+	return SNullWidget::NullWidget;
 }
 
 void FTedsOutlinerTreeItem::GenerateContextMenu(UToolMenu* Menu, SSceneOutliner& Outliner)
diff --git a/Engine/Plugins/Experimental/EditorDataStorageFeatures/Source/TedsOutliner/Private/Widgets/OutlinerTextWidget.cpp b/Engine/Plugins/Experimental/EditorDataStorageFeatures/Source/TedsOutliner/Private/Widgets/OutlinerTextWidget.cpp
--- a/Engine/Plugins/Experimental/EditorDataStorageFeatures/Source/TedsOutliner/Private/Widgets/OutlinerTextWidget.cpp
+++ b/Engine/Plugins/Experimental/EditorDataStorageFeatures/Source/TedsOutliner/Private/Widgets/OutlinerTextWidget.cpp
@@ -27,42 +27,46 @@
 namespace UE::Editor::Outliner::TextWidget::Private
 {
 
-	FSlateColor GetForegroundColor(
+	TAttribute<FSlateColor> GetForegroundColor(
 		TNotNull<DataStorage::ICoreProvider*> DataStorage, 
 		DataStorage::RowHandle DataRow, 
 		DataStorage::RowHandle WidgetRow
 	)
 	{
-		if (FLevelInstanceEditingColumn* LevelInstanceEditingColumn = DataStorage->GetColumn<FLevelInstanceEditingColumn>(DataRow))
-		{
-			switch (LevelInstanceEditingColumn->EditMode)
+		return MakeAttributeLambda([DataStorage, WidgetRow, DataRow]() -> FSlateColor
 			{
-			case ILevelInstanceEditorModule::ELevelInstanceEditMode::None:
-				break;
-			case ILevelInstanceEditorModule::ELevelInstanceEditMode::Original:
-				return FAppStyle::Get().GetSlateColor("Colors.AccentGreen").GetSpecifiedColor();
-			case ILevelInstanceEditorModule::ELevelInstanceEditMode::Override:
-				return FAppStyle::Get().GetSlateColor("Colors.AccentBlue").GetSpecifiedColor();
-			default:
-				break;
-			}
-		}
+				if (FLevelInstanceEditingColumn* LevelInstanceEditingColumn = DataStorage->GetColumn<FLevelInstanceEditingColumn>(DataRow))
+				{
+					switch (LevelInstanceEditingColumn->EditMode)
+					{
+					case ILevelInstanceEditorModule::ELevelInstanceEditMode::None:
+						break;
+					case ILevelInstanceEditorModule::ELevelInstanceEditMode::Original:
+						return FAppStyle::Get().GetSlateColor("Colors.AccentGreen").GetSpecifiedColor();
+					case ILevelInstanceEditorModule::ELevelInstanceEditMode::Override:
+						return FAppStyle::Get().GetSlateColor("Colors.AccentBlue").GetSpecifiedColor();
+					default:
+						break;
+					}
+				}
 
-		if (const FTedsOutlinerColumn* OutlinerColumn = DataStorage->GetColumn<FTedsOutlinerColumn>(WidgetRow))
-		{
-			if (const TSharedPtr<ISceneOutliner> Outliner = OutlinerColumn->Outliner.Pin())
-			{
-				if (const FSceneOutlinerTreeItemPtr TreeItemPtr = Helpers::GetTreeItemFromRowHandle(DataStorage, Outliner.ToSharedRef(), DataRow))
+				if (const FTedsOutlinerColumn* OutlinerColumn = DataStorage->GetColumn<FTedsOutlinerColumn>(WidgetRow))
 				{
-					if (TOptional<FLinearColor> ForegroundColor = FSceneOutlinerCommonLabelData::GetForegroundColor(Outliner, *TreeItemPtr))
+					if (const TSharedPtr<ISceneOutliner> Outliner = OutlinerColumn->Outliner.Pin())
 					{
-						return ForegroundColor.GetValue();
+						if (const FSceneOutlinerTreeItemPtr TreeItemPtr = Helpers::GetTreeItemFromRowHandle(DataStorage, Outliner.ToSharedRef(), DataRow))
+						{
+							if (TOptional<FLinearColor> ForegroundColor = FSceneOutlinerCommonLabelData::GetForegroundColor(Outliner, *TreeItemPtr))
+							{
+								return ForegroundColor.GetValue();
+							}
+						}
 					}
 				}
-			}
-		}
 
-		return FSlateColor::UseForeground();
+				return FSlateColor::UseForeground();
+			});
+		
 	}
 
 	// We are using an inheritance based model with a custom SWidget class just so we can manage lifetimes of delegates easier using BindSP
@@ -87,9 +91,13 @@ namespace UE::Editor::Outliner::TextWidget::Private
 			WidgetRow = InArgs._WidgetRow;
 			
 			using namespace UE::Editor::DataStorage;
+
 			FAttributeBinder DataRowBinder(DataRow, DataStorage);
 			FAttributeBinder WidgetRowBinder(WidgetRow, DataStorage);
 
+			const RowHandle OutlinerWidgetRowHandle = TableViewerUtils::GetTableViewerUiRow(DataStorage, WidgetRow);
+			FAttributeBinder OutlinerWidgetRowBinder(OutlinerWidgetRowHandle, DataStorage);
+
 			const bool bShouldUseMiddleEllipsis = GetDefault<UEditorStyleSettings>()->bEnableMiddleEllipsis;
 
 			SInlineEditableTextBlock::Construct(
@@ -100,15 +108,13 @@ namespace UE::Editor::Outliner::TextWidget::Private
 				.ToolTipText(DataRowBinder.BindText(&FTypedElementLabelColumn::Label))
 				.OverflowPolicy(bShouldUseMiddleEllipsis ? ETextOverflowPolicy::MiddleEllipsis : TOptional<ETextOverflowPolicy>())
 				.IsSelected(WidgetRowBinder.BindEvent(&FExternalWidgetSelectionColumn::IsSelected))
-				.ColorAndOpacity(TAttribute<FSlateColor>::CreateLambda([this]() -> FSlateColor
-					{
-						return Private::GetForegroundColor(DataStorage, DataRow, WidgetRow);
-					}))
-				);
+				.HighlightText(OutlinerWidgetRowBinder.BindData(&FHighlightTextColumn::HighlightText))
+				.ColorAndOpacity(Private::GetForegroundColor(DataStorage, DataRow, WidgetRow))
+			);
 
 			BindTableViewerRenameDelegate();
 		}
-		
+
 		virtual ~SOutlinerEditableTextWidget() override = default;
 
 	protected:
@@ -121,7 +127,7 @@ namespace UE::Editor::Outliner::TextWidget::Private
 				{
 					EnterEditingMode();
 				})
-			});
+				});
 
 			// The table viewer, if present, is the grandparent widget of our widget row (Table -> Row -> Cell)
 			const DataStorage::RowHandle ParentRow = DataStorage::TableViewerUtils::GetTableViewerUiRow(DataStorage, WidgetRow);
@@ -147,19 +153,19 @@ namespace UE::Editor::Outliner::TextWidget::Private
 			if (RenameColumn && ReferenceColumn)
 			{
 				RenameColumn->OnRenameSelection.RemoveAll(this);
-				
+
 				if (TSharedPtr<DataStorage::ITableViewer> TableViewer = ReferenceColumn->TableViewerWeak.Pin())
 				{
 					TableViewer->ForEachSelectedRow([this, RenameColumn](DataStorage::RowHandle SelectedRow)
-					{
-						if (SelectedRow == DataRow)
 						{
-							RenameColumn->OnRenameSelection.AddSPLambda(this, [this]()
-								{
-									EnterEditingMode();
-								});
-						}
-					});
+							if (SelectedRow == DataRow)
+							{
+								RenameColumn->OnRenameSelection.AddSPLambda(this, [this]()
+									{
+										EnterEditingMode();
+									});
+							}
+						});
 				}
 			}
 		}
@@ -213,7 +219,7 @@ TSharedPtr<SWidget> FOutlinerTextWidgetConstructor::CreateWidget(UE::Editor::Dat
 	TSharedPtr<SWidget> Result = SNullWidget::NullWidget;
 
 	bool bEditable = false;
-	
+
 	FMetaDataEntryView MetaDataEntryView = Arguments.FindForColumn<FTypedElementLabelColumn>(IsEditableName);
 
 	if (MetaDataEntryView.IsSet())
@@ -223,7 +229,7 @@ TSharedPtr<SWidget> FOutlinerTextWidgetConstructor::CreateWidget(UE::Editor::Dat
 			bEditable = *EditableMetaDataPtr;
 		}
 	}
-	
+
 	if (bEditable)
 	{
 		Result = CreateEditableWidget(DataStorage, DataStorageUi, TargetRow, WidgetRow, Arguments);
@@ -232,7 +238,7 @@ TSharedPtr<SWidget> FOutlinerTextWidgetConstructor::CreateWidget(UE::Editor::Dat
 	{
 		Result = CreateNonEditableWidget(DataStorage, DataStorageUi, TargetRow, WidgetRow, Arguments);
 	}
-	
+
 	return Result;
 }
 
@@ -288,11 +294,11 @@ TSharedPtr<SInlineEditableTextBlock> FOutlinerTextWidgetConstructor::CreateEdita
 	UE::Editor::DataStorage::RowHandle TargetRow,
 	UE::Editor::DataStorage::RowHandle WidgetRow) const
 {
-	
+
 	return SNew(UE::Editor::Outliner::TextWidget::Private::SOutlinerEditableTextWidget)
-			.DataStorage(DataStorage)
-			.DataRow(TargetRow)
-			.WidgetRow(WidgetRow);
+		.DataStorage(DataStorage)
+		.DataRow(TargetRow)
+		.WidgetRow(WidgetRow);
 }
 
 TSharedPtr<STextBlock> FOutlinerTextWidgetConstructor::CreateNonEditableTextBlock(
@@ -303,13 +309,14 @@ TSharedPtr<STextBlock> FOutlinerTextWidgetConstructor::CreateNonEditableTextBloc
 	using namespace UE::Editor::DataStorage;
 	FAttributeBinder TargetRowBinder(TargetRow, DataStorage);
 
+	const RowHandle OutlinerWidgetRowHandle = TableViewerUtils::GetTableViewerUiRow(DataStorage, WidgetRow);
+	FAttributeBinder OutlinerWidgetRowBinder(OutlinerWidgetRowHandle, DataStorage);
+
 	TSharedPtr<STextBlock> TextBlock = SNew(STextBlock)
 		.Text(TargetRowBinder.BindText(&FTypedElementLabelColumn::Label))
 		.ToolTipText(TargetRowBinder.BindText(&FTypedElementLabelColumn::Label))
-		.ColorAndOpacity(TAttribute<FSlateColor>::CreateLambda([DataStorage, TargetRow, WidgetRow]() -> FSlateColor
-			{
-				return UE::Editor::Outliner::TextWidget::Private::GetForegroundColor(DataStorage, TargetRow, WidgetRow);
-			}));
+		.HighlightText(OutlinerWidgetRowBinder.BindData(&FHighlightTextColumn::HighlightText))
+		.ColorAndOpacity(UE::Editor::Outliner::TextWidget::Private::GetForegroundColor(DataStorage, TargetRow, WidgetRow));
 
 	return TextBlock;
 }
diff --git a/Engine/Plugins/Experimental/EditorDataStorageFeatures/Source/TedsOutliner/Private/Widgets/TypeInfoWidget.cpp b/Engine/Plugins/Experimental/EditorDataStorageFeatures/Source/TedsOutliner/Private/Widgets/TypeInfoWidget.cpp
--- a/Engine/Plugins/Experimental/EditorDataStorageFeatures/Source/TedsOutliner/Private/Widgets/TypeInfoWidget.cpp
+++ b/Engine/Plugins/Experimental/EditorDataStorageFeatures/Source/TedsOutliner/Private/Widgets/TypeInfoWidget.cpp
@@ -7,8 +7,6 @@
 #include "Elements/Columns/TypedElementMiscColumns.h"
 #include "Elements/Columns/TypedElementSlateWidgetColumns.h"
 #include "Elements/Columns/TypedElementTypeInfoColumns.h"
-#include "Elements/Interfaces/Capabilities/TypedElementUiTextCapability.h"
-#include "Elements/Interfaces/TypedElementDataStorageUiInterface.h"
 #include "Styling/SlateIconFinder.h"
 #include "TedsTableViewerUtils.h"
 #include "Elements/Framework/TypedElementAttributeBinding.h"
diff --git a/Engine/Plugins/Experimental/EditorDataStorageFeatures/Source/TedsOutliner/Public/TedsOutlinerImpl.h b/Engine/Plugins/Experimental/EditorDataStorageFeatures/Source/TedsOutliner/Public/TedsOutlinerImpl.h
--- a/Engine/Plugins/Experimental/EditorDataStorageFeatures/Source/TedsOutliner/Public/TedsOutlinerImpl.h
+++ b/Engine/Plugins/Experimental/EditorDataStorageFeatures/Source/TedsOutliner/Public/TedsOutlinerImpl.h
@@ -347,6 +347,9 @@ protected:
 	
 	// Ticker for selection updates so we don't fire the delegate multiple times in one frame for multi select
 	FTSTicker::FDelegateHandle TickerHandle;
+
+	// Handle for OnTextFilteredChanged delegate, so we can remove it upon destruction
+	FDelegateHandle OnTextFilteredChangedHandle;
 	
 	FOnTedsOutlinerSelectionChanged OnTedsOutlinerSelectionChanged;
 
diff --git a/Engine/Source/Editor/SceneOutliner/Private/SSceneOutliner.cpp b/Engine/Source/Editor/SceneOutliner/Private/SSceneOutliner.cpp
--- a/Engine/Source/Editor/SceneOutliner/Private/SSceneOutliner.cpp
+++ b/Engine/Source/Editor/SceneOutliner/Private/SSceneOutliner.cpp
@@ -2318,6 +2318,12 @@ TAttribute<FText> SSceneOutliner::GetFilterHighlightText() const
 	}, TWeakPtr<SceneOutliner::TreeItemTextFilter>(SearchBoxFilter)));
 }
 
+SceneOutliner::TreeItemTextFilter::FChangedEvent& SSceneOutliner::OnTextFilterChanged() const
+{
+	check(SearchBoxFilter.IsValid());
+	return SearchBoxFilter->OnChanged();
+}
+
 void SSceneOutliner::SetKeyboardFocus()
 {
 	if (SupportsKeyboardFocus())
diff --git a/Engine/Source/Editor/SceneOutliner/Public/SSceneOutliner.h b/Engine/Source/Editor/SceneOutliner/Public/SSceneOutliner.h
--- a/Engine/Source/Editor/SceneOutliner/Public/SSceneOutliner.h
+++ b/Engine/Source/Editor/SceneOutliner/Public/SSceneOutliner.h
@@ -297,6 +297,9 @@ public:
 	/** @return Returns a string to use for highlighting results in the outliner list */
 	UE_API virtual TAttribute<FText> GetFilterHighlightText() const override;
 
+	/** @return Returns an Event to use for highlighting changes in the outliner list */
+	UE_API SceneOutliner::TreeItemTextFilter::FChangedEvent& OnTextFilterChanged() const;
+
 	/** Set the keyboard focus to the outliner */
 	UE_API virtual void SetKeyboardFocus() override;
 
diff --git a/Engine/Source/Runtime/TypedElementFramework/Public/Elements/Interfaces/Capabilities/TypedElementUiTextCapability.h b/Engine/Source/Runtime/TypedElementFramework/Public/Elements/Interfaces/Capabilities/TypedElementUiTextCapability.h
--- a/Engine/Source/Runtime/TypedElementFramework/Public/Elements/Interfaces/Capabilities/TypedElementUiTextCapability.h
+++ b/Engine/Source/Runtime/TypedElementFramework/Public/Elements/Interfaces/Capabilities/TypedElementUiTextCapability.h
@@ -9,7 +9,7 @@
 /**
  * Interface to provide access to widgets that support working with text.
  */
-class ITypedElementUiTextCapability : public ITypedElementUiCapability
+class UE_DEPRECATED(5.8, "Please use FTedsOutlinerColumn to get the Outlier ptr and use its GetFilterHighlightText() function") ITypedElementUiTextCapability : public ITypedElementUiCapability
 {
 public:
 	SLATE_METADATA_TYPE(ITypedElementUiTextCapability, ITypedElementUiCapability)
@@ -22,7 +22,7 @@ public:
 };
 
 template<typename WidgetType>
-class TTypedElementUiTextCapability : public ITypedElementUiTextCapability
+class UE_DEPRECATED(5.8, "Please use FTedsOutlinerColumn to get the Outlier ptr and use its GetFilterHighlightText() function") TTypedElementUiTextCapability : public ITypedElementUiTextCapability
 {
 public:
 	explicit TTypedElementUiTextCapability(WidgetType& InWidget) : Widget(InWidget){}
