[TedsOperations] Previews of WidgetDropHandlers are now deleted through TEv1 instead of DeletionOperations.

#Synced-CL 50191948

diff --git a/Engine/Plugins/Experimental/EditorDataStorageFeatures/Source/TedsOperations/Private/DragAndDrop/Widgets/PreviewWidgetDropHandler.cpp b/Engine/Plugins/Experimental/EditorDataStorageFeatures/Source/TedsOperations/Private/DragAndDrop/Widgets/PreviewWidgetDropHandler.cpp
--- a/Engine/Plugins/Experimental/EditorDataStorageFeatures/Source/TedsOperations/Private/DragAndDrop/Widgets/PreviewWidgetDropHandler.cpp
+++ b/Engine/Plugins/Experimental/EditorDataStorageFeatures/Source/TedsOperations/Private/DragAndDrop/Widgets/PreviewWidgetDropHandler.cpp
@@ -6,18 +6,25 @@
 #include "Deletion/DeletionOperationSystem.h"
 #include "DragAndDrop/DropOperationInput.h"
 #include "DragAndDrop/DropOperationSystem.h"
+#include "Elements/Framework/TypedElementRegistry.h"
 #include "Elements/Interfaces/TypedElementDataStorageInterface.h"
+#include "Elements/Interfaces/TypedElementWorldInterface.h"
 #include "ISceneOutliner.h"
 #include "LevelEditor.h"
+#include "Algo/AllOf.h"
 #include "Modules/ModuleManager.h"
 
 #define LOCTEXT_NAMESPACE "PreviewWidgetDropHandler"
 
+DEFINE_LOG_CATEGORY_STATIC(LogPreviewWidgetDropHandler, Log, All);
+
 namespace UE::Editor
 {
 	
 namespace PreviewWidgetDropHandler_Private
 {
+
+static bool bUseTypedElementHandlesForDeletion = true;
 	
 static void UpdateOutliner()
 {
@@ -158,6 +165,15 @@ void FPreviewWidgetDropHandler::CreatePreviews(DataStorage::ICoreProvider& Stora
 	{
 		Storage.RemoveColumn<FResultColumn>(InputRow);
 		Storage.RemoveColumn<FDropPreviewTag>(InputRow);
+
+		if (PreviewWidgetDropHandler_Private::bUseTypedElementHandlesForDeletion)
+		{
+			if (const FDropTypedElementResultColumn* Column = Storage.GetColumn<FDropTypedElementResultColumn>(InputRow))
+			{
+				PreviewElements.Append(Column->Elements);
+			}
+		}
+		Storage.RemoveColumn<FDropTypedElementResultColumn>(InputRow);
 	}
 
 	// Create rows that hold our created previews for updates.
@@ -188,25 +204,54 @@ void FPreviewWidgetDropHandler::RemovePreviews(DataStorage::ICoreProvider& Stora
 		return;
 	}
 
-	UDeletionOperationSystem* DeletionSystem = Storage.FindFactory<UDeletionOperationSystem>();
-	if (!ensureMsgf(DeletionSystem, TEXT("Deletion system not found. Unable to delete previews!")))
+	bool bAllDeleted = true;
+	// Currently, TEDS rows are not updated immediately. This means, if the preview is removed the same tick it was created, the created rows will not
+	// yet have enough data for the deletion system to remove the placed objects from the world again.
+	// Thus, for now, we delete the previews through the TEv1 system.
+	if (PreviewWidgetDropHandler_Private::bUseTypedElementHandlesForDeletion)
 	{
-		return;
+		TNotNull<UTypedElementRegistry*> Registry = UTypedElementRegistry::GetInstance();
+
+		FTypedElementDeletionOptions Options;
+		Options.SetVerifyDeletionCanHappen(false);
+		Options.SetWarnAboutReferences(false);
+		Options.SetWarnAboutSoftReferences(false);
+		for (const FTypedElementHandle& ElementHandle : PreviewElements)
+		{
+			TTypedElement<ITypedElementWorldInterface> WorldElement = Registry->GetElement<ITypedElementWorldInterface>(ElementHandle);
+			bAllDeleted &= WorldElement && WorldElement.DeleteElement(WorldElement.GetOwnerWorld(), SelectionSet.Get(), Options);
+		}
 	}
+	else
+	{
+		UDeletionOperationSystem* DeletionSystem = Storage.FindFactory<UDeletionOperationSystem>();
+		if (!ensureMsgf(DeletionSystem, TEXT("Deletion system not found. Unable to delete previews!")))
+		{
+			return;
+		}
+		
+		// Make sure the previews gets deleted.
+		for (RowHandle PreviewInputRow : PreviewInputRows.GetRows())
+		{
+			Storage.AddColumn<FDeletionForceTag>(PreviewInputRow);
+		}
+		
+		// Delete our created previews.
+		(void)DeletionSystem->Apply(PreviewInputRows.GetRows());
 
-	// Make sure the previews gets deleted.
-	for (RowHandle PreviewInputRow : PreviewInputRows.GetRows())
+		bAllDeleted = Algo::AllOf(PreviewInputRows.GetRows(), [&Storage](RowHandle Row) { return Storage.HasColumns<FResultColumn>(Row); });
+	}
+
+	if (!bAllDeleted)
 	{
-		Storage.AddColumn<FDeletionForceTag>(PreviewInputRow);	
+		UE_LOG(LogPreviewWidgetDropHandler, Error, TEXT("Deletion of previews failed!"));
 	}
-	
-	// Delete our created previews.
-	(void)DeletionSystem->Apply(PreviewInputRows.GetRows());
 
 	// Free our preview input rows.
 	DropSystem.RemoveInputRows(PreviewInputRows.GetRows());
 	PreviewInputRows.Empty();
-
+	PreviewElements.Empty();
+	
 	PreviewWidgetDropHandler_Private::UpdateOutliner();
 }
 
diff --git a/Engine/Plugins/Experimental/EditorDataStorageFeatures/Source/TedsOperations/Public/DragAndDrop/DropOperationInput.h b/Engine/Plugins/Experimental/EditorDataStorageFeatures/Source/TedsOperations/Public/DragAndDrop/DropOperationInput.h
--- a/Engine/Plugins/Experimental/EditorDataStorageFeatures/Source/TedsOperations/Public/DragAndDrop/DropOperationInput.h
+++ b/Engine/Plugins/Experimental/EditorDataStorageFeatures/Source/TedsOperations/Public/DragAndDrop/DropOperationInput.h
@@ -3,6 +3,7 @@
 #pragma once
 
 #include "DataStorage/CommonTypes.h"
+#include "Elements/Framework/TypedElementHandle.h"
 #include "TedsOperationInput.h"
 
 #include "DropOperationInput.generated.h"
@@ -45,6 +46,14 @@ struct FDropOperationPreviewTag final : public FEditorDataStorageTag
 	GENERATED_BODY()
 };
 
+USTRUCT(meta = (DisplayName = "Drop Operation Typed Element Result"))
+struct FDropOperationTypedElementResultColumn final : public FEditorDataStorageColumn
+{
+	GENERATED_BODY()
+
+	TArray<FTypedElementHandle> Elements;
+};
+
 namespace UE::Editor::DataStorage::Operations
 {
 /** Holds the target row of the drop operation (e.g. a level). */	
@@ -57,6 +66,8 @@ using FDropNameColumn = FDropOperationNameColumn;
 using FDropTransformColumn = FDropOperationTransformColumn;
 /** Holds the folder path that should be used by the operation. */
 using FDropFolderColumn = FDropOperationFolderColumn;
+/** Holds the TEv1 handles created by the operation. */
+using FDropTypedElementResultColumn = FDropOperationTypedElementResultColumn;
 
 namespace Utilities
 {
diff --git a/Engine/Plugins/Experimental/EditorDataStorageFeatures/Source/TedsOperations/Public/DragAndDrop/Widgets/PreviewWidgetDropHandler.h b/Engine/Plugins/Experimental/EditorDataStorageFeatures/Source/TedsOperations/Public/DragAndDrop/Widgets/PreviewWidgetDropHandler.h
--- a/Engine/Plugins/Experimental/EditorDataStorageFeatures/Source/TedsOperations/Public/DragAndDrop/Widgets/PreviewWidgetDropHandler.h
+++ b/Engine/Plugins/Experimental/EditorDataStorageFeatures/Source/TedsOperations/Public/DragAndDrop/Widgets/PreviewWidgetDropHandler.h
@@ -6,6 +6,8 @@
 
 #define UE_API TEDSOPERATIONS_API
 
+struct FTypedElementHandle;
+
 namespace UE::Editor
 {
 
@@ -39,6 +41,8 @@ protected:
 	
 protected:
 	DataStorage::FRowHandleArray PreviewInputRows;
+
+	TArray<FTypedElementHandle> PreviewElements;
 };
 
 }
