[TEDS Package] Upgrade the package sync to delete and remap rows on package change

#Synced-CL 50303485

diff --git a/Engine/Plugins/Experimental/EditorDataStorageFeatures/Source/TedsRevisionControl/Private/Queries/ObjectPackagePathToColumnQueries.cpp b/Engine/Plugins/Experimental/EditorDataStorageFeatures/Source/TedsRevisionControl/Private/Queries/ObjectPackagePathToColumnQueries.cpp
--- a/Engine/Plugins/Experimental/EditorDataStorageFeatures/Source/TedsRevisionControl/Private/Queries/ObjectPackagePathToColumnQueries.cpp
+++ b/Engine/Plugins/Experimental/EditorDataStorageFeatures/Source/TedsRevisionControl/Private/Queries/ObjectPackagePathToColumnQueries.cpp
@@ -26,7 +26,7 @@ namespace UE::Editor::RevisionControl::Private
     	TEXT("Automatically query package information and fill information into TEDS")
     );
 
-	static void ResolvePackageReference(IQueryContext& Context, const UPackage* Package, const UObject* Object, RowHandle Row, RowHandle PackageRow)
+	static void ResolvePackageReference(IQueryContext& Context, UPackage* Package, const UObject* Object, RowHandle Row, RowHandle PackageRow)
 	{
 		FTypedElementPackageReference PackageReference;
 		PackageReference.Row = PackageRow;
@@ -41,9 +41,10 @@ namespace UE::Editor::RevisionControl::Private
 		Context.AddColumn(PackageRow, MoveTemp(PathColumn));
 		Context.AddColumn(PackageRow, MoveTemp(LoadedPathColumn));
 
-		if (const UObject* PackageInstance = Package->FindAssetInPackage())
+		// If the Package is Transient we can skip finding the primary asset as that will loop over all assets in the package to find nothing
+		if (Package != GetTransientPackage() && !Package->HasAnyFlags(RF_Transient))
 		{
-			if (Object == PackageInstance)
+			if (const UObject* PackageInstance = Package->FindAssetInPackage(); Object == PackageInstance)
 			{
 				Context.AddColumns<FTedsPrimaryPackageObjectTag>(Row);
 			}
@@ -63,7 +64,7 @@ namespace UE::Editor::RevisionControl::Private
 				{
 					PackageObjectCDO = Object->GetClass()->GetDefaultObject();
 				}
-				
+			
 				if (Object == PackageObjectCDO)
 				{
 					Context.AddColumns<FTedsPrimaryPackageObjectTag>(Row);
@@ -71,6 +72,11 @@ namespace UE::Editor::RevisionControl::Private
 			}
 		}
 
+		if (Package->IsDirty())
+		{
+			Context.AddColumns<FTedsPackageDirtyTag>(Row);
+		}
+
 		Context.SetParentRow(Row, PackageRow);
 	};
 } // namespace UE::Editor::RevisionControl::Private
@@ -102,6 +108,7 @@ void UTypedElementUObjectPackagePathFactory::RegisterQueries(UE::Editor::DataSto
 			else
 			{
 				DataStorage.UnregisterQuery(TryAddPackageRef);
+				DataStorage.UnregisterQuery(TryReplacePackageRef);
 			}
 		}
 	);
@@ -122,7 +129,8 @@ void UTypedElementUObjectPackagePathFactory::RegisterQueries(UE::Editor::DataSto
 				{
 					return;
 				}
-				const UPackage* Package = ObjectInstance->GetPackage();
+				UPackage* Package = ObjectInstance->GetPackage();
+				
 				Context.RemoveColumns(Row, { FTypedElementPackageUnresolvedReference::StaticStruct() });
 
 				UE::Editor::RevisionControl::Private::ResolvePackageReference(Context, Package, ObjectInstance, Row, PackageRow);
@@ -148,94 +156,76 @@ void UTypedElementUObjectPackagePathFactory::RegisterTryAddPackageRef(UE::Editor
 			if (CoreProvider)
 			{
 				// Check in case this row was mapped before this command was pushed 
-				RowHandle PackageRow = CoreProvider->LookupMappedRow(UE::Editor::RevisionControl::MappingDomain, FMapKey(PackagePathname));
-				if (!CoreProvider->IsRowAvailable(PackageRow))
+				if (const RowHandle PackageRow = CoreProvider->LookupMappedRow(UE::Editor::RevisionControl::MappingDomain, FMapKey(PackagePathname)); !CoreProvider->IsRowAvailable(PackageRow))
 				{
-					const RowHandle RevisionControlTable = CoreProvider->FindTable(FName("Editor_RevisionControlTable"));
-					if (RevisionControlTable != InvalidRowHandle)
+					if (const TableHandle RevisionControlTable = CoreProvider->FindTable(FName("Editor_RevisionControlTable")); RevisionControlTable != InvalidRowHandle)
 					{
-						RowHandle Row = CoreProvider->AddRow(RevisionControlTable);
+						const RowHandle Row = CoreProvider->AddRow(RevisionControlTable);
 						CoreProvider->MapRow(UE::Editor::RevisionControl::MappingDomain, FMapKey(PackagePathname), Row);
-						if (bIsMapPackage)
-						{
-							// Also map to the .uasset path
-							CoreProvider->MapRow(
-								UE::Editor::RevisionControl::MappingDomain,
-								FMapKey(FPaths::ChangeExtension(PackagePathname, FPackageName::GetAssetPackageExtension())),
-								Row);
-						}
 					}
 				}
 			}
 		}
 		FString PackagePathname;
 		ICoreProvider* CoreProvider = nullptr;
-		bool bIsMapPackage = false;
+	};
+
+	auto PackageSyncCallback = [](IQueryContext& Context, RowHandle Row, const FTypedElementUObjectColumn& Object)
+	{
+		if (const UObject* ObjectInstance = Object.Object.Get(); IsValid(ObjectInstance))
+		{
+			UPackage* Target = ObjectInstance->GetPackage();
+			
+			if (ObjectInstance->HasAnyFlags(RF_Transient | RF_BeginDestroyed) || Target->HasAnyPackageFlags(PKG_Cooked))
+			{
+				// Don't monitor transient objects since they will never exist on disk. And don't monitor objects in cooked packages since
+				// they are (likely) not under the source control root and even if they happen to be, we didn't load the loose asset on disk.
+				return;
+			}
+
+			FString Path;
+			Target->GetPathName(nullptr, Path);
+			if (FString PackageFilename; FPackageName::TryConvertLongPackageNameToFilename(Path, PackageFilename))
+			{
+				FPaths::NormalizeFilename(PackageFilename);
+				FString FullPackageFilename = FPaths::ConvertRelativePathToFull(PackageFilename);
+				const RowHandle PackageRow = Context.LookupMappedRow(UE::Editor::RevisionControl::MappingDomain, FMapKeyView(FullPackageFilename));
+				if (Context.IsRowAvailable(PackageRow))
+				{
+					UE::Editor::RevisionControl::Private::ResolvePackageReference(Context, Target, ObjectInstance, Row, PackageRow);
+				}
+				else
+				{
+					FTypedElementPackageUnresolvedReference UnresolvedPackageReference;
+					UnresolvedPackageReference.PathOnDisk = FullPackageFilename;
+					Context.AddColumn(Row, MoveTemp(UnresolvedPackageReference));
+
+					FRegisterPackageRowCommand Command
+					{
+						.PackagePathname = FullPackageFilename,
+						.CoreProvider = GetMutableDataStorageFeature<ICoreProvider>(StorageFeatureName),
+					};
+					Context.PushCommand(Command);
+				}
+			}
+		}
 	};
 	
 	TryAddPackageRef = DataStorage.RegisterQuery(
 		Select(
-			TEXT("Sync UObject package info to columns"),
+			TEXT("Sync UObject package info to new rows"),
 			FObserver::OnAdd<FTypedElementUObjectColumn>()
 				.SetExecutionMode(EExecutionMode::GameThread),
-			[](IQueryContext& Context, RowHandle Row, const FTypedElementUObjectColumn& Object)
-			{
-				if (const UObject* ObjectInstance = Object.Object.Get(); ObjectInstance != nullptr)
-				{
-					if (ObjectInstance->HasAnyFlags(RF_Transient))
-					{
-						// Don't monitor transient objects since they will never exist on disk.
-						return;
-					}
-					
-					const UPackage* Target = ObjectInstance->GetPackage();
-					if (Target->HasAnyPackageFlags(PKG_Cooked))
-					{
-						// Don't monitor objects in cooked packages since they are (likely) not under the source control root,
-						// and even if they happen to be, we didn't load the loose asset on disk.
-						return;
-					}
-
-					// Certain Package utilities don't get the .umap and will return the .uasset instead so we want to map this row to both of those
-					FString PackageExtension = FPackageName::GetAssetPackageExtension();
-
-					bool bIsMapPackage = false;
-					if (Target->ContainsMap())
-					{
-						bIsMapPackage = true;
-						PackageExtension = FPackageName::GetMapPackageExtension();
-					}
+			PackageSyncCallback)
+		.AccessesHierarchy(PackageHierarchyName)
+		.Compile());
 
-					FString Path;
-					Target->GetPathName(nullptr, Path);
-					if (FString PackageFilename; FPackageName::TryConvertLongPackageNameToFilename(Path, PackageFilename, PackageExtension))
-					{
-						FPaths::NormalizeFilename(PackageFilename);
-						FString FullPackageFilename = FPaths::ConvertRelativePathToFull(PackageFilename);
-						FMapKeyView Key = FMapKeyView(FullPackageFilename);
-						RowHandle PackageRow = Context.LookupMappedRow(UE::Editor::RevisionControl::MappingDomain, Key);
-						if (Context.IsRowAvailable(PackageRow))
-						{
-							const UPackage* Package = ObjectInstance->GetPackage();
-							UE::Editor::RevisionControl::Private::ResolvePackageReference(Context, Package, ObjectInstance, Row, PackageRow);
-						}
-						else
-						{
-							FTypedElementPackageUnresolvedReference UnresolvedPackageReference;
-							UnresolvedPackageReference.PathOnDisk = FullPackageFilename;
-							Context.AddColumn(Row, MoveTemp(UnresolvedPackageReference));
-
-							FRegisterPackageRowCommand Command
-							{
-								.PackagePathname = FullPackageFilename,
-								.CoreProvider = GetMutableDataStorageFeature<ICoreProvider>(StorageFeatureName),
-								.bIsMapPackage = bIsMapPackage
-							};
-							Context.PushCommand(Command);
-						}
-					}
-				}
-			})
+	TryReplacePackageRef = DataStorage.RegisterQuery(
+		Select(
+			TEXT("Resync UObject package info to columns on package removal"),
+			FObserver::OnRemove<FTypedElementPackageReference>()
+				.SetExecutionMode(EExecutionMode::GameThread),
+			PackageSyncCallback)
 		.AccessesHierarchy(PackageHierarchyName)
 		.Compile());
 }
diff --git a/Engine/Plugins/Experimental/EditorDataStorageFeatures/Source/TedsRevisionControl/Private/Queries/ObjectPackagePathToColumnQueries.h b/Engine/Plugins/Experimental/EditorDataStorageFeatures/Source/TedsRevisionControl/Private/Queries/ObjectPackagePathToColumnQueries.h
--- a/Engine/Plugins/Experimental/EditorDataStorageFeatures/Source/TedsRevisionControl/Private/Queries/ObjectPackagePathToColumnQueries.h
+++ b/Engine/Plugins/Experimental/EditorDataStorageFeatures/Source/TedsRevisionControl/Private/Queries/ObjectPackagePathToColumnQueries.h
@@ -26,4 +26,5 @@ public:
 private:
 	void RegisterTryAddPackageRef(UE::Editor::DataStorage::ICoreProvider& DataStorage);
 	UE::Editor::DataStorage::QueryHandle TryAddPackageRef = UE::Editor::DataStorage::InvalidQueryHandle;
+	UE::Editor::DataStorage::QueryHandle TryReplacePackageRef = UE::Editor::DataStorage::InvalidQueryHandle;
 };
diff --git a/Engine/Source/Developer/UnsavedAssetsTracker/Source/Private/UnsavedAssetsTracker.cpp b/Engine/Source/Developer/UnsavedAssetsTracker/Source/Private/UnsavedAssetsTracker.cpp
--- a/Engine/Source/Developer/UnsavedAssetsTracker/Source/Private/UnsavedAssetsTracker.cpp
+++ b/Engine/Source/Developer/UnsavedAssetsTracker/Source/Private/UnsavedAssetsTracker.cpp
@@ -26,6 +26,7 @@
 #include "UObject/UObjectGlobals.h"
 
 #include "DataStorage/Features.h"
+#include "Elements/Columns/TypedElementCompatibilityColumns.h"
 #include "Elements/Columns/TypedElementPackageColumns.h"
 #include "Elements/Columns/TypedElementRevisionControlColumns.h"
 #include "Elements/Interfaces/TypedElementDataStorageInterface.h"
@@ -36,6 +37,12 @@ DEFINE_LOG_CATEGORY_STATIC(LogUnsavedAssetsTracker, Log, All);
 
 namespace
 {
+static bool bSyncUnsavedAssetInfoToTEDS = true;
+FAutoConsoleVariableRef CVarSyncUnsavedInformation(
+	TEXT("TEDS.RevisionControl.SyncUnsavedInformation"),
+	bSyncUnsavedAssetInfoToTEDS,
+	TEXT("Automatically sync information about unsaved assets to TEDS")
+);
 
 FString GetHumanFriendlyAssetName(const UPackage* Package)
 {
@@ -329,6 +336,7 @@ void FUnsavedAssetsTracker::OnPackageDeleted(UPackage* Package)
 	{
 		bSyncWithDirtyPackageList = true;
 	}
+	RemoveTedsPackage(Package->GetPathName());
 }
 
 void FUnsavedAssetsTracker::OnPostGarbageCollect()
@@ -345,7 +353,7 @@ void FUnsavedAssetsTracker::SyncWithDirtyPackageList()
 {
 	TRACE_CPUPROFILER_EVENT_SCOPE(FUnsavedAssetsTracker::SyncWithDirtyPackageList);
 
-	// The the list of dirty packages tracked by the engine (considered source of truth)
+	// The list of dirty packages tracked by the engine (considered source of truth)
 	TArray<UPackage*> DirtyPackages;
 	FEditorFileUtils::GetDirtyPackages(DirtyPackages);
 
@@ -439,7 +447,7 @@ void FUnsavedAssetsTracker::StartTrackingDirtyPackage(const UPackage* Package)
 		AbsPackageFilePathname,
 		FSourceControlFileStatusMonitor::FOnSourceControlFileStatus::CreateSP(this, &FUnsavedAssetsTracker::OnSourceControlFileStatusUpdate));
 
-	UpdateTEDSPackage(PackagePathname, true);
+	UpdateTedsPackage(PackagePathname, true);
 
 	FUnsavedAssetsTrackerModule::Get().OnUnsavedAssetAdded.Broadcast(AbsPackageFilePathname);
 
@@ -458,7 +466,7 @@ void FUnsavedAssetsTracker::StopTrackingDirtyPackage(const FString& PackagePathn
 		WarningPackagePathnames.Remove(PackagePathname);
 		UnsavedAbsFilePathames.Remove(AbsPackageFilePathname);
 
-		UpdateTEDSPackage(PackagePathname, false);
+		UpdateTedsPackage(PackagePathname, false);
 
 		ISourceControlModule::Get().GetSourceControlFileStatusMonitor().StopMonitoringFile(reinterpret_cast<uintptr_t>(this), AbsPackageFilePathname);
 		FUnsavedAssetsTrackerModule::Get().OnUnsavedAssetRemoved.Broadcast(AbsPackageFilePathname);
@@ -472,36 +480,103 @@ void FUnsavedAssetsTracker::StopTrackingDirtyPackage(const FString& PackagePathn
 	}
 }
 
-void FUnsavedAssetsTracker::UpdateTEDSPackage(const FString& PackagePathname, const bool bIsDirty) const
+void FUnsavedAssetsTracker::UpdateTedsPackage(const FString& PackagePathname, const bool bIsDirty)
 {
 	using namespace UE::Editor::DataStorage;
 
-	if (ICoreProvider* Storage = GetMutableDataStorageFeature<ICoreProvider>(StorageFeatureName))
+	if (bSyncUnsavedAssetInfoToTEDS)
 	{
-		FString FullPackagePath;
-		if (FPackageName::TryConvertLongPackageNameToFilename(PackagePathname, FullPackagePath, FPackageName::GetAssetPackageExtension()))
+		if (ICoreProvider* Storage = GetMutableDataStorageFeature<ICoreProvider>(StorageFeatureName))
 		{
-			FullPackagePath = FPaths::ConvertRelativePathToFull(FullPackagePath);
-			const RowHandle PackageRow = Storage->LookupMappedRow(UE::Editor::RevisionControl::MappingDomain, FMapKey(FullPackagePath));
-			if (Storage->IsRowAvailable(PackageRow))
+			if (FString FullPackagePath; FPackageName::TryConvertLongPackageNameToFilename(PackagePathname, FullPackagePath))
 			{
-				TArray<RowHandle> ChildrenRowHandles;
-				auto GetChildRows = [&ChildrenRowHandles](const ICoreProvider& Context, RowHandle, RowHandle TargetRow)
+				FullPackagePath = FPaths::ConvertRelativePathToFull(FullPackagePath);
+				if (const RowHandle PackageRow = Storage->LookupMappedRow(UE::Editor::RevisionControl::MappingDomain, FMapKey(FullPackagePath)); Storage->IsRowAvailable(PackageRow))
 				{
-					ChildrenRowHandles.Add(TargetRow);
-				};
-				Storage->WalkDepthFirst(Storage->FindHierarchyByName(PackageHierarchyName), PackageRow, GetChildRows);
-	
-				for (RowHandle ChildRow : ChildrenRowHandles)
-				{
-					if (bIsDirty)
+					TArray<RowHandle> ChildrenRowHandles; // Rows that remain in the package, but need to sync IsDirty
+					TArray<RowHandle> RemapChildrenRowHandles; // Rows that are no longer in the package and need to be resynced in the processor
+					auto GetChildRows = [&ChildrenRowHandles, &RemapChildrenRowHandles, PackagePathname](const ICoreProvider& Context, RowHandle, RowHandle TargetRow)
 					{
-						Storage->AddColumns<FTedsPackageDirtyTag>(ChildRow);
+						if (const FTypedElementUObjectColumn* ObjectColumn = Context.GetColumn<FTypedElementUObjectColumn>(TargetRow))
+						{
+							if (const UObject* ObjectInstance = ObjectColumn->Object.Get(); ObjectInstance != nullptr)
+							{
+								// If the package of the object is null or isn't the same, we need to resync this object row to its current package
+								if (const UPackage* Target = ObjectInstance->GetPackage())
+								{
+									FString Path;
+									Target->GetPathName(nullptr, Path);
+									if (!Path.Equals(PackagePathname))
+									{
+										RemapChildrenRowHandles.Add(TargetRow);
+										return;
+									}
+								}
+								else
+								{
+									RemapChildrenRowHandles.Add(TargetRow);
+									return;
+								}
+							}
+						}
+						ChildrenRowHandles.Add(TargetRow);
+					};
+					Storage->WalkDepthFirst(Storage->FindHierarchyByName(PackageHierarchyName), PackageRow, GetChildRows);
+
+					// Remove package columns from any rows that don't belong to this package anymore so they can resync
+					Storage->BatchAddRemoveColumns(
+						TConstArrayView<RowHandle>(RemapChildrenRowHandles),
+						{},
+						{FTypedElementPackageReference::StaticStruct(), FTedsPackageDirtyTag::StaticStruct(), FTedsPrimaryPackageObjectTag::StaticStruct()});
+
+					// Any rows still in the package should update to the correct IsDirty state
+					for (const RowHandle ChildRow : ChildrenRowHandles)
+					{
+						if (bIsDirty)
+						{
+							Storage->AddColumns<FTedsPackageDirtyTag>(ChildRow);
+						}
+						else
+						{
+							Storage->RemoveColumn<FTedsPackageDirtyTag>(ChildRow);
+						}
 					}
-					else
+				}
+			}
+		}
+	}
+}
+
+void FUnsavedAssetsTracker::RemoveTedsPackage(const FString& PackagePathname)
+{
+	using namespace UE::Editor::DataStorage;
+    
+	if (bSyncUnsavedAssetInfoToTEDS)
+	{
+		if (ICoreProvider* Storage = GetMutableDataStorageFeature<ICoreProvider>(StorageFeatureName))
+		{
+			if (FString FullPackagePath; FPackageName::TryConvertLongPackageNameToFilename(PackagePathname, FullPackagePath))
+			{
+				FullPackagePath = FPaths::ConvertRelativePathToFull(FullPackagePath);
+				if (const RowHandle PackageRow = Storage->LookupMappedRow(UE::Editor::RevisionControl::MappingDomain, FMapKey(FullPackagePath)); Storage->IsRowAvailable(PackageRow))
+				{
+					TArray<RowHandle> ChildRowHandles;
+					auto GetChildRows = [&ChildRowHandles](const ICoreProvider&, RowHandle, RowHandle TargetRow)
+						{
+							ChildRowHandles.Add(TargetRow);
+						};
+					Storage->WalkDepthFirst(Storage->FindHierarchyByName(PackageHierarchyName), PackageRow, GetChildRows);
+
+					for (const RowHandle ChildRow : ChildRowHandles)
 					{
-						Storage->RemoveColumn<FTedsPackageDirtyTag>(ChildRow);
+						Storage->RemoveColumns(ChildRow, {FTypedElementPackageReference::StaticStruct(), FTedsPackageDirtyTag::StaticStruct(), FTedsPrimaryPackageObjectTag::StaticStruct()});
 					}
+					Storage->BatchAddRemoveColumns(
+						TConstArrayView<RowHandle>(ChildRowHandles),
+						{},
+						{FTypedElementPackageReference::StaticStruct(), FTedsPackageDirtyTag::StaticStruct(), FTedsPrimaryPackageObjectTag::StaticStruct()});
+					
+					Storage->RemoveRow(PackageRow);
 				}
 			}
 		}
diff --git a/Engine/Source/Developer/UnsavedAssetsTracker/Source/Private/UnsavedAssetsTracker.h b/Engine/Source/Developer/UnsavedAssetsTracker/Source/Private/UnsavedAssetsTracker.h
--- a/Engine/Source/Developer/UnsavedAssetsTracker/Source/Private/UnsavedAssetsTracker.h
+++ b/Engine/Source/Developer/UnsavedAssetsTracker/Source/Private/UnsavedAssetsTracker.h
@@ -138,7 +138,10 @@ private:
 	bool ShouldTrackDirtyPackage(const UPackage* Package);
 	
 	/** Adds the Dirty tag to the Package and its contained object rows */
-	void UpdateTEDSPackage(const FString& PackagePathname, const bool bIsDirty) const;
+	static void UpdateTedsPackage(const FString& PackagePathname, const bool bIsDirty);
+	
+	/** Removes the Package row from TEDS */
+	static void RemoveTedsPackage(const FString& PackagePathname);
 
 private:
 	/** Supports multithreaded Dirty notifications. */
