TEDSActorComponentFactory: Some minor follow up fixes

#Synced-CL 50304502

diff --git a/Engine/Plugins/Experimental/EditorDataStorageFeatures/Source/TedsActorCompatibility/Private/Compatibility/TedsActorComponentFactory.cpp b/Engine/Plugins/Experimental/EditorDataStorageFeatures/Source/TedsActorCompatibility/Private/Compatibility/TedsActorComponentFactory.cpp
--- a/Engine/Plugins/Experimental/EditorDataStorageFeatures/Source/TedsActorCompatibility/Private/Compatibility/TedsActorComponentFactory.cpp
+++ b/Engine/Plugins/Experimental/EditorDataStorageFeatures/Source/TedsActorCompatibility/Private/Compatibility/TedsActorComponentFactory.cpp
@@ -27,17 +27,9 @@ void UTedsActorComponentFactory::RegisterTables(UE::Editor::DataStorage::ICorePr
 	
 	if (bEnableActorComponentCompatibility)
 	{
-		TableHandle ActorComponentTable = DataStorage.RegisterTable(
-					{
-						FTypedElementUObjectColumn::StaticStruct(),
-						FTypedElementUObjectIdColumn::StaticStruct(),
-						FTypedElementClassTypeInfoColumn::StaticStruct(),
-						FActorComponentTypeTag::StaticStruct(),
-						FTypedElementLabelColumn::StaticStruct(),
-						FTypedElementLabelHashColumn::StaticStruct(),
-						FTypedElementSyncFromWorldTag::StaticStruct()
-					},
-					FName("Editor_ActorComponentTable"));
+		TableHandle ActorComponentTable = DataStorage.RegisterTable<FActorComponentTypeTag, FTypedElementLabelColumn, FTypedElementLabelHashColumn>(
+											DataStorage.FindTable(FName(TEXT("Editor_StandardUObjectTable"))),
+											FName("Editor_ActorComponentTable"));
 		
 		if (ICompatibilityProvider* CompatibilityProvider = GetMutableDataStorageFeature<ICompatibilityProvider>(CompatibilityFeatureName))
 		{
@@ -84,14 +76,15 @@ void UTedsActorComponentFactory::RegisterQueries(UE::Editor::DataStorage::ICoreP
 				{
 					return;
 				}
-				ICompatibilityProvider* CompatibilityProvider = GetMutableDataStorageFeature<ICompatibilityProvider>(CompatibilityFeatureName);
-
-				for (TWeakObjectPtr<UActorComponent> Component : ComponentsToRegister)
+				if (ICompatibilityProvider* CompatibilityProvider = GetMutableDataStorageFeature<ICompatibilityProvider>(CompatibilityFeatureName))
 				{
-					UActorComponent* NewComponent = Component.Get();
-					if (NewComponent)
+					for (TWeakObjectPtr<UActorComponent> Component : ComponentsToRegister)
 					{
-						CompatibilityProvider->AddCompatibleObject(NewComponent);
+						UActorComponent* NewComponent = Component.Get();
+						if (NewComponent)
+						{
+							CompatibilityProvider->AddCompatibleObject(NewComponent);
+						}
 					}
 				}
 			};
@@ -126,8 +119,7 @@ void UTedsActorComponentFactory::RegisterQueries(UE::Editor::DataStorage::ICoreP
 					}
 				})
 			.Where()
-				.All<FTypedElementActorTag>()
-				.Any<FTypedElementSyncFromWorldTag>()
+				.All<FTypedElementActorTag, FTypedElementSyncFromWorldTag>()
 			.Compile());
 
 		// 2. For components that aren't a child, which means they don't have an associated owner, set up that owner
