Fixed ULevel not being registered as DataStorage compatibility object in some cases.

#Synced-CL 50349448

diff --git a/Engine/Plugins/Experimental/EditorDataStorageFeatures/Source/TedsEditorCompatibility/Private/Level/TedsLevelFactory.cpp b/Engine/Plugins/Experimental/EditorDataStorageFeatures/Source/TedsEditorCompatibility/Private/Level/TedsLevelFactory.cpp
--- a/Engine/Plugins/Experimental/EditorDataStorageFeatures/Source/TedsEditorCompatibility/Private/Level/TedsLevelFactory.cpp
+++ b/Engine/Plugins/Experimental/EditorDataStorageFeatures/Source/TedsEditorCompatibility/Private/Level/TedsLevelFactory.cpp
@@ -27,22 +27,10 @@ void UTedsLevelFactory::RegisterTables(UE::Editor::DataStorage::ICoreProvider& D
 	
 	TableHandle LevelTable = DataStorage.RegisterTable<FLevelTag, FTypedElementLabelColumn, FTypedElementIconOverrideColumn>(Levels::LevelTableName);
 
-	DataStorageCompat = GetMutableDataStorageFeature<ICompatibilityProvider>(CompatibilityFeatureName);
-	checkf(DataStorageCompat, TEXT("TEDS Factory cannot be init before the data storage compatibility feature is available"));
+	ICompatibilityProvider* Compatibility = GetMutableDataStorageFeature<ICompatibilityProvider>(CompatibilityFeatureName);
+	checkf(Compatibility, TEXT("TEDS Factory cannot be init before the data storage compatibility feature is available"));
 	
-	DataStorageCompat->RegisterTypeTableAssociation(ULevel::StaticClass(), LevelTable);
-}
-
-void UTedsLevelFactory::PreRegister(UE::Editor::DataStorage::ICoreProvider& DataStorage)
-{
-	FWorldDelegates::LevelAddedToWorld.AddUObject(this, &UTedsLevelFactory::OnLevelAddedToWorld);
-	FWorldDelegates::LevelRemovedFromWorld.AddUObject(this, &UTedsLevelFactory::OnLevelRemovedFromWorld);
-}
-
-void UTedsLevelFactory::PreShutdown(UE::Editor::DataStorage::ICoreProvider& DataStorage)
-{
-	FWorldDelegates::LevelAddedToWorld.RemoveAll(this);
-	FWorldDelegates::LevelRemovedFromWorld.RemoveAll(this);
+	Compatibility->RegisterTypeTableAssociation(ULevel::StaticClass(), LevelTable);
 }
 
 void UTedsLevelFactory::RegisterQueries(UE::Editor::DataStorage::ICoreProvider& DataStorage)
@@ -67,71 +55,10 @@ void UTedsLevelFactory::RegisterQueries(UE::Editor::DataStorage::ICoreProvider&
 			.All<FLevelTag, FTypedElementSyncFromWorldTag>()
 		.Compile());
 
-	struct FRegisterLevelCommand
-	{
-		void operator()() const
-		{
-			if (ULevel* Level = LevelPtr.Get())
-			{
-				TedsCompat->AddCompatibleObject(Level);
-			}
-		}
-
-		ICompatibilityProvider* TedsCompat;
-		TWeakObjectPtr<ULevel> LevelPtr;
-	};
-
-	// The Persistent world for each world needs to be tracked separately, since it doesn't go through LevelAddedToWorld/LevelRemovedFromWorld
-	DataStorage.RegisterQuery(
-		Select(TEXT("Register Persistent Level on World creation"),
-			FObserver::OnAdd<FWorldTag>(),
-			[this](IQueryContext& Context, RowHandle Row, const FTypedElementUObjectColumn& ObjectColumn)
-			{
-				if (UWorld* World = Cast<UWorld>(ObjectColumn.Object.Get()))
-				{
-					FRegisterLevelCommand Command{.TedsCompat = DataStorageCompat, .LevelPtr = World->PersistentLevel};
-					Context.PushCommand(Command);
-				}
-			})
-		.Where()
-			.All<FWorldTag>()
-		.Compile());
-
-	struct FUnregisterLevelCommand
-	{
-		void operator()() const
-		{
-			if (ULevel* Level = LevelPtr.Get())
-			{
-				TedsCompat->RemoveCompatibleObject(Level);
-			}
-		}
-
-		ICompatibilityProvider* TedsCompat;
-		TWeakObjectPtr<ULevel> LevelPtr;
-	};
-
-	DataStorage.RegisterQuery(
-		Select(TEXT("Unregister Persistent Level on World removal"),
-			FObserver::OnRemove<FTypedElementUObjectColumn>(),
-			[this](IQueryContext& Context, RowHandle Row, const FTypedElementUObjectColumn& ObjectColumn)
-			{
-				// We need to make sure the persistent level is removed properly if the world is marked for GC
-				constexpr bool bEvenIfPendingKill = true;
-				if (UWorld* World = Cast<UWorld>(ObjectColumn.Object.Get(bEvenIfPendingKill)))
-				{
-					FUnregisterLevelCommand Command{.TedsCompat = DataStorageCompat, .LevelPtr = World->PersistentLevel};
-					Context.PushCommand(Command);
-				}
-			})
-		.Where()
-			.All<FWorldTag>()
-		.Compile());
-
 	DataStorage.RegisterQuery(
 		Select(TEXT("Add Icon override to level"),
 			FObserver::OnAdd<FTypedElementUObjectColumn>(),
-			[this](IQueryContext& Context, RowHandle Row, FTypedElementIconOverrideColumn& IconOverrideColumn)
+			[](IQueryContext& Context, RowHandle Row, FTypedElementIconOverrideColumn& IconOverrideColumn)
 			{
 				// ULevels show up with the icon for UWorld in the editor (e.g Outliner)
 				FSlateIcon WorldIcon = FSlateIconFinder::FindIconForClass(UWorld::StaticClass());
@@ -141,19 +68,3 @@ void UTedsLevelFactory::RegisterQueries(UE::Editor::DataStorage::ICoreProvider&
 			.All<FLevelTag>()
 		.Compile());
 }
-
-void UTedsLevelFactory::OnLevelAddedToWorld(ULevel* InLevel, UWorld* InWorld)
-{
-	if (DataStorageCompat)
-	{
-		DataStorageCompat->AddCompatibleObject(InLevel);
-	}
-}
-
-void UTedsLevelFactory::OnLevelRemovedFromWorld(ULevel* InLevel, UWorld* InWorld)
-{
-	if (DataStorageCompat)
-	{
-		DataStorageCompat->RemoveCompatibleObject(InLevel);
-	}
-}
diff --git a/Engine/Plugins/Experimental/EditorDataStorageFeatures/Source/TedsEditorCompatibility/Private/Level/TedsLevelFactory.h b/Engine/Plugins/Experimental/EditorDataStorageFeatures/Source/TedsEditorCompatibility/Private/Level/TedsLevelFactory.h
--- a/Engine/Plugins/Experimental/EditorDataStorageFeatures/Source/TedsEditorCompatibility/Private/Level/TedsLevelFactory.h
+++ b/Engine/Plugins/Experimental/EditorDataStorageFeatures/Source/TedsEditorCompatibility/Private/Level/TedsLevelFactory.h
@@ -13,17 +13,8 @@ class UTedsLevelFactory : public UEditorDataStorageFactory
 	GENERATED_BODY()
 
 public:
-	~UTedsLevelFactory() override = default;
+	virtual ~UTedsLevelFactory() override = default;
 
 	virtual void RegisterTables(UE::Editor::DataStorage::ICoreProvider& DataStorage) override;
-	virtual void PreRegister(UE::Editor::DataStorage::ICoreProvider& DataStorage) override;
-	virtual void PreShutdown(UE::Editor::DataStorage::ICoreProvider& DataStorage) override;
 	virtual void RegisterQueries(UE::Editor::DataStorage::ICoreProvider& DataStorage) override;
-
-private:
-	void OnLevelAddedToWorld(ULevel* InLevel, UWorld* InWorld);
-	void OnLevelRemovedFromWorld(ULevel* InLevel, UWorld* InWorld);
-
-private:
-	UE::Editor::DataStorage::ICompatibilityProvider* DataStorageCompat;
 };
\ No newline at end of file
diff --git a/Engine/Source/Runtime/Engine/Private/Level.cpp b/Engine/Source/Runtime/Engine/Private/Level.cpp
--- a/Engine/Source/Runtime/Engine/Private/Level.cpp
+++ b/Engine/Source/Runtime/Engine/Private/Level.cpp
@@ -646,6 +646,17 @@ void ULevel::PostInitProperties()
 {
 	Super::PostInitProperties();
 
+#if WITH_EDITOR
+	if (IsInGameThread())
+	{
+		using namespace UE::Editor::DataStorage;
+		if (ICompatibilityProvider* Compatibility = GetMutableDataStorageFeature<ICompatibilityProvider>(CompatibilityFeatureName))
+		{
+			Compatibility->AddCompatibleObject(this);
+		}
+	}
+#endif // WITH_EDITOR
+
 	// Initialize LevelBuildDataId to something unique, in case this is a new ULevel
 	if(!IsTemplate())
 	{
@@ -1606,6 +1617,14 @@ void ULevel::BeginDestroy()
 		IStreamingManager::Get().RemoveLevel( this );
 	}
 
+#if WITH_EDITOR
+	using namespace UE::Editor::DataStorage;
+	if (ICompatibilityProvider* Compatibility = GetMutableDataStorageFeature<ICompatibilityProvider>(CompatibilityFeatureName))
+	{
+		Compatibility->RemoveCompatibleObject(this);
+	}
+#endif // WITH_EDITOR
+	
 	Super::BeginDestroy();
 
 	// Remove this level from its OwningWorld's collection
