[TEDS Outliner] Move Outliner Params to accept an array of columns and metadata instead of a Query. This will allow tags to appear as columns and match the TEDS Viewer parameters

#Synced-CL 50399512

diff --git a/Engine/Plugins/ChaosVD/Source/ChaosVD/Private/ChaosVDWorldOutlinerTab.cpp b/Engine/Plugins/ChaosVD/Source/ChaosVD/Private/ChaosVDWorldOutlinerTab.cpp
--- a/Engine/Plugins/ChaosVD/Source/ChaosVD/Private/ChaosVDWorldOutlinerTab.cpp
+++ b/Engine/Plugins/ChaosVD/Source/ChaosVD/Private/ChaosVDWorldOutlinerTab.cpp
@@ -54,13 +54,7 @@ void FChaosVDWorldOutlinerTab::CreateWorldOutlinerWidget()
 	Params.bShowRowHandleColumn = false;
 	Params.bUseDefaultObservers = false; // CVD uses custom observers to time slice addition and removal currently
 
-	if (ICoreProvider* Storage = GetMutableDataStorageFeature<ICoreProvider>(StorageFeatureName))
-	{
-		Params.ColumnQuery = Storage->RegisterQuery(
-			Select()
-				.ReadOnly<FVisibleInEditorColumn>()
-			.Compile());
-	}
+	Params.ColumnDescription = FTedsOutlinerColumnDescription({FVisibleInEditorColumn::StaticStruct()});
 
 	const FTedsOutlinerHierarchyData::FGetParentRowHandle RowHandleGetter = FTedsOutlinerHierarchyData::FGetParentRowHandle::CreateLambda([](const void* InColumnData)
 	{
diff --git a/Engine/Plugins/Experimental/Dataflow/Source/DataflowEditor/Private/Dataflow/DataflowOutlinerView.cpp b/Engine/Plugins/Experimental/Dataflow/Source/DataflowEditor/Private/Dataflow/DataflowOutlinerView.cpp
--- a/Engine/Plugins/Experimental/Dataflow/Source/DataflowEditor/Private/Dataflow/DataflowOutlinerView.cpp
+++ b/Engine/Plugins/Experimental/Dataflow/Source/DataflowEditor/Private/Dataflow/DataflowOutlinerView.cpp
@@ -38,10 +38,10 @@ TSharedPtr<ISceneOutliner> FDataflowOutlinerView::CreateWidget()
 		.Where(TColumn<FDataflowSceneObjectTag>(GetEditorContent()->GetDataflowOwner().GetFName()) || TColumn<FDataflowSceneStructTag>(GetEditorContent()->GetDataflowOwner().GetFName()) )
 		.Compile();
 
-	UE::Editor::DataStorage::FQueryDescription ColumnQueryDescription =
-		Select()
-		.ReadOnly<FTypedElementClassTypeInfoColumn, FVisibleInEditorColumn>()
-		.Compile();
+	const UE::Editor::Outliner::FTedsOutlinerColumnDescription ColumnDescription({
+		FTypedElementClassTypeInfoColumn::StaticStruct(),
+		FVisibleInEditorColumn::StaticStruct()
+	});
 
 	FSceneOutlinerInitializationOptions InitOptions;
 	InitOptions.bShowHeaderRow = true;
@@ -52,7 +52,7 @@ TSharedPtr<ISceneOutliner> FDataflowOutlinerView::CreateWidget()
 	UE::Editor::Outliner::FTedsOutlinerParams Params(nullptr);
 	Params.QueryDescription = RowQueryDescription;
 	Params.bShowRowHandleColumn = false;
-	Params.ColumnQuery = Storage->RegisterQuery(MoveTemp(ColumnQueryDescription));
+	Params.ColumnDescription = ColumnDescription;
 	
 	// Add outliner filter queries
 	Params.Filters.Emplace("Dataflow Construction",
diff --git a/Engine/Plugins/Experimental/EditorDataStorageFeatures/Source/TedsOutliner/Private/Compatibility/SceneOutlinerTedsBridge.cpp b/Engine/Plugins/Experimental/EditorDataStorageFeatures/Source/TedsOutliner/Private/Compatibility/SceneOutlinerTedsBridge.cpp
--- a/Engine/Plugins/Experimental/EditorDataStorageFeatures/Source/TedsOutliner/Private/Compatibility/SceneOutlinerTedsBridge.cpp
+++ b/Engine/Plugins/Experimental/EditorDataStorageFeatures/Source/TedsOutliner/Private/Compatibility/SceneOutlinerTedsBridge.cpp
@@ -40,7 +40,7 @@ void USceneOutlinerTedsBridgeFactory::PreRegister(UE::Editor::DataStorage::ICore
 
 void USceneOutlinerTedsBridgeFactory::RegisterTables(UE::Editor::DataStorage::ICoreProvider& DataStorage)
 {
-	DataStorage.RegisterTable<FTedsOutlinerColumn, FTedsOutlinerColumnQueryColumn, FTedsOutlinerSelectionChangeColumn>(UE::Editor::Outliner::Helpers::GetTedsOutlinerTableName());
+	DataStorage.RegisterTable<FTedsOutlinerColumn, FTedsOutlinerColumnRefreshEventColumn, FTedsOutlinerSelectionChangeColumn>(UE::Editor::Outliner::Helpers::GetTedsOutlinerTableName());
 }
 
 void USceneOutlinerTedsBridgeFactory::RegisterWidgetPurposes(UE::Editor::DataStorage::IUiProvider& DataStorageUi) const
diff --git a/Engine/Plugins/Experimental/EditorDataStorageFeatures/Source/TedsOutliner/Private/TedsOutlinerColumn.cpp b/Engine/Plugins/Experimental/EditorDataStorageFeatures/Source/TedsOutliner/Private/TedsOutlinerColumn.cpp
--- a/Engine/Plugins/Experimental/EditorDataStorageFeatures/Source/TedsOutliner/Private/TedsOutlinerColumn.cpp
+++ b/Engine/Plugins/Experimental/EditorDataStorageFeatures/Source/TedsOutliner/Private/TedsOutlinerColumn.cpp
@@ -12,9 +12,9 @@
 
 namespace UE::Editor::Outliner
 {
-	FTedsOutlinerUiColumnInitParams::FTedsOutlinerUiColumnInitParams(DataStorage::QueryHandle InQuery, DataStorage::ICoreProvider& InStorage,
+	FTedsOutlinerUiColumnInitParams::FTedsOutlinerUiColumnInitParams(const DataStorage::FMetaData& InMetaData, DataStorage::ICoreProvider& InStorage,
 		DataStorage::IUiProvider& InStorageUi, DataStorage::ICompatibilityProvider& InStorageCompatibility)
-		: Query(InQuery)
+		: MetaData(InMetaData)
 		, Storage(InStorage)
 		, StorageUi(InStorageUi)
 		, StorageCompatibility(InStorageCompatibility)
@@ -25,7 +25,7 @@ namespace UE::Editor::Outliner
 		: Storage(InitParams.Storage)
 		, StorageUi(InitParams.StorageUi)
 		, StorageCompatibility(InitParams.StorageCompatibility)
-		, Query(InitParams.Query)
+		, MetaData(InitParams.MetaData)
 		, NameId(InitParams.NameId)
 		, FallbackColumn(InitParams.FallbackColumn)
 		, OwningOutliner(InitParams.OwningOutliner)
@@ -37,10 +37,8 @@ namespace UE::Editor::Outliner
 
 		using namespace UE::Editor::DataStorage;
 
-		FQueryDescription QueryDescription = Storage.GetQueryDescription(Query);
-
 		TableViewerColumnImpl = MakeUnique<FTedsTableViewerColumn>(NameId, InitParams.CellWidgetConstructor, InitParams.ColumnTypes,
-			InitParams.HeaderWidgetConstructor, QueryDescription.MetaData);
+			InitParams.HeaderWidgetConstructor, MetaData);
 
 		TableViewerColumnImpl->SetIsRowVisibleDelegate(
 			FTedsTableViewerColumn::FIsRowVisible::CreateRaw(this, &FTedsOutlinerUiColumn::IsRowVisible)
diff --git a/Engine/Plugins/Experimental/EditorDataStorageFeatures/Source/TedsOutliner/Private/TedsOutlinerColumn.h b/Engine/Plugins/Experimental/EditorDataStorageFeatures/Source/TedsOutliner/Private/TedsOutlinerColumn.h
--- a/Engine/Plugins/Experimental/EditorDataStorageFeatures/Source/TedsOutliner/Private/TedsOutlinerColumn.h
+++ b/Engine/Plugins/Experimental/EditorDataStorageFeatures/Source/TedsOutliner/Private/TedsOutlinerColumn.h
@@ -14,11 +14,11 @@ namespace UE::Editor::Outliner
 
 	struct FTedsOutlinerUiColumnInitParams
 	{
-		FTedsOutlinerUiColumnInitParams(DataStorage::QueryHandle InQuery, DataStorage::ICoreProvider& InStorage,
+		FTedsOutlinerUiColumnInitParams(const DataStorage::FMetaData& InMetaData, DataStorage::ICoreProvider& InStorage,
 			DataStorage::IUiProvider& InStorageUi, DataStorage::ICompatibilityProvider& InStorageCompatibility);
 		
-		// The query used for this column
-		DataStorage::QueryHandle Query;
+		// The metadata used for this column
+		DataStorage::FMetaData MetaData;
 
 		// Data storage interfaces
 		DataStorage::ICoreProvider& Storage;
@@ -73,7 +73,6 @@ namespace UE::Editor::Outliner
 		DataStorage::ICoreProvider& Storage;
 		DataStorage::IUiProvider& StorageUi;
 		DataStorage::ICompatibilityProvider& StorageCompatibility;
-		DataStorage::QueryHandle Query;
 		DataStorage::FMetaData MetaData;
 		FName NameId;
 		TSharedPtr<ISceneOutlinerColumn> FallbackColumn;
diff --git a/Engine/Plugins/Experimental/EditorDataStorageFeatures/Source/TedsOutliner/Private/TedsOutlinerImpl.cpp b/Engine/Plugins/Experimental/EditorDataStorageFeatures/Source/TedsOutliner/Private/TedsOutlinerImpl.cpp
--- a/Engine/Plugins/Experimental/EditorDataStorageFeatures/Source/TedsOutliner/Private/TedsOutlinerImpl.cpp
+++ b/Engine/Plugins/Experimental/EditorDataStorageFeatures/Source/TedsOutliner/Private/TedsOutlinerImpl.cpp
@@ -38,6 +38,7 @@
 #include "Elements/Columns/TypedElementOverrideColumns.h"
 #include "Elements/Columns/TypedElementSelectionColumns.h"
 #include "Elements/Columns/TypedElementSlateWidgetColumns.h"
+#include "Elements/Columns/TypedElementUIColumns.h"
 #include "Filters/FilterBase.h"
 #include "Widgets/SWidget.h"
 
@@ -189,6 +190,36 @@ namespace LabelWidgetUtils
 	}
 }
 
+FTedsOutlinerColumnDescription::FTedsOutlinerColumnDescription(const TArray<TWeakObjectPtr<const UScriptStruct>>& InColumns,
+	const TArray<TSharedPtr<FColumnMetaData>>& InColumnMetaData, const FMetaData& InGenericMetaData)
+	: Columns(InColumns)
+	, GenericMetaData(InGenericMetaData)
+{
+	ColumnMetaData.Reserve(Columns.Num());
+
+	// Initialize the ColumnMetaData with ReadOnly MetaData for each column if nothing was given. Custom MetaData
+	// can override this by passing through the constructor.
+	if (InColumnMetaData.IsEmpty())
+	{
+		for (const TWeakObjectPtr<const UScriptStruct> Column : Columns)
+		{
+			if (Column.IsValid())
+			{
+				ColumnMetaData.Add(MakeShared<FColumnMetaData>(Column.Get(), FColumnMetaData::EFlags::None));
+			}
+		}
+	}
+	else
+	{
+		ColumnMetaData = InColumnMetaData;
+	}
+}
+	
+FTedsOutlinerColumnDescription::FTedsOutlinerColumnDescription(const TArray<TWeakObjectPtr<const UScriptStruct>>& InColumns, const FMetaData& InGenericMetaData)
+	: Columns(InColumns)
+	, GenericMetaData(InGenericMetaData)
+{}
+
 FTedsOutlinerParams::FTedsOutlinerParams(SSceneOutliner* InSceneOutliner)
 	: SceneOutliner(InSceneOutliner)
 	, QueryDescription()
@@ -328,7 +359,7 @@ void FTedsOutlinerImpl::Init()
 	// FTedsOutlinerImpl::Init is called before the Outliner UI has been fully init, so we generate the columns in the next tick
 	TickerHandle = FTSTicker::GetCoreTicker().AddTicker(FTickerDelegate::CreateLambda([this](float)
 	{
-		RegenerateColumns();
+		RegenerateColumns(CreationParams.ColumnDescription);
 		return false;
 	}));
 }
@@ -372,8 +403,6 @@ FTedsOutlinerImpl::~FTedsOutlinerImpl()
 {
 	if (Storage)
 	{
-		// We take ownership of the column query and so have to handle unregistering it
-		Storage->UnregisterQuery(CreationParams.ColumnQuery);
 		Storage->RemoveRow(OutlinerRowHandle);
 		Storage->OnUpdateCompleted().RemoveAll(this);
 	}
@@ -526,10 +555,8 @@ void FTedsOutlinerImpl::RegisterTedsOutliner()
 				}
 			});
 
-		if (FTedsOutlinerColumnQueryColumn* ColumnQueryColumn = Storage->GetColumn<FTedsOutlinerColumnQueryColumn>(OutlinerRowHandle))
+		if (FTedsOutlinerColumnRefreshEventColumn* ColumnQueryColumn = Storage->GetColumn<FTedsOutlinerColumnRefreshEventColumn>(OutlinerRowHandle))
 		{
-			ColumnQueryColumn->ColumnQueryHandle = CreationParams.ColumnQuery;
-
 			ColumnQueryColumn->OnRefreshColumns.BindSP(this, &FTedsOutlinerImpl::RegenerateColumns);
 		}
 
@@ -555,14 +582,13 @@ void FTedsOutlinerImpl::ClearColumns() const
 	}
 }
 
-void FTedsOutlinerImpl::RegenerateColumns()
+void FTedsOutlinerImpl::RegenerateColumns(const Outliner::FTedsOutlinerColumnDescription& ColumnDescription)
 {
 	ClearColumns();
-	
-	FTedsOutlinerColumnQueryColumn* ColumnQueryColumn = Storage->GetColumn<FTedsOutlinerColumnQueryColumn>(OutlinerRowHandle);
+
 	FTedsOutlinerColumn* TedsOutlinerColumn = Storage->GetColumn<FTedsOutlinerColumn>(OutlinerRowHandle);
 
-	if  (!ColumnQueryColumn || !TedsOutlinerColumn)
+	if  (!TedsOutlinerColumn)
 	{
 		return;
 	}
@@ -576,22 +602,19 @@ void FTedsOutlinerImpl::RegenerateColumns()
 
 	using MatchApproach = DataStorage::IUiProvider::EMatchApproach;
 	constexpr uint8 DefaultPriorityIndex = 100;
-
-	FQueryDescription ColumnQueryDescription = Storage->GetQueryDescription(ColumnQueryColumn->ColumnQueryHandle);
 	
-	DataStorage::FQueryMetaDataView MetaDataView(ColumnQueryDescription);
+	DataStorage::FColumnsMetaDataView MetaDataView(ColumnDescription.GetColumnMetaData());
+	
+	AddedColumns.Reset(ColumnDescription.GetColumns().Num());
 	
-	int32 SelectionCount = ColumnQueryDescription.SelectionTypes.Num();
-	AddedColumns.Reset(SelectionCount);
-
 	TArray<TWeakObjectPtr<const UScriptStruct>> ColumnTypes =
-		DataStorage::TableViewerUtils::CreateVerifiedColumnTypeArray(ColumnQueryDescription.SelectionTypes);
+		TableViewerUtils::CreateVerifiedColumnTypeArray(ColumnDescription.GetColumns());
 	
-	DataStorage::RowHandle CellPurposeRow = StorageUi->FindPurpose(CellWidgetPurpose);
-	DataStorage::RowHandle HeaderPurposeRow = StorageUi->FindPurpose(HeaderWidgetPurpose);
+	const RowHandle CellPurposeRow = StorageUi->FindPurpose(CellWidgetPurpose);
+	const RowHandle HeaderPurposeRow = StorageUi->FindPurpose(HeaderWidgetPurpose);
 
 	int32 IndexOffset = 0;
-	auto ColumnConstructor = [this, ColumnQueryColumn, TedsOutlinerColumn, MetaDataView, &IndexOffset, HeaderPurposeRow, Dealiaser](
+	auto ColumnConstructor = [this, ColumnDescription, TedsOutlinerColumn, MetaDataView, &IndexOffset, HeaderPurposeRow, Dealiaser](
 		IUiProvider::FWidgetConstructorPtr CellConstructor, TConstArrayView<TWeakObjectPtr<const UScriptStruct>> ColumnTypes)
 		{
 			/* If we have a fallback column for this query, remove it, take over it's priority and 
@@ -619,7 +642,7 @@ void FTedsOutlinerImpl::RegenerateColumns()
 					ESceneOutlinerColumnVisibility::Visible, 
 					ColumnPriority,
 					FCreateSceneOutlinerColumn::CreateLambda(
-						[this, ColumnQueryColumn, TedsOutlinerColumn, MetaDataView, NameId, &ColumnTypes, CellConstructor, FallbackColumnPtr, HeaderPurposeRow, Dealiaser](ISceneOutliner&)
+						[this, ColumnDescription, TedsOutlinerColumn, MetaDataView, NameId, &ColumnTypes, CellConstructor, FallbackColumnPtr, HeaderPurposeRow, Dealiaser](ISceneOutliner&)
 						{
 							IUiProvider::FWidgetConstructorPtr HeaderConstructor = 
 								UE::Editor::DataStorage::TableViewerUtils::CreateHeaderWidgetConstructor(*StorageUi, MetaDataView, ColumnTypes, HeaderPurposeRow);
@@ -630,7 +653,7 @@ void FTedsOutlinerImpl::RegenerateColumns()
 								.IsSortingHandler = DataStorage::FTedsTableViewerColumn::FIsSorting::CreateSP(this, &FTedsOutlinerImpl::IsSorting)
 							};
 
-							FTedsOutlinerUiColumnInitParams Params(ColumnQueryColumn->ColumnQueryHandle, *Storage, *StorageUi, *StorageCompatibility);
+							FTedsOutlinerUiColumnInitParams Params(ColumnDescription.GetGenericMetaData(), *Storage, *StorageUi, *StorageCompatibility);
 							Params.NameId = NameId;
 							Params.ColumnTypes = TArray<TWeakObjectPtr<const UScriptStruct>>(ColumnTypes.GetData(), ColumnTypes.Num());
 							Params.HeaderWidgetConstructor = MoveTemp(HeaderConstructor);
diff --git a/Engine/Plugins/Experimental/EditorDataStorageFeatures/Source/TedsOutliner/Private/TedsOutlinerModule.cpp b/Engine/Plugins/Experimental/EditorDataStorageFeatures/Source/TedsOutliner/Private/TedsOutlinerModule.cpp
--- a/Engine/Plugins/Experimental/EditorDataStorageFeatures/Source/TedsOutliner/Private/TedsOutlinerModule.cpp
+++ b/Engine/Plugins/Experimental/EditorDataStorageFeatures/Source/TedsOutliner/Private/TedsOutlinerModule.cpp
@@ -119,46 +119,25 @@ void FTedsOutlinerModule::ShutdownModule()
 	IModuleInterface::ShutdownModule();
 }
 
-QueryHandle FTedsOutlinerModule::GetLevelEditorTedsOutlinerColumnQuery()
+FTedsOutlinerColumnDescription FTedsOutlinerModule::GetLevelEditorTedsOutlinerColumnDescription()
 {
-	using namespace UE::Editor::DataStorage;
 	using namespace UE::Editor::DataStorage::Columns;
-	Storage = GetMutableDataStorageFeature<ICoreProvider>(StorageFeatureName);
 		
-	using namespace DataStorage::Queries;
-
-	QueryHandle ColumnQuery;
-
-	if (!CVarUseNewRevisionControlWidgets->GetBool())
-	{
-		ColumnQuery = Storage->RegisterQuery(
-			Select()
-				.ReadOnly<FTypedElementClassTypeInfoColumn,
-					  FAlertColumn,
-					  FChildAlertColumn,
-					  FUObjectIdNameColumn,
-					  FTedsActorMobilityColumn,
-					  FLevelColumn,
-					  FVisibleInEditorColumn>()
-			.Compile());
-
-	}
-	else
+	TArray<TWeakObjectPtr<const UScriptStruct>> OutColumns({
+		FTypedElementClassTypeInfoColumn::StaticStruct(),
+		FAlertColumn::StaticStruct(),
+		FChildAlertColumn::StaticStruct(),
+		FUObjectIdNameColumn::StaticStruct(),
+		FTedsActorMobilityColumn::StaticStruct(),
+		FLevelColumn::StaticStruct(),
+		FVisibleInEditorColumn::StaticStruct()});
+
+	if (CVarUseNewRevisionControlWidgets->GetBool())
 	{
-		ColumnQuery = Storage->RegisterQuery( 
-			Select()
-				.ReadOnly<FTypedElementClassTypeInfoColumn,  
-					  FTypedElementPackageReference, 
-					  FAlertColumn, 
-					  FUObjectIdNameColumn,
-					  FTedsActorMobilityColumn,
-					  FLevelColumn,
-					  FVisibleInEditorColumn>()
-			.Compile());
-
+		OutColumns.Add(FTypedElementPackageReference::StaticStruct());
 	}
 
-	return ColumnQuery;
+	return FTedsOutlinerColumnDescription(OutColumns);
 }
 
 TSharedRef<SWidget> FTedsOutlinerModule::CreateLevelEditorTedsOutliner()
@@ -171,10 +150,12 @@ TSharedRef<SWidget> FTedsOutlinerModule::CreateLevelEditorTedsOutliner()
 
 	using namespace DataStorage::Queries;
 
+	Storage = GetMutableDataStorageFeature<ICoreProvider>(StorageFeatureName);
+
 	// The Outliner is populated with Actors and Entities
-	DataStorage::FQueryDescription OutlinerQueryDescription =
+	const FQueryDescription OutlinerQueryDescription =
 		Select()
-		.Where(TColumn<FTypedElementActorTag>() || TColumn<FFolderCompatibilityColumn>())
+		.Where((TColumn<FTypedElementUObjectColumn>() || TColumn<FFolderCompatibilityColumn>()) && TColumn<FTypedElementWorldColumn>())
 		.Compile();
 
 	FSceneOutlinerInitializationOptions InitOptions;
@@ -185,8 +166,9 @@ TSharedRef<SWidget> FTedsOutlinerModule::CreateLevelEditorTedsOutliner()
 
 	FTedsOutlinerParams Params(nullptr);
 	Params.QueryDescription = OutlinerQueryDescription;
-	Params.ColumnQuery = GetLevelEditorTedsOutlinerColumnQuery();
-		Params.ClassFilters.Add(AActor::StaticClass());
+	Params.ColumnDescription = GetLevelEditorTedsOutlinerColumnDescription();
+	
+	Params.ClassFilters.Add(AActor::StaticClass());
 	Params.ClassFilters.Add(AStaticMeshActor::StaticClass());
 	Params.ClassFilters.Add(ADirectionalLight::StaticClass());
 	
@@ -243,7 +225,7 @@ void FTedsOutlinerModule::RegisterLevelEditorTedsOutlinerTab()
 		TSharedPtr<FTabManager> LevelEditorTabManager = LevelEditorModule.GetLevelEditorTabManager();
 
 		LevelEditorTabManager->RegisterTabSpawner(TedsOutlinerTabName, FOnSpawnTab::CreateRaw(this, &FTedsOutlinerModule::OpenLevelEditorTedsOutliner))
-		.SetDisplayName(LOCTEXT("TedsTableVIewerTitle", "Table Viewer (Experimental)"))
+		.SetDisplayName(LOCTEXT("TedsOutlinerTabName", "Data-Driven Outliner (Experimental)"))
 		.SetGroup(WorkspaceMenu::GetMenuStructure().GetLevelEditorOutlinerCategory())
 		.SetIcon(FSlateIcon(FAppStyle::GetAppStyleSetName(), "LevelEditor.Tabs.Outliner"))
 		.SetAutoGenerateMenuEntry(false); // This can only be summoned from the Cvar now
diff --git a/Engine/Plugins/Experimental/EditorDataStorageFeatures/Source/TedsOutliner/Public/Columns/TedsOutlinerColumns.h b/Engine/Plugins/Experimental/EditorDataStorageFeatures/Source/TedsOutliner/Public/Columns/TedsOutlinerColumns.h
--- a/Engine/Plugins/Experimental/EditorDataStorageFeatures/Source/TedsOutliner/Public/Columns/TedsOutlinerColumns.h
+++ b/Engine/Plugins/Experimental/EditorDataStorageFeatures/Source/TedsOutliner/Public/Columns/TedsOutlinerColumns.h
@@ -8,6 +8,11 @@
 
 #include "TedsOutlinerColumns.generated.h"
 
+namespace UE::Editor::Outliner
+{
+	struct FTedsOutlinerColumnDescription;
+}
+
 class ISceneOutliner;
 class UToolMenu;
 class SSceneOutliner;
@@ -15,7 +20,7 @@ class SSceneOutliner;
 namespace UE::Editor::Outliner
 {
 	// Delegate that refreshes columns on the Teds Outliner it is bound to when executed
-	DECLARE_DELEGATE(FOnRefreshColumns)
+	DECLARE_DELEGATE_OneParam(FOnRefreshColumns, const UE::Editor::Outliner::FTedsOutlinerColumnDescription&)
 
 	// Delegate called when the selection in the Outliner changes
 	DECLARE_MULTICAST_DELEGATE(FOnTedsOutlinerSelectionChanged)
@@ -48,16 +53,13 @@ struct FTedsOutlinerContextMenuColumn final : public FEditorDataStorageColumn
 	FTedsOutlinerContextMenuDelegate OnCreateContextMenu;
 };
 
-/** Column used to store information about the columns being viewed in the Teds Outliner */
+/** Column used to refresh information about the columns being viewed in the Teds Outliner */
 USTRUCT()
-struct FTedsOutlinerColumnQueryColumn final : public FEditorDataStorageColumn
+struct FTedsOutlinerColumnRefreshEventColumn final : public FEditorDataStorageColumn
 {
 	GENERATED_BODY()
 
-	// The SELECT() part of the query description attached to this handle specifies the columns to be displayed
-	UE::Editor::DataStorage::QueryHandle ColumnQueryHandle;
-
-	// Call this delegate to refresh the columns in the Teds Outliner with the values specified in ColumnQueryDescription
+	// Call this delegate to refresh the columns in the Teds Outliner
 	UE::Editor::Outliner::FOnRefreshColumns OnRefreshColumns;
 };
 
@@ -77,4 +79,4 @@ struct FTedsOutlinerSelectionChangeColumn final : public FEditorDataStorageColum
 	GENERATED_BODY()
 
 	UE::Editor::Outliner::FOnTedsOutlinerSelectionChanged OnSelectionChanged;
-};
\ No newline at end of file
+};
diff --git a/Engine/Plugins/Experimental/EditorDataStorageFeatures/Source/TedsOutliner/Public/TedsOutlinerImpl.h b/Engine/Plugins/Experimental/EditorDataStorageFeatures/Source/TedsOutliner/Public/TedsOutlinerImpl.h
--- a/Engine/Plugins/Experimental/EditorDataStorageFeatures/Source/TedsOutliner/Public/TedsOutlinerImpl.h
+++ b/Engine/Plugins/Experimental/EditorDataStorageFeatures/Source/TedsOutliner/Public/TedsOutlinerImpl.h
@@ -102,6 +102,29 @@ struct FTedsOutlinerHierarchyData
 		return FTedsOutlinerHierarchyData(FTableRowParentColumn::StaticStruct(), RowHandleGetter, RowHandleSetter, FGetChildrenRowsHandles());
 	}
 };
+	
+struct FTedsOutlinerColumnDescription
+{
+	TEDSOUTLINER_API FTedsOutlinerColumnDescription(
+		const TArray<TWeakObjectPtr<const UScriptStruct>>& InColumns = {},
+		const TArray<TSharedPtr<FColumnMetaData>>& InColumnMetaData = {},
+		const FMetaData& InGenericMetaData = FMetaData());
+
+	TEDSOUTLINER_API FTedsOutlinerColumnDescription(
+		const TArray<TWeakObjectPtr<const UScriptStruct>>& InColumns, const FMetaData& InGenericMetaData);
+
+	TEDSOUTLINER_API TArray<TWeakObjectPtr<const UScriptStruct>> GetColumns() const { return Columns; }
+	TEDSOUTLINER_API TArray<TSharedPtr<FColumnMetaData>> GetColumnMetaData() const { return ColumnMetaData; }
+	TEDSOUTLINER_API FMetaData GetGenericMetaData() const { return GenericMetaData; }
+
+protected:
+	// Columns to display in the Outliner
+	TArray<TWeakObjectPtr<const UScriptStruct>> Columns;
+	// Metadata for Columns (ReadOnly/ReadWrite), by default the constructor will populate this with ReadOnly Metadata for all columns
+	TArray<TSharedPtr<FColumnMetaData>> ColumnMetaData;
+	// Generic meta data for the TEDS (i.e. GenericMetaData.AddOrSetMutableData("bAllowSorting", false))
+	FMetaData GenericMetaData;
+};
 
 struct FTedsOutlinerParams
 {
@@ -112,9 +135,8 @@ struct FTedsOutlinerParams
 	// The query description that will be used to populate rows in the TEDS Outliner
 	DataStorage::FQueryDescription QueryDescription;
 
-	// The query that will be used to populate the initial columns in the TEDS Outliner
-	// NOTE: The TEDS Outliner takes ownership of this query and will unregister it on destruction.
-	DataStorage::QueryHandle ColumnQuery;
+	// The columns to display in the TEDS Outliner and any additional metadata
+	FTedsOutlinerColumnDescription ColumnDescription;
 	
 	// TEDS Filters that utilize QueryDescriptions or QueryFunctions in this TEDS Outliner
 	TArray<FTedsFilterData> Filters;
@@ -256,7 +278,7 @@ protected:
 	UE_API void RegisterTedsOutliner();
 
 	UE_API void ClearColumns() const;
-	UE_API void RegenerateColumns();
+	UE_API void RegenerateColumns(const FTedsOutlinerColumnDescription& ColumnDescription);
 
 protected:
 	// TEDS Storage Constructs
diff --git a/Engine/Plugins/Experimental/EditorDataStorageFeatures/Source/TedsOutliner/Public/TedsOutlinerModule.h b/Engine/Plugins/Experimental/EditorDataStorageFeatures/Source/TedsOutliner/Public/TedsOutlinerModule.h
--- a/Engine/Plugins/Experimental/EditorDataStorageFeatures/Source/TedsOutliner/Public/TedsOutlinerModule.h
+++ b/Engine/Plugins/Experimental/EditorDataStorageFeatures/Source/TedsOutliner/Public/TedsOutlinerModule.h
@@ -17,6 +17,7 @@ class FSpawnTabArgs;
 namespace UE::Editor::Outliner
 {
 	struct FTedsOutlinerParams;
+	struct FTedsOutlinerColumnDescription;
 
 /**
  * Implements the Scene Outliner module.
@@ -31,9 +32,8 @@ public:
 	/**
 	 * Creates a TEDS-Outliner widget
 	 *
-	 * @param	InitOptions						Programmer-driven configuration for the scene outliner
+	 * @param	InInitOptions						Programmer-driven configuration for the scene outliner
 	 * @param	InInitTedsOptions				Programmer-driven configuration for the TEDS queries that drive the outliner
-	 * @param   ColumnQuery						Query to describe which columns will be available in the TEDS-Outliner
 	 *
 	 * @return	New scene outliner widget
 	 */
@@ -43,9 +43,9 @@ public:
 	// IModuleInterface interface
 	virtual void StartupModule() override;
 	virtual void ShutdownModule() override;
-
-	// Get the column query the default table viewer uses
-	virtual UE::Editor::DataStorage::QueryHandle GetLevelEditorTedsOutlinerColumnQuery();
+	
+	// Get the column description the default table viewer uses
+	virtual FTedsOutlinerColumnDescription GetLevelEditorTedsOutlinerColumnDescription();
 
 	// The name of the tab the default table viewer is opened in
 	FName GetTedsOutlinerTabName();
