[TEDS Package] Further improvements to the pacakge syncs
*Tested with OFPE/OFPA

#Synced-CL 50399512

diff --git a/Engine/Plugins/Experimental/EditorDataStorageFeatures/Source/TedsRevisionControl/Private/Queries/ObjectPackagePathToColumnQueries.cpp b/Engine/Plugins/Experimental/EditorDataStorageFeatures/Source/TedsRevisionControl/Private/Queries/ObjectPackagePathToColumnQueries.cpp
--- a/Engine/Plugins/Experimental/EditorDataStorageFeatures/Source/TedsRevisionControl/Private/Queries/ObjectPackagePathToColumnQueries.cpp
+++ b/Engine/Plugins/Experimental/EditorDataStorageFeatures/Source/TedsRevisionControl/Private/Queries/ObjectPackagePathToColumnQueries.cpp
@@ -44,9 +44,12 @@ namespace UE::Editor::RevisionControl::Private
 		// If the Package is Transient we can skip finding the primary asset as that will loop over all assets in the package to find nothing
 		if (Package != GetTransientPackage() && !Package->HasAnyFlags(RF_Transient))
 		{
-			if (const UObject* PackageInstance = Package->FindAssetInPackage(); Object == PackageInstance)
+			if (Object->HasAnyFlags(RF_HasExternalPackage))
 			{
-				Context.AddColumns<FTedsPrimaryPackageObjectTag>(Row);
+				if (const UObject* ExternalPackage = Object->GetExternalPackage(); Package == ExternalPackage)
+				{
+					Context.AddColumns<FTedsPrimaryPackageObjectTag>(Row);
+				}
 			}
 			// If the Row Object is a CDO, we have to compare the CDO of the primary package object instead of the object itself
 			else if (Context.HasColumn<FTypedElementClassDefaultObjectTag>())
@@ -118,22 +121,21 @@ void UTypedElementUObjectPackagePathFactory::RegisterQueries(UE::Editor::DataSto
 			FProcessor(EQueryTickPhase::FrameEnd, DataStorage.GetQueryTickGroupName(EQueryTickGroups::SyncExternalToDataStorage)),
 			[](IQueryContext& Context, RowHandle Row, const FTypedElementUObjectColumn& Object, const FTypedElementPackageUnresolvedReference& UnresolvedPackageReference)
 			{
-				RowHandle PackageRow = Context.LookupMappedRow(
+				const RowHandle PackageRow = Context.LookupMappedRow(
 					UE::Editor::RevisionControl::MappingDomain, FMapKeyView(UnresolvedPackageReference.PathOnDisk));
 				if (!Context.IsRowAvailable(PackageRow))
 				{
 					return;
 				}
-				const UObject* ObjectInstance = Object.Object.Get();
-				if (!ObjectInstance)
+				if (const UObject* ObjectInstance = Object.Object.Get(); IsValid(ObjectInstance))
 				{
-					return;
-				}
-				UPackage* Package = ObjectInstance->GetPackage();
-				
-				Context.RemoveColumns(Row, { FTypedElementPackageUnresolvedReference::StaticStruct() });
+					if (UPackage* Package = ObjectInstance->GetPackage())
+					{
+						Context.RemoveColumns(Row, { FTypedElementPackageUnresolvedReference::StaticStruct() });
 
-				UE::Editor::RevisionControl::Private::ResolvePackageReference(Context, Package, ObjectInstance, Row, PackageRow);
+						UE::Editor::RevisionControl::Private::ResolvePackageReference(Context, Package, ObjectInstance, Row, PackageRow);
+					}
+				}
 			})
 		.AccessesHierarchy(UE::Editor::DataStorage::PackageHierarchyName)
 		.Compile());
@@ -172,40 +174,36 @@ void UTypedElementUObjectPackagePathFactory::RegisterTryAddPackageRef(UE::Editor
 
 	auto PackageSyncCallback = [](IQueryContext& Context, RowHandle Row, const FTypedElementUObjectColumn& Object)
 	{
-		if (const UObject* ObjectInstance = Object.Object.Get(); IsValid(ObjectInstance))
+		// Don't monitor transient objects since they will never exist on disk. And don't monitor objects in cooked packages since
+		// they are (likely) not under the source control root and even if they happen to be, we didn't load the loose asset on disk.
+		if (const UObject* ObjectInstance = Object.Object.Get(); IsValid(ObjectInstance) && !ObjectInstance->HasAnyFlags(RF_Transient | RF_BeginDestroyed))
 		{
-			UPackage* Target = ObjectInstance->GetPackage();
-			
-			if (ObjectInstance->HasAnyFlags(RF_Transient | RF_BeginDestroyed) || Target->HasAnyPackageFlags(PKG_Cooked))
+			if (UPackage* Target = ObjectInstance->GetPackage(); !Target->HasAnyPackageFlags(PKG_Cooked))
 			{
-				// Don't monitor transient objects since they will never exist on disk. And don't monitor objects in cooked packages since
-				// they are (likely) not under the source control root and even if they happen to be, we didn't load the loose asset on disk.
-				return;
-			}
-
-			FString Path;
-			Target->GetPathName(nullptr, Path);
-			if (FString PackageFilename; FPackageName::TryConvertLongPackageNameToFilename(Path, PackageFilename))
-			{
-				FPaths::NormalizeFilename(PackageFilename);
-				FString FullPackageFilename = FPaths::ConvertRelativePathToFull(PackageFilename);
-				const RowHandle PackageRow = Context.LookupMappedRow(UE::Editor::RevisionControl::MappingDomain, FMapKeyView(FullPackageFilename));
-				if (Context.IsRowAvailable(PackageRow))
-				{
-					UE::Editor::RevisionControl::Private::ResolvePackageReference(Context, Target, ObjectInstance, Row, PackageRow);
-				}
-				else
+				FString Path;
+				Target->GetPathName(nullptr, Path);
+				if (FString PackageFilename; FPackageName::TryConvertLongPackageNameToFilename(Path, PackageFilename))
 				{
-					FTypedElementPackageUnresolvedReference UnresolvedPackageReference;
-					UnresolvedPackageReference.PathOnDisk = FullPackageFilename;
-					Context.AddColumn(Row, MoveTemp(UnresolvedPackageReference));
-
-					FRegisterPackageRowCommand Command
+					FPaths::NormalizeFilename(PackageFilename);
+					FString FullPackageFilename = FPaths::ConvertRelativePathToFull(PackageFilename);
+					const RowHandle PackageRow = Context.LookupMappedRow(UE::Editor::RevisionControl::MappingDomain, FMapKeyView(FullPackageFilename));
+					if (Context.IsRowAvailable(PackageRow))
+					{
+						UE::Editor::RevisionControl::Private::ResolvePackageReference(Context, Target, ObjectInstance, Row, PackageRow);
+					}
+					else
 					{
-						.PackagePathname = FullPackageFilename,
-						.CoreProvider = GetMutableDataStorageFeature<ICoreProvider>(StorageFeatureName),
-					};
-					Context.PushCommand(Command);
+						FTypedElementPackageUnresolvedReference UnresolvedPackageReference;
+						UnresolvedPackageReference.PathOnDisk = FullPackageFilename;
+						Context.AddColumn(Row, MoveTemp(UnresolvedPackageReference));
+
+						FRegisterPackageRowCommand Command
+						{
+							.PackagePathname = FullPackageFilename,
+							.CoreProvider = GetMutableDataStorageFeature<ICoreProvider>(StorageFeatureName),
+						};
+						Context.PushCommand(Command);
+					}
 				}
 			}
 		}
