Fix EverythingPicker being unable to have its size set

#Synced-CL 50170641

diff --git a/Engine/Plugins/Experimental/EditorDataStorageFeatures/Source/TedsEverythingPicker/Private/Widgets/SEverythingPicker.cpp b/Engine/Plugins/Experimental/EditorDataStorageFeatures/Source/TedsEverythingPicker/Private/Widgets/SEverythingPicker.cpp
--- a/Engine/Plugins/Experimental/EditorDataStorageFeatures/Source/TedsEverythingPicker/Private/Widgets/SEverythingPicker.cpp
+++ b/Engine/Plugins/Experimental/EditorDataStorageFeatures/Source/TedsEverythingPicker/Private/Widgets/SEverythingPicker.cpp
@@ -10,6 +10,29 @@
 
 namespace UE::Editor::DataStorage::Picker
 {
+	SLATE_IMPLEMENT_WIDGET(SEverythingPicker)
+	void SEverythingPicker::PrivateRegisterAttributes(FSlateAttributeInitializer& AttributeInitializer)
+	{
+		SLATE_ADD_MEMBER_ATTRIBUTE_DEFINITION(AttributeInitializer, WidthOverride, EInvalidateWidgetReason::Layout);
+		SLATE_ADD_MEMBER_ATTRIBUTE_DEFINITION(AttributeInitializer, HeightOverride, EInvalidateWidgetReason::Layout);
+		SLATE_ADD_MEMBER_ATTRIBUTE_DEFINITION(AttributeInitializer, MinDesiredWidth, EInvalidateWidgetReason::Layout);
+		SLATE_ADD_MEMBER_ATTRIBUTE_DEFINITION(AttributeInitializer, MinDesiredHeight, EInvalidateWidgetReason::Layout);
+		SLATE_ADD_MEMBER_ATTRIBUTE_DEFINITION(AttributeInitializer, MaxDesiredWidth, EInvalidateWidgetReason::Layout);
+		SLATE_ADD_MEMBER_ATTRIBUTE_DEFINITION(AttributeInitializer, MaxDesiredHeight, EInvalidateWidgetReason::Layout);
+	}
+
+	SEverythingPicker::SEverythingPicker()
+		: WidthOverride(*this, 350.0f)
+		, HeightOverride(*this, 300.0f)
+		, MinDesiredWidth(*this)
+		, MinDesiredHeight(*this)
+		, MaxDesiredWidth(*this)
+		, MaxDesiredHeight(*this)
+	{
+	}
+
+	SEverythingPicker::~SEverythingPicker() = default;
+
 	TAutoConsoleVariable<bool> CVarPickerSystemEnabled(
 		TEXT("TEDS.Feature.PickerEnabled"),
 		false, // Defaulting to false until we hit feature parity (search bar, custom filtering, etc.)
@@ -17,11 +40,17 @@ namespace UE::Editor::DataStorage::Picker
 
 	void SEverythingPicker::Construct(const FArguments& InArgs)
 	{
+		SetWidthOverride(InArgs._WidthOverride);
+		SetHeightOverride(InArgs._HeightOverride);
+
+		SetMinDesiredWidth(InArgs._MinDesiredWidth);
+		SetMinDesiredHeight(InArgs._MinDesiredHeight);
+		SetMaxDesiredWidth(InArgs._MaxDesiredWidth);
+		SetMaxDesiredHeight(InArgs._MaxDesiredHeight);
+
 		ContextTabs = SNew(SHorizontalBox);
 		ContextTabViews = SNew(SWidgetSwitcher);
 
-		static const FVector2D PickerWindowSize(350.0f, 300.0f);
-
 		FMenuBuilder MenuBuilder(true, nullptr);
 
 		MenuBuilder.BeginSection(NAME_None, LOCTEXT("PickerView", "View"));
@@ -34,9 +63,13 @@ namespace UE::Editor::DataStorage::Picker
 			}
 
 			MenuContent =
-				SNew(SBox)
-				.WidthOverride(static_cast<float>(PickerWindowSize.X))
-				.HeightOverride(static_cast<float>(PickerWindowSize.Y))
+				SAssignNew(ContainerBox, SBox)
+				.WidthOverride(WidthOverride.Get())
+				.HeightOverride(HeightOverride.Get())
+				.MinDesiredWidth(MinDesiredWidth.Get())
+				.MinDesiredHeight(MinDesiredHeight.Get())
+				.MaxDesiredWidth(MaxDesiredWidth.Get())
+				.MaxDesiredHeight(MaxDesiredHeight.Get())
 				[
 					SNew(SVerticalBox)
 					+ SVerticalBox::Slot()
@@ -125,6 +158,66 @@ namespace UE::Editor::DataStorage::Picker
 		ActiveContextId = ContextId;
 		ContextTabViews->SetActiveWidgetIndex(ActiveContextId);
 	}
+
+	void SEverythingPicker::SetWidthOverride(TAttribute<FOptionalSize> InWidthOverride)
+	{
+		WidthOverride.Assign(*this, InWidthOverride);
+
+		if (ContainerBox)
+		{
+			ContainerBox->SetWidthOverride(WidthOverride.Get());
+		}
+	}
+
+	void SEverythingPicker::SetHeightOverride(TAttribute<FOptionalSize> InHeightOverride)
+	{
+		HeightOverride.Assign(*this, InHeightOverride);
+
+		if (ContainerBox)
+		{
+			ContainerBox->SetHeightOverride(HeightOverride.Get());
+		}
+	}
+
+	void SEverythingPicker::SetMinDesiredWidth(TAttribute<FOptionalSize> InMinDesiredWidth)
+	{
+		MinDesiredWidth.Assign(*this, InMinDesiredWidth);
+
+		if (ContainerBox)
+		{
+			ContainerBox->SetMinDesiredWidth(MinDesiredWidth.Get());
+		}
+	}
+
+	void SEverythingPicker::SetMinDesiredHeight(TAttribute<FOptionalSize> InMinDesiredHeight)
+	{
+		MinDesiredHeight.Assign(*this, InMinDesiredHeight);
+
+		if (ContainerBox)
+		{
+			ContainerBox->SetMinDesiredHeight(MinDesiredHeight.Get());
+		}
+	}
+
+	void SEverythingPicker::SetMaxDesiredWidth(TAttribute<FOptionalSize> InMaxDesiredWidth)
+	{
+		MaxDesiredWidth.Assign(*this, InMaxDesiredWidth);
+
+		if (ContainerBox)
+		{
+			ContainerBox->SetMaxDesiredWidth(MaxDesiredWidth.Get());
+		}
+	}
+
+	void SEverythingPicker::SetMaxDesiredHeight(TAttribute<FOptionalSize> InMaxDesiredHeight)
+	{
+		MaxDesiredHeight.Assign(*this, InMaxDesiredHeight);
+
+		if (ContainerBox)
+		{
+			ContainerBox->SetMaxDesiredHeight(MaxDesiredHeight.Get());
+		}
+	}
 }
 
 #undef LOCTEXT_NAMESPACE // "SEverythingPicker"
\ No newline at end of file
diff --git a/Engine/Plugins/Experimental/EditorDataStorageFeatures/Source/TedsEverythingPicker/Public/Widgets/SEverythingPicker.h b/Engine/Plugins/Experimental/EditorDataStorageFeatures/Source/TedsEverythingPicker/Public/Widgets/SEverythingPicker.h
--- a/Engine/Plugins/Experimental/EditorDataStorageFeatures/Source/TedsEverythingPicker/Public/Widgets/SEverythingPicker.h
+++ b/Engine/Plugins/Experimental/EditorDataStorageFeatures/Source/TedsEverythingPicker/Public/Widgets/SEverythingPicker.h
@@ -8,6 +8,7 @@
 
 #include "Context/TedsPickerContext.h"
 
+class SBox;
 class SWidgetSwitcher;
 class SHorizontalBox;
 
@@ -15,12 +16,35 @@ namespace UE::Editor::DataStorage::Picker
 {
 	class SEverythingPicker : public SCompoundWidget
 	{
+		SLATE_DECLARE_WIDGET_API(SEverythingPicker, SCompoundWidget, TEDSEVERYTHINGPICKER_API)
+
 	public:
 		SLATE_BEGIN_ARGS(SEverythingPicker)
+			: _WidthOverride(FOptionalSize())
+			, _HeightOverride(FOptionalSize())
+			, _MinDesiredWidth(FOptionalSize())
+			, _MinDesiredHeight(FOptionalSize())
+			, _MaxDesiredWidth(FOptionalSize())
+			, _MaxDesiredHeight(FOptionalSize())
 			{}
 			SLATE_SLOT_ARGUMENT(FPickerContext, Contexts)
+
+			SLATE_ATTRIBUTE(FOptionalSize, WidthOverride)
+
+			SLATE_ATTRIBUTE(FOptionalSize, HeightOverride)
+
+			SLATE_ATTRIBUTE(FOptionalSize, MinDesiredWidth)
+
+			SLATE_ATTRIBUTE(FOptionalSize, MinDesiredHeight)
+
+			SLATE_ATTRIBUTE(FOptionalSize, MaxDesiredWidth)
+
+			SLATE_ATTRIBUTE(FOptionalSize, MaxDesiredHeight)
 		SLATE_END_ARGS()
 
+		TEDSEVERYTHINGPICKER_API SEverythingPicker();
+		TEDSEVERYTHINGPICKER_API virtual ~SEverythingPicker();
+
 		TEDSEVERYTHINGPICKER_API void Construct(const FArguments& InArgs);
 
 		TEDSEVERYTHINGPICKER_API void AddContext(const FPickerContext::FSlotArguments& SlotArgs);
@@ -34,13 +58,45 @@ namespace UE::Editor::DataStorage::Picker
 		{
 			return MoveTemp(InContext);
 		}
+
+		/** See WidthOverride attribute */
+		TEDSEVERYTHINGPICKER_API void SetWidthOverride(TAttribute<FOptionalSize> InWidthOverride);
+
+		/** See HeightOverride attribute */
+		TEDSEVERYTHINGPICKER_API void SetHeightOverride(TAttribute<FOptionalSize> InHeightOverride);
+
+		/** See MinDesiredWidth attribute */
+		TEDSEVERYTHINGPICKER_API void SetMinDesiredWidth(TAttribute<FOptionalSize> InMinDesiredWidth);
+
+		/** See MinDesiredHeight attribute */
+		TEDSEVERYTHINGPICKER_API void SetMinDesiredHeight(TAttribute<FOptionalSize> InMinDesiredHeight);
+
+		/** See MaxDesiredWidth attribute */
+		TEDSEVERYTHINGPICKER_API void SetMaxDesiredWidth(TAttribute<FOptionalSize> InMaxDesiredWidth);
+
+		/** See MaxDesiredHeight attribute */
+		TEDSEVERYTHINGPICKER_API void SetMaxDesiredHeight(TAttribute<FOptionalSize> InMaxDesiredHeight);
+
 	private:
 
 		ECheckBoxState IsContextTabSelected(int32 ContextId) const;
 		void OnActiveContextChanged(ECheckBoxState State, int32 ContextId);
 
-		int32 ActiveContextId;
+		int32 ActiveContextId = 0;
+		TSharedPtr<SBox> ContainerBox;
 		TSharedPtr<SHorizontalBox> ContextTabs;
 		TSharedPtr<SWidgetSwitcher> ContextTabViews;
+
+		TSlateAttribute<FOptionalSize> WidthOverride;
+
+		TSlateAttribute<FOptionalSize> HeightOverride;
+
+		TSlateAttribute<FOptionalSize> MinDesiredWidth;
+
+		TSlateAttribute<FOptionalSize> MinDesiredHeight;
+
+		TSlateAttribute<FOptionalSize> MaxDesiredWidth;
+
+		TSlateAttribute<FOptionalSize> MaxDesiredHeight;
 	};
 }
\ No newline at end of file
