[TEDS Viewers] Add Settings options to the CompositeHierarchyViewer and add the option for expand/collapse hierarchy and show filtered out parent rows.

#Synced-CL 50483943

diff --git a/Engine/Plugins/Experimental/EditorDataStorageFeatures/Source/TedsTableViewer/Private/QueryStack/AddParentHierarchyRowsNode.cpp b/Engine/Plugins/Experimental/EditorDataStorageFeatures/Source/TedsTableViewer/Private/QueryStack/AddParentHierarchyRowsNode.cpp
new file mode 100644
--- /dev/null
+++ b/Engine/Plugins/Experimental/EditorDataStorageFeatures/Source/TedsTableViewer/Private/QueryStack/AddParentHierarchyRowsNode.cpp
@@ -0,0 +1,77 @@
+// Copyright Epic Games, Inc. All Rights Reserved.
+
+#include "QueryStack/AddParentHierarchyRowsNode.h"
+
+#include "HierarchyViewerIntefaces.h"
+
+namespace UE::Editor::DataStorage::QueryStack
+{
+	FAddParentHierarchyRowNode::FAddParentHierarchyRowNode(const ICoreProvider* InStorage, TSharedPtr<IHierarchyViewerDataInterface> InHierarchyData,
+		TSharedPtr<IRowNode> InParent)
+			: Storage(InStorage)
+			, HierarchyData(MoveTemp(InHierarchyData))
+			, Parent(MoveTemp(InParent))
+	{
+		checkf(HierarchyData, TEXT("This node requires a valid hierarchy to view"));
+		UpdateRows();
+	}
+
+	FRowHandleArrayView FAddParentHierarchyRowNode::GetRows() const
+	{
+		return Rows.GetRows();
+	}
+
+	FRowHandleArray& FAddParentHierarchyRowNode::GetMutableRows()
+	{
+		return Rows;
+	}
+
+	INode::RevisionId FAddParentHierarchyRowNode::GetRevision() const
+	{
+		return Revision;
+	}
+
+	void FAddParentHierarchyRowNode::Update(FTimespan AllottedTime)
+	{
+		if (Parent->GetRevision() != CachedParentRevision)
+		{
+			UpdateRows();
+			CachedParentRevision = Parent->GetRevision();
+			Revision++;
+		}
+	}
+
+	void FAddParentHierarchyRowNode::VisitParents(ParentListCallback Callback)
+	{
+		Callback(Parent);
+	}
+
+	void FAddParentHierarchyRowNode::UpdateRows()
+	{
+		TRACE_CPUPROFILER_EVENT_SCOPE_TEXT(TEXT("[TEDS] FAddParentHierarchyRowNode Update"));
+		const FRowHandleArrayView ParentRows = Parent->GetRows();
+		
+		Rows.Empty();
+		Rows.Append(ParentRows);
+		Rows.Sort();
+
+		for (const RowHandle Row : ParentRows)
+		{
+			RowHandle ParentRow = HierarchyData->GetParent(*Storage, Row);
+			while (ParentRow != InvalidRowHandle)
+			{
+				// Since the Rows are sorted, the search for contains will use a binary search
+				if(!Rows.Contains(ParentRow))
+				{
+					Rows.Add(ParentRow);
+					ParentRow = HierarchyData->GetParent(*Storage, ParentRow);
+				}
+				else
+				{
+					// Break if the Rows already contain the parent row as that means it has already been searched or it will be in the for-loop
+					break;
+				}
+			}
+		}
+	}
+}
diff --git a/Engine/Plugins/Experimental/EditorDataStorageFeatures/Source/TedsTableViewer/Private/QueryStack/AddParentHierarchyRowsNode.h b/Engine/Plugins/Experimental/EditorDataStorageFeatures/Source/TedsTableViewer/Private/QueryStack/AddParentHierarchyRowsNode.h
new file mode 100644
--- /dev/null
+++ b/Engine/Plugins/Experimental/EditorDataStorageFeatures/Source/TedsTableViewer/Private/QueryStack/AddParentHierarchyRowsNode.h
@@ -0,0 +1,46 @@
+// Copyright Epic Games, Inc. All Rights Reserved.
+
+#pragma once
+
+#include "TedsQueryStackInterfaces.h"
+#include "Elements/Interfaces/TypedElementDataStorageInterface.h"
+
+#define UE_API TEDSTABLEVIEWER_API
+
+namespace UE::Editor::DataStorage
+{
+	class IHierarchyViewerDataInterface;
+}
+
+namespace UE::Editor::DataStorage::QueryStack
+{
+	// A custom query stack node to add any missing parents of the given node rows in the hierarchy viewer
+	// Ex: Filter matches child row and moves that child row to the outer hierarchy level since it doesn't have a parent
+	class FAddParentHierarchyRowNode : public IRowNode
+	{
+	public:
+		
+		UE_API FAddParentHierarchyRowNode(const ICoreProvider* InStorage, TSharedPtr<IHierarchyViewerDataInterface> InHierarchyData,
+			TSharedPtr<IRowNode> InParent);
+		virtual ~FAddParentHierarchyRowNode() override = default;
+		
+		/** Retrieve access to the rows used by this node. */
+		virtual FRowHandleArrayView GetRows() const override;
+		virtual FRowHandleArray& GetMutableRows() override;
+		virtual RevisionId GetRevision() const override;
+		virtual void VisitParents(ParentListCallback Callback) override;
+		
+	protected:
+		virtual void Update(FTimespan AllottedTime) override;
+		void UpdateRows();
+	protected:
+		RevisionId Revision = UninitializedRevisionId;
+		const ICoreProvider* Storage = nullptr;
+		TSharedPtr<IHierarchyViewerDataInterface> HierarchyData;
+		TSharedPtr<IRowNode> Parent;
+		FRowHandleArray Rows;
+		RevisionId CachedParentRevision = UninitializedRevisionId;
+	};
+}
+
+#undef UE_API
\ No newline at end of file
diff --git a/Engine/Plugins/Experimental/EditorDataStorageFeatures/Source/TedsTableViewer/Private/Widgets/Composite/STedsCompositeHierarchyViewer.cpp b/Engine/Plugins/Experimental/EditorDataStorageFeatures/Source/TedsTableViewer/Private/Widgets/Composite/STedsCompositeHierarchyViewer.cpp
--- a/Engine/Plugins/Experimental/EditorDataStorageFeatures/Source/TedsTableViewer/Private/Widgets/Composite/STedsCompositeHierarchyViewer.cpp
+++ b/Engine/Plugins/Experimental/EditorDataStorageFeatures/Source/TedsTableViewer/Private/Widgets/Composite/STedsCompositeHierarchyViewer.cpp
@@ -4,6 +4,8 @@
 
 #include "DragAndDrop/Widgets/WidgetDropHandler.h"
 #include "TedsQueryStackInterfaces.h"
+#include "TedsRowMergeNode.h"
+#include "QueryStack/AddParentHierarchyRowsNode.h"
 #include "Widgets/STedsFilterBar.h"
 #include "Widgets/STedsSearchBox.h"
 #include "Widgets/Input/SComboButton.h"
@@ -13,53 +15,64 @@
 namespace UE::Editor::DataStorage
 {
 
-	FTedsCompositeHierarchyViewerShowOption::FTedsCompositeHierarchyViewerShowOption(const FText& InShowSettingTitle, const FText& InShowSettingTitleToolTip, bool bInActive)
-		: ShowOptionTitle(InShowSettingTitle)
-		, ShowOptionToolTip(InShowSettingTitleToolTip)
+	FTedsCompositeHierarchyViewerOption::FTedsCompositeHierarchyViewerOption(const FText& InSettingTitle, const FText& InSettingTitleToolTip, EOptionCategories InOptionCategory, bool bInActive)
+		: OptionTitle(InSettingTitle)
+		, OptionToolTip(InSettingTitleToolTip)
+		, OptionCategory(InOptionCategory)
 		, bActive(bInActive)
 	{
 	}
 
-	bool FTedsCompositeHierarchyViewerShowOption::IsActive() const
+	bool FTedsCompositeHierarchyViewerOption::IsActive() const
 	{
 		return bActive;
 	}
 
-	void FTedsCompositeHierarchyViewerShowOption::ToggleIsActive()
+	void FTedsCompositeHierarchyViewerOption::ToggleIsActive()
 	{
 		bActive = !bActive;
 		OnToggleEvent.Broadcast(bActive);
 	}
 
-	void FTedsCompositeHierarchyViewerShowOption::AddMenu(FMenuBuilder& InMenuBuilder)
+	void FTedsCompositeHierarchyViewerOption::AddMenu(FMenuBuilder& InMenuBuilder)
 	{
-		TSharedRef<FTedsCompositeHierarchyViewerShowOption> SharedThis = AsShared();
+		TSharedRef<FTedsCompositeHierarchyViewerOption> SharedThis = AsShared();
 
 		InMenuBuilder.AddMenuEntry(
-			ShowOptionTitle,
-			ShowOptionToolTip,
+			OptionTitle,
+			OptionToolTip,
 			FSlateIcon(),
 			FUIAction(
-				FExecuteAction::CreateSP(SharedThis, &FTedsCompositeHierarchyViewerShowOption::ToggleIsActive),
+				FExecuteAction::CreateSP(SharedThis, &FTedsCompositeHierarchyViewerOption::ToggleIsActive),
 				FCanExecuteAction(),
-				FIsActionChecked::CreateSP(SharedThis, &FTedsCompositeHierarchyViewerShowOption::IsActive)
+				FIsActionChecked::CreateSP(SharedThis, &FTedsCompositeHierarchyViewerOption::IsActive)
 			),
 			NAME_None,
 			EUserInterfaceActionType::ToggleButton
 		);
 	}
 
+	FTedsCompositeHierarchyViewerOption::EOptionCategories FTedsCompositeHierarchyViewerOption::GetCategory() const
+	{
+		return OptionCategory;
+	}
+
 	void STedsCompositeHierarchyViewer::Construct(const FArguments& InArgs, TSharedPtr<IHierarchyViewerDataInterface> InHierarchyInterface)
 	{
-		ShowOptions = InArgs._ShowOptions;
+		Options = InArgs._Options;
 
 		TSharedPtr<SVerticalBox> ContentWidget = SNew(SVerticalBox);
 
 		const bool bFilteringEnabled = InArgs._EnableFiltering;
 		const bool bSearchingEnabled = InArgs._EnableSearching;
+		const bool bSettingsEnabled = InArgs._EnableSettings;
 
+		bShowFilteredOutParentHierarchy = InArgs._ShowFilteredParentHierarchy;
+		
 		SHierarchyViewer::FArguments HierarchyViewerArgs = InArgs._HierarchyViewerArgs;
-		TSharedPtr<QueryStack::IRowNode> QueryStackNode = HierarchyViewerArgs._AllNodeProvider;
+
+		QueryNode = HierarchyViewerArgs._AllNodeProvider;
+		CompositeQueryNode = QueryNode;
 		
 		if(bFilteringEnabled || bSearchingEnabled)
 		{
@@ -69,18 +82,23 @@ namespace UE::Editor::DataStorage
 			if (bSearchingEnabled)
 			{
 				SearchBox = SNew(STedsSearchBox)
-					.InSearchableRowNode(QueryStackNode)
+					.InSearchableRowNode(CompositeQueryNode)
 					.OutSearchNode(&SearchNode);
 
 				if (SearchNode.IsValid())
 				{
-					QueryStackNode = SearchNode;
+					// If filtering is enabled, we want to add back the parent rows on that node instead (if enabled) since it is later down the chain
+					CompositeQueryNode = SearchNode;
+					if (!bFilteringEnabled)
+					{
+						AddFilteredParentRowsToQueryStack(&CompositeQueryNode, InHierarchyInterface);
+					}
 				}
 			}
 			if (bFilteringEnabled)
 			{
 				FilterBar = SNew(STedsFilterBar)
-					.InFilterableRowNode(QueryStackNode)
+					.InFilterableRowNode(CompositeQueryNode)
 					.OutFilteredNode(&FilterNode)
 					.CommonSectionFilters(InArgs._CommonSectionFilters)
 					.Filters(InArgs._Filters)
@@ -92,13 +110,22 @@ namespace UE::Editor::DataStorage
 						{
 							// Forward the updated FilterNode to the Model QueryStack since it holds a static copy to the old filter node
 							// Only called when the filters have changed, meaning the query stack has to be recomputed
-							HierarchyViewer->SetQueryStack(FilterNode);
+							CompositeQueryNode = FilterNode;
+							AddFilteredParentRowsToQueryStack(&CompositeQueryNode, HierarchyViewer->HierarchyInterface);
+							HierarchyViewer->SetQueryStack(CompositeQueryNode);
 						}
 					});
 
 				if (FilterNode.IsValid())
 				{
-					QueryStackNode = FilterNode;
+					// If searching is disabled, there is no reason to add filtered parent rows since on init there should be nothing
+					// filtering out rows (filters cannot init enabled yet), when a filter is enabled, OnPostFiltersChanged will take
+					// care of adding the FilteredParentNode
+					CompositeQueryNode = FilterNode;
+					if (!bSearchingEnabled)
+					{
+						AddFilteredParentRowsToQueryStack(&CompositeQueryNode, InHierarchyInterface);
+					}
 				}
 			}
 
@@ -144,25 +171,28 @@ namespace UE::Editor::DataStorage
 				];
 			}
 
-			// View mode combo button
-			FilterSearchBar->AddSlot()
-			.VAlign(VAlign_Center)
-			.AutoWidth()
-			[
-				SAssignNew(SettingsComboButton, SComboButton)
-					.ComboButtonStyle(FAppStyle::Get(), "SimpleComboButtonWithIcon") // Use the tool bar item style for this button
+			if (bSettingsEnabled)
+			{
+				// View mode combo button
+				FilterSearchBar->AddSlot()
+				.VAlign(VAlign_Center)
+				.AutoWidth()
+				[
+					SAssignNew(SettingsComboButton, SComboButton)
+					.ComboButtonStyle(FAppStyle::Get(), "SimpleComboButtonWithIcon") // Use the toolbar item style for this button
 					.OnGetMenuContent(this, &STedsCompositeHierarchyViewer::GetSettingsButtonContent)
 					.HasDownArrow(false)
 					.ButtonContent()
 					[
 						SNew(SImage)
-							.ColorAndOpacity(FSlateColor::UseForeground())
-							.Image(FAppStyle::Get().GetBrush("Icons.Settings"))
+						.ColorAndOpacity(FSlateColor::UseForeground())
+						.Image(FAppStyle::Get().GetBrush("Icons.Settings"))
 					]
-			];
+				];
+			}
 		}
 
-		HierarchyViewerArgs.AllNodeProvider(QueryStackNode);
+		HierarchyViewerArgs.AllNodeProvider(CompositeQueryNode);
 		SAssignNew(HierarchyViewer, SHierarchyViewer, MoveTemp(InHierarchyInterface))
 			.AllNodeProvider(HierarchyViewerArgs._AllNodeProvider)
 			.TableViewerIdentifier(HierarchyViewerArgs._TableViewerIdentifier)
@@ -190,6 +220,9 @@ namespace UE::Editor::DataStorage
 
 	void STedsCompositeHierarchyViewer::SetQueryStack(TSharedPtr<QueryStack::IRowNode> InRowQueryStack)
 	{
+		// SetQueryStack is not yet supported in the composite viewer to recreate the search/filter query stack
+		// The call is available from the base class ptr, but will only display the given QueryStack with no composite functionality
+		ensureMsgf(false, TEXT("SetQueryStack is not supported in the Composite Viewer yet."));
 		HierarchyViewer->SetQueryStack(InRowQueryStack);
 	}
 
@@ -316,13 +349,91 @@ namespace UE::Editor::DataStorage
 		HierarchyViewer->SetDropHandler(MoveTemp(InDropHandler));
 	}
 
+	void STedsCompositeHierarchyViewer::AddFilteredParentRowsToQueryStack(TSharedPtr<QueryStack::IRowNode>* QueryStack, const TSharedPtr<IHierarchyViewerDataInterface>& HierarchyInterface)
+	{
+		if (bShowFilteredOutParentHierarchy)
+		{
+			ICoreProvider* Storage = nullptr;
+			if (HierarchyViewer.IsValid())
+			{
+				if (const TSharedPtr<FTedsTableViewerModel> Model = HierarchyViewer->Model)
+				{
+					Storage = Model->GetDataStorageInterface();
+				}
+			}
+			else
+			{
+				Storage = GetMutableDataStorageFeature<ICoreProvider>(StorageFeatureName);
+			}
+
+			if(Storage)
+			{
+				TSharedPtr<QueryStack::IRowNode> HierarchyNode =
+					MakeShared<QueryStack::FAddParentHierarchyRowNode>(
+						Storage,
+						HierarchyInterface,
+						*QueryStack);
+				// Merge with a repeating approach with the original query node so we can make sure that we aren't adding back parents 
+				// that weren't in the original query given to the viewer.
+				*QueryStack = MakeShared<QueryStack::FRowMergeNode>(
+					MakeConstArrayView({HierarchyNode, QueryNode}),
+					QueryStack::FRowMergeNode::EMergeApproach::Repeating);
+			}
+		}
+	}
+
 	TSharedRef<SWidget> STedsCompositeHierarchyViewer::GetSettingsButtonContent()
 	{
+		TArray<TSharedPtr<FTedsCompositeHierarchyViewerOption>> HierarchyOptions;
+		TArray<TSharedPtr<FTedsCompositeHierarchyViewerOption>> ShowOptions;
+		HierarchyOptions.Reserve(Options.Num());
+		ShowOptions.Reserve(Options.Num());
+		
+		for (TSharedPtr<FTedsCompositeHierarchyViewerOption>& Option : Options)
+		{
+			switch (Option->GetCategory())
+			{
+				case FTedsCompositeHierarchyViewerOption::EOptionCategories::Hierarchy:
+					HierarchyOptions.Add(Option);
+					break;
+				case FTedsCompositeHierarchyViewerOption::EOptionCategories::Show:
+					ShowOptions.Add(Option);
+					break;
+				default:
+					break;
+			}
+		}
+		
 		FMenuBuilder MenuBuilder(true, nullptr);
+		// HIERARCHY
+		MenuBuilder.BeginSection("Hierarchy", LOCTEXT("HierarchyHeading", "Hierarchy"));
+		// Default Hierarchy Actions
+		{
+			MenuBuilder.AddMenuEntry(
+				LOCTEXT("ExpandAll", "Expand All"),
+				LOCTEXT("ExpandAllTooltip", "Expand all items in the Hierarchy."),
+				FSlateIcon(),
+				FUIAction( FExecuteAction::CreateSP(this, &STedsCompositeHierarchyViewer::ExpandAll) )
+				);
+			MenuBuilder.AddMenuEntry(
+				LOCTEXT("CollapseAll", "Collapse All"),
+				LOCTEXT("CollapseAllTooltip", "Collapse all items in the Hierarchy."),
+				FSlateIcon(),
+				FUIAction( FExecuteAction::CreateSP(this, &STedsCompositeHierarchyViewer::CollapseAll) )
+				);
+		}
+		
+		for (const TSharedPtr<FTedsCompositeHierarchyViewerOption>& HierarchyOption : HierarchyOptions)
+		{
+			HierarchyOption->AddMenu(MenuBuilder);
+		}
 
+		MenuBuilder.EndSection();
+		
+		// SHOW
 		MenuBuilder.BeginSection("Show", LOCTEXT("ShowHeading", "Show"));
 
-		for (TSharedPtr<FTedsCompositeHierarchyViewerShowOption>& ShowOption : ShowOptions)
+		for (const TSharedPtr<FTedsCompositeHierarchyViewerOption>& ShowOption : ShowOptions)
 		{
 			ShowOption->AddMenu(MenuBuilder);
 		}
@@ -332,7 +443,26 @@ namespace UE::Editor::DataStorage
 		return MenuBuilder.MakeWidget();
 	}
 
-} // namespace UE::Editor::DataStorage
+	void STedsCompositeHierarchyViewer::OnToggleShowFilteredParentHierarchy()
+	{
+		bShowFilteredOutParentHierarchy = !bShowFilteredOutParentHierarchy;
 
+		if (FilterNode.IsValid())
+		{
+			CompositeQueryNode = FilterNode;
+		}
+		else if (SearchNode.IsValid())
+		{
+			CompositeQueryNode = SearchNode;
+		}
+		else
+		{
+			return; // If we don't have a filter or search node, there is no point in adding back the parents
+		}
+		AddFilteredParentRowsToQueryStack(&CompositeQueryNode, HierarchyViewer->HierarchyInterface);
+		HierarchyViewer->SetQueryStack(CompositeQueryNode);
+	}
+
+} // namespace UE::Editor::DataStorage
 
 #undef LOCTEXT_NAMESPACE
diff --git a/Engine/Plugins/Experimental/EditorDataStorageFeatures/Source/TedsTableViewer/Public/Widgets/Composite/STedsCompositeHierarchyViewer.h b/Engine/Plugins/Experimental/EditorDataStorageFeatures/Source/TedsTableViewer/Public/Widgets/Composite/STedsCompositeHierarchyViewer.h
--- a/Engine/Plugins/Experimental/EditorDataStorageFeatures/Source/TedsTableViewer/Public/Widgets/Composite/STedsCompositeHierarchyViewer.h
+++ b/Engine/Plugins/Experimental/EditorDataStorageFeatures/Source/TedsTableViewer/Public/Widgets/Composite/STedsCompositeHierarchyViewer.h
@@ -19,14 +19,19 @@ namespace UE::Editor::DataStorage
 	class FTedsTableViewerColumn;
 
 	/**
-	* A class corresponding to one "show" boolean option (e.g. "Show Components"). An array of these objects can be passed into the 
+	* A class corresponding to one boolean option (e.g. "Show Components"). An array of these objects can be passed into the 
 	* STedsCompositeHierarchyViewer and it will auto-populate a menu containing toggle buttons for each object.
 	*/
-	class FTedsCompositeHierarchyViewerShowOption : public TSharedFromThis<FTedsCompositeHierarchyViewerShowOption>
+	class FTedsCompositeHierarchyViewerOption : public TSharedFromThis<FTedsCompositeHierarchyViewerOption>
 	{
 	public:
+		enum class EOptionCategories
+		{
+			Hierarchy,
+			Show
+		};
 
-		TEDSTABLEVIEWER_API FTedsCompositeHierarchyViewerShowOption(const FText& InShowSettingTitle, const FText& InShowSettingTitleToolTip, bool bInActive);
+		TEDSTABLEVIEWER_API FTedsCompositeHierarchyViewerOption(const FText& InSettingTitle, const FText& InSettingTitleToolTip, EOptionCategories InOptionCategory, bool bInActive);
 
 		TEDSTABLEVIEWER_API bool IsActive() const;
 
@@ -34,13 +39,16 @@ namespace UE::Editor::DataStorage
 
 		TEDSTABLEVIEWER_API void AddMenu(FMenuBuilder& InMenuBuilder);
 
-		DECLARE_EVENT_OneParam(FTedsCompositeHierarchyViewerShowOption, FOnToggle, bool);
+		TEDSTABLEVIEWER_API EOptionCategories GetCategory() const;
+
+		DECLARE_EVENT_OneParam(FTedsCompositeHierarchyViewerOption, FOnToggle, bool);
 		FOnToggle& OnToggle() { return OnToggleEvent; }
 
 	private:
 
-		FText ShowOptionTitle;
-		FText ShowOptionToolTip;
+		FText OptionTitle;
+		FText OptionToolTip;
+		EOptionCategories OptionCategory;
 		bool bActive;
 
 		FOnToggle OnToggleEvent;
@@ -68,43 +76,52 @@ namespace UE::Editor::DataStorage
 			: _HierarchyViewerArgs()
 			, _EnableSearching(true)
 			, _EnableFiltering(true)
+			, _EnableSettings(true)
 			, _UseSectionsForCategories(false)
+			, _ShowFilteredParentHierarchy(true)
 		{
 		}
 
 		// Arguments used to create the Hierarchy Viewer
 		SLATE_ARGUMENT(SHierarchyViewer::FArguments, HierarchyViewerArgs)
-		
+			
 		// If we want to enable the search bar on the table
-        SLATE_ARGUMENT(bool, EnableSearching)
+		SLATE_ARGUMENT(bool, EnableSearching)
 
 		// If we want to enable the filter bar on the table
 		SLATE_ARGUMENT(bool, EnableFiltering)
+
+		// If we want to enable the settings cog on the table (auto-populates with a few default hierarchy actions)
+		SLATE_ARGUMENT(bool, EnableSettings)
 		
 		/** Array of filter names to also add to the 'common filters' section (section at the top of the menu) */
-        SLATE_ARGUMENT(TArray<FName>, CommonSectionFilters)
+		SLATE_ARGUMENT(TArray<FName>, CommonSectionFilters)
 
 		// Array of filters to add to the filter bar *AND Filters
 		SLATE_ARGUMENT(TArray<FTedsFilterData>, Filters)
 
 		// Array of class filters to add to the filter bar *OR Filters
 		SLATE_ARGUMENT(TArray<UClass*>, ClassFilters)
-		
-		// Array of boolean option objects used to populate a visibility menu
-		SLATE_ARGUMENT(TArray<TSharedPtr<FTedsCompositeHierarchyViewerShowOption>>, ShowOptions)
 
 		// Whether to use submenus or sections for categories in the filter menu
 		SLATE_ARGUMENT(bool, UseSectionsForCategories)
+		
+		// Array of boolean option objects used to populate the settings menu
+		SLATE_ARGUMENT(TArray<TSharedPtr<FTedsCompositeHierarchyViewerOption>>, Options)
+
+		// Whether to show filtered out parents of a child that matches
+		SLATE_ARGUMENT(bool, ShowFilteredParentHierarchy)
 
 		SLATE_END_ARGS()
+		
+	protected:
+		// Clear the current QueryStack being displayed, set it to the given node, and recreate the sorting nodes 
+		TEDSTABLEVIEWER_API virtual void SetQueryStack(TSharedPtr<QueryStack::IRowNode> InRowQueryStack) override;
 
 	public:
 		
 		TEDSTABLEVIEWER_API void Construct(const FArguments& InArgs, TSharedPtr<IHierarchyViewerDataInterface> InHierarchyInterface);
 
-		// Clear the current QueryStack being displayed, set it to the given node, and recreate the sorting nodes 
-		TEDSTABLEVIEWER_API virtual void SetQueryStack(TSharedPtr<QueryStack::IRowNode> InRowQueryStack) override;
-
 		// Clear the current list of columns being displayed and set it to the given list
 		TEDSTABLEVIEWER_API virtual void SetColumns(const TArray<TWeakObjectPtr<const UScriptStruct>>& InColumns) override;
 		
@@ -155,18 +172,30 @@ namespace UE::Editor::DataStorage
 
 		// Expand this row, and all its parents
 		TEDSTABLEVIEWER_API void ExpandWithParents(RowHandle Row);
+
+		// Menu actions
+		TEDSTABLEVIEWER_API void OnToggleShowFilteredParentHierarchy();
 		
 		// Set a handler that this viewer will forward its drop operations to.
 		TEDSTABLEVIEWER_API void SetDropHandler(TUniquePtr<FWidgetDropHandler> InDropHandler);
 		
 	private:
+		void AddFilteredParentRowsToQueryStack(TSharedPtr<QueryStack::IRowNode>* QueryStack, const TSharedPtr<IHierarchyViewerDataInterface>& HierarchyInterface);
 
 		// Build a menu of Toggle Buttons from the ShowOptions objects
 		TSharedRef<SWidget> GetSettingsButtonContent();
+		
+		// Option objects used to construct the settings menu
+		TArray<TSharedPtr<FTedsCompositeHierarchyViewerOption>> Options;
+		
+		// Whether to add back the filtered out parent hierarchy (removed from searching or filtering)
+		bool bShowFilteredOutParentHierarchy = true;
 
+		// Original Row Node given as the AllNodeProvider of the HierarchyViewer
+		TSharedPtr<QueryStack::IRowNode> QueryNode;
 
-		// Option objects used to construct the "Show" menu
-		TArray<TSharedPtr<FTedsCompositeHierarchyViewerShowOption>> ShowOptions;
+		// Final Composite Row Node given to the Hierarchy Viewer
+		TSharedPtr<QueryStack::IRowNode> CompositeQueryNode;
 
 		// Row Node passed to the SearchBox that receives the searched result
 		TSharedPtr<QueryStack::IRowNode> SearchNode;
diff --git a/Engine/Plugins/Experimental/EditorDataStorageFeatures/Source/TedsTableViewer/Public/Widgets/STedsHierarchyViewer.h b/Engine/Plugins/Experimental/EditorDataStorageFeatures/Source/TedsTableViewer/Public/Widgets/STedsHierarchyViewer.h
--- a/Engine/Plugins/Experimental/EditorDataStorageFeatures/Source/TedsTableViewer/Public/Widgets/STedsHierarchyViewer.h
+++ b/Engine/Plugins/Experimental/EditorDataStorageFeatures/Source/TedsTableViewer/Public/Widgets/STedsHierarchyViewer.h
@@ -48,6 +48,8 @@ namespace UE::Editor::DataStorage
 	 */
 	class SHierarchyViewer : public SCompoundWidget, public ITableViewer
 	{
+	friend class STedsCompositeHierarchyViewer;
+	
 	public:
 		
 		enum class EExpansionState : uint8
