[TEDS Hierarchy Viewer] Add a OnSetExpansionRecursive function for TEDS Hierarchy Viewers. By default, this will allow shift-click on the expander arrow to expand/collapse that item and all it's children

#Synced-CL 50652112

diff --git a/Engine/Plugins/Experimental/EditorDataStorageFeatures/Source/TedsTableViewer/Private/Widgets/STedsHierarchyViewer.cpp b/Engine/Plugins/Experimental/EditorDataStorageFeatures/Source/TedsTableViewer/Private/Widgets/STedsHierarchyViewer.cpp
--- a/Engine/Plugins/Experimental/EditorDataStorageFeatures/Source/TedsTableViewer/Private/Widgets/STedsHierarchyViewer.cpp
+++ b/Engine/Plugins/Experimental/EditorDataStorageFeatures/Source/TedsTableViewer/Private/Widgets/STedsHierarchyViewer.cpp
@@ -95,6 +95,7 @@ namespace UE::Editor::DataStorage
 			.RowsSource(&Model->GetItems())
 			.OnGenerateRow(this, &SHierarchyViewer::MakeTableRowWidget)
 			.OnExpansionChanged(this, &SHierarchyViewer::OnExpansionChanged)
+			.OnSetExpansionRecursive(this, &SHierarchyViewer::OnSetExpansionRecursive)
 			.OnSelectionChanged(this, &SHierarchyViewer::OnListSelectionChanged)
 			.SelectionMode(InArgs._ListSelectionMode)
 			.AllowInvisibleItemSelection(InArgs._AllowInvisibleItemSelection);
@@ -304,6 +305,21 @@ namespace UE::Editor::DataStorage
 		}
 	}
 
+	void SHierarchyViewer::OnSetExpansionRecursive(TableViewerItemPtr Item, bool bExpanded) const
+	{
+		if (const ICoreProvider* Storage = Model->GetDataStorageInterface(); HierarchyInterface.IsValid())
+		{
+			TArray<RowHandle> ChildRowsToSet;
+			HierarchyInterface->WalkDepthFirst(*Storage, Item.RowHandle,
+			[&ChildRowsToSet](const ICoreProvider&, RowHandle, const RowHandle Target)
+			{
+				ChildRowsToSet.Add(Target);
+			});
+			
+			SetExpansionState_Internal(FRowHandleArrayView(ChildRowsToSet, FRowHandleArrayView::EFlags::IsUnique), bExpanded);
+		}
+	}
+
 	void SHierarchyViewer::OnListSelectionChanged(TableViewerItemPtr Item, ESelectInfo::Type SelectInfo)
 	{
 		if(OnSelectionChanged.IsBound())
diff --git a/Engine/Plugins/Experimental/EditorDataStorageFeatures/Source/TedsTableViewer/Private/Widgets/STedsTreeView.cpp b/Engine/Plugins/Experimental/EditorDataStorageFeatures/Source/TedsTableViewer/Private/Widgets/STedsTreeView.cpp
--- a/Engine/Plugins/Experimental/EditorDataStorageFeatures/Source/TedsTableViewer/Private/Widgets/STedsTreeView.cpp
+++ b/Engine/Plugins/Experimental/EditorDataStorageFeatures/Source/TedsTableViewer/Private/Widgets/STedsTreeView.cpp
@@ -24,6 +24,7 @@ namespace UE::Editor::DataStorage
 			.TreeItemsSource( InArgs._TopLevelRowsSource)
 			.OnGenerateRow(InArgs._OnGenerateRow)
 			.OnExpansionChanged(InArgs._OnExpansionChanged)
+			.OnSetExpansionRecursive(InArgs._OnSetExpansionRecursive)
 			.OnSelectionChanged(InArgs._OnSelectionChanged)
 			.SelectionMode(InArgs._SelectionMode)
 			.OnContextMenuOpening(Binder.BindEvent(&FWidgetContextMenuColumn::OnContextMenuOpening))
diff --git a/Engine/Plugins/Experimental/EditorDataStorageFeatures/Source/TedsTableViewer/Private/Widgets/STedsTreeView.h b/Engine/Plugins/Experimental/EditorDataStorageFeatures/Source/TedsTableViewer/Private/Widgets/STedsTreeView.h
--- a/Engine/Plugins/Experimental/EditorDataStorageFeatures/Source/TedsTableViewer/Private/Widgets/STedsTreeView.h
+++ b/Engine/Plugins/Experimental/EditorDataStorageFeatures/Source/TedsTableViewer/Private/Widgets/STedsTreeView.h
@@ -34,6 +34,9 @@ namespace UE::Editor::DataStorage
 			/** Delegate fired on expansion change */
 			SLATE_EVENT( FOnExpansionChanged, OnExpansionChanged )
 
+			/** Delegate fired on a request to change expansion on an item and it's children */
+			SLATE_EVENT( FOnSetExpansionRecursive, OnSetExpansionRecursive )
+
 			/** Delegate that determines the selection mode */
 			SLATE_ATTRIBUTE( ESelectionMode::Type, SelectionMode )
 
diff --git a/Engine/Plugins/Experimental/EditorDataStorageFeatures/Source/TedsTableViewer/Public/Widgets/STedsHierarchyViewer.h b/Engine/Plugins/Experimental/EditorDataStorageFeatures/Source/TedsTableViewer/Public/Widgets/STedsHierarchyViewer.h
--- a/Engine/Plugins/Experimental/EditorDataStorageFeatures/Source/TedsTableViewer/Public/Widgets/STedsHierarchyViewer.h
+++ b/Engine/Plugins/Experimental/EditorDataStorageFeatures/Source/TedsTableViewer/Public/Widgets/STedsHierarchyViewer.h
@@ -187,6 +187,8 @@ namespace UE::Editor::DataStorage
 		void RefreshColumnWidgets();
 
 		void OnExpansionChanged(TableViewerItemPtr Item, bool bExpanded) const;
+		
+		void OnSetExpansionRecursive(TableViewerItemPtr Item, bool bExpanded) const;
 
 		void OnListSelectionChanged(TableViewerItemPtr Item, ESelectInfo::Type SelectInfo);
 
