[TEDS Outliner] Update HideFromUITag to run through the query stack and define actors and entities that shouldn't be shown in usual circumstances

#Synced-CL 50773268

diff --git a/Engine/Plugins/Experimental/EditorDataStorageFeatures/Source/TedsActorCompatibility/Private/Compatibility/TedsActorUIDisplayQueries.cpp b/Engine/Plugins/Experimental/EditorDataStorageFeatures/Source/TedsActorCompatibility/Private/Compatibility/TedsActorUIDisplayQueries.cpp
new file mode 100644
--- /dev/null
+++ b/Engine/Plugins/Experimental/EditorDataStorageFeatures/Source/TedsActorCompatibility/Private/Compatibility/TedsActorUIDisplayQueries.cpp
@@ -0,0 +1,67 @@
+// Copyright Epic Games, Inc. All Rights Reserved.
+
+#include "Compatibility/TedsActorUIDisplayQueries.h"
+
+#include "ActorEditorUtils.h"
+#include "Columns/TedsActorComponentCompatibilityColumns.h"
+#include "Elements/Columns/TypedElementCompatibilityColumns.h"
+#include "Elements/Columns/TypedElementMiscColumns.h"
+#include "Elements/Columns/TypedElementUIColumns.h"
+#include "Elements/Framework/TypedElementQueryBuilder.h"
+#include "GameFramework/Actor.h"
+
+#include UE_INLINE_GENERATED_CPP_BY_NAME(TedsActorUIDisplayQueries)
+
+void UActorUIDisplayDataStorageFactory::RegisterQueries(UE::Editor::DataStorage::ICoreProvider& DataStorage)
+{
+	RegisterSyncActorHideRowFromUITagQuery(DataStorage);
+}
+
+void UActorUIDisplayDataStorageFactory::RegisterSyncActorHideRowFromUITagQuery(UE::Editor::DataStorage::ICoreProvider& DataStorage) const
+{
+	using namespace UE::Editor::DataStorage::Queries;
+	
+	DataStorage.RegisterQuery(
+		Select(
+			TEXT("Sync Actor's Hide From UI State"),
+			FProcessor(EQueryTickPhase::PrePhysics, DataStorage.GetQueryTickGroupName(EQueryTickGroups::SyncExternalToDataStorage))
+			.SetExecutionMode(EExecutionMode::GameThread),
+			[](IQueryContext& Context, RowHandle Row, const FTypedElementUObjectColumn& Object)
+			{
+				if (FActorEditorUtils::IsActorDisplayable(Cast<AActor>(Object.Object)))
+				{
+					Context.RemoveColumns<FHideRowFromUITag>(Row);
+				}
+				else
+				{
+					Context.AddColumns<FHideRowFromUITag>(Row);
+				}
+			})
+		.Where()
+			.All<FTypedElementSyncFromWorldTag, FTypedElementActorTag>()
+		.Compile()
+	);
+	DataStorage.RegisterQuery(
+		Select(
+			TEXT("Sync Actor Component's Hide From UI State"),
+			FProcessor(EQueryTickPhase::PrePhysics, DataStorage.GetQueryTickGroupName(EQueryTickGroups::SyncExternalToDataStorage))
+			.SetExecutionMode(EExecutionMode::GameThread),
+			[](IQueryContext& Context, RowHandle Row, const FTypedElementUObjectColumn& Object)
+			{
+				if (const UActorComponent* ActorComponentInstance = Cast<UActorComponent>(Object.Object); ActorComponentInstance != nullptr)
+				{
+					if (FActorEditorUtils::IsActorDisplayable(Cast<AActor>(ActorComponentInstance->GetOuter())))
+					{
+						Context.RemoveColumns<FHideRowFromUITag>(Row);
+					}
+					else
+					{
+						Context.AddColumns<FHideRowFromUITag>(Row);
+					}
+				}
+			})
+		.Where()
+			.All<FTypedElementSyncFromWorldTag, FActorComponentTypeTag>()
+		.Compile()
+	);
+}
diff --git a/Engine/Plugins/Experimental/EditorDataStorageFeatures/Source/TedsActorCompatibility/Private/Compatibility/TedsActorUIDisplayQueries.h b/Engine/Plugins/Experimental/EditorDataStorageFeatures/Source/TedsActorCompatibility/Private/Compatibility/TedsActorUIDisplayQueries.h
new file mode 100644
--- /dev/null
+++ b/Engine/Plugins/Experimental/EditorDataStorageFeatures/Source/TedsActorCompatibility/Private/Compatibility/TedsActorUIDisplayQueries.h
@@ -0,0 +1,24 @@
+// Copyright Epic Games, Inc. All Rights Reserved.
+
+#pragma once
+
+#include "Elements/Interfaces/TypedElementDataStorageFactory.h"
+#include "UObject/ObjectMacros.h"
+
+#include "TedsActorUIDisplayQueries.generated.h"
+
+UCLASS()
+class UActorUIDisplayDataStorageFactory : public UEditorDataStorageFactory
+{
+	GENERATED_BODY()
+
+public:
+	~UActorUIDisplayDataStorageFactory() override = default;
+
+	void RegisterQueries(UE::Editor::DataStorage::ICoreProvider& DataStorage) override;
+private:
+	/**
+	 * Registers queries to check if the updated Actor row or the Component row of an Actor should be given the HideRowFromUITag
+	 */
+	void RegisterSyncActorHideRowFromUITagQuery(UE::Editor::DataStorage::ICoreProvider& DataStorage) const;
+};
diff --git a/Engine/Plugins/Experimental/EditorDataStorageFeatures/Source/TedsOutliner/Private/TedsOutlinerModule.cpp b/Engine/Plugins/Experimental/EditorDataStorageFeatures/Source/TedsOutliner/Private/TedsOutlinerModule.cpp
--- a/Engine/Plugins/Experimental/EditorDataStorageFeatures/Source/TedsOutliner/Private/TedsOutlinerModule.cpp
+++ b/Engine/Plugins/Experimental/EditorDataStorageFeatures/Source/TedsOutliner/Private/TedsOutlinerModule.cpp
@@ -12,6 +12,7 @@
 #include "Elements/Columns/TypedElementMiscColumns.h"
 #include "Elements/Columns/TypedElementPackageColumns.h"
 #include "Elements/Columns/TypedElementTypeInfoColumns.h"
+#include "Elements/Columns/TypedElementUIColumns.h"
 #include "Elements/Columns/TypedElementVisibilityColumns.h"
 #include "Columns/TedsActorComponentCompatibilityColumns.h"
 #include "Columns/TedsActorMobilityColumns.h"
@@ -25,9 +26,6 @@
 #include "Columns/ActorFolderColumns.h"
 #include "Widgets/Docking/SDockTab.h"
 
-#include "Engine/DirectionalLight.h"
-#include "Engine/StaticMeshActor.h"
-
 #define LOCTEXT_NAMESPACE "TedsOutlinerModule"
 
 namespace UE::Editor::Outliner
@@ -155,7 +153,10 @@ TSharedRef<SWidget> FTedsOutlinerModule::CreateLevelEditorTedsOutliner()
 	// The Outliner is populated with Actors and Entities
 	const FQueryDescription OutlinerQueryDescription =
 		Select()
-		.Where((TColumn<FTypedElementUObjectColumn>() || TColumn<FFolderCompatibilityColumn>()) && TColumn<FTypedElementWorldColumn>())
+		.Where()
+			.All<FTypedElementUObjectColumn>()
+			.Any<FFolderCompatibilityColumn, FTypedElementWorldColumn>()
+			.None<FHideRowFromUITag>()
 		.Compile();
 
 	FSceneOutlinerInitializationOptions InitOptions;
@@ -169,32 +170,23 @@ TSharedRef<SWidget> FTedsOutlinerModule::CreateLevelEditorTedsOutliner()
 	Params.ColumnDescription = GetLevelEditorTedsOutlinerColumnDescription();
 	
 	Params.ClassFilters.Add(AActor::StaticClass());
-	Params.ClassFilters.Add(AStaticMeshActor::StaticClass());
-	Params.ClassFilters.Add(ADirectionalLight::StaticClass());
-	
-	// Example ways to create filters using query descriptions and functions
-	// Note QueryDescriptions should always be used over QueryFunctions when possible as they are more performant
-	Params.Filters.Emplace(
-		"ActorComponents", 
-		LOCTEXT("ActorFilterDisplayName", "Actor Components"), 
-		LOCTEXT("ActorFilterToolTip", "Filters for actor components"), 
-		"ClassIcon.Actor", 
-		nullptr,
-		Storage->RegisterQuery(
-			Select()
-			.Where(TColumn<FActorComponentTypeTag>())
-			.Compile()));
 	
-	Params.Filters.Emplace(
-		"Static",
-		LOCTEXT("StaticFilterDisplayName", "Static"), 
-		LOCTEXT("StaticFilterToolTip", "Filter for Static Actors"),
-		FName(), 
-		nullptr,
-		[](TQueryContext<SingleRowInfo> Context, const FTedsActorMobilityColumn& MobilityColumn)
-		{
-			return MobilityColumn.Mobility == EComponentMobility::Static;
-		});
+	/** Example ways to create filters using query descriptions and functions
+	 * Note QueryDescriptions should always be used over QueryFunctions when possible as they are more performant
+	 * Params.Filters.Emplace(
+	 * 	...
+	 *	Storage->RegisterQuery(
+	 *		Select()
+	 *		.Where(TColumn<FActorComponentTypeTag>())
+	 *		.Compile()));
+	 *
+	 * Params.Filters.Emplace(
+	 *	...
+	 *	[](TQueryContext<SingleRowInfo> Context, const FTedsActorMobilityColumn& MobilityColumn)
+	 *	{
+	 *		return MobilityColumn.Mobility == EComponentMobility::Static;
+	 *	});
+	 */
 		
 	// Empty selection set name is currently the level editor
 	Params.SelectionSetOverride = FName();
@@ -244,7 +236,7 @@ FName FTedsOutlinerModule::GetTedsOutlinerTabName()
 {
 	return TedsOutlinerTabName;
 }
-} // namsepace UE::Editor::Outliner
+} // namespace UE::Editor::Outliner
 
 IMPLEMENT_MODULE(UE::Editor::Outliner::FTedsOutlinerModule, TedsOutliner);
 
diff --git a/Engine/Source/Editor/SceneOutliner/Private/ActorMode.cpp b/Engine/Source/Editor/SceneOutliner/Private/ActorMode.cpp
--- a/Engine/Source/Editor/SceneOutliner/Private/ActorMode.cpp
+++ b/Engine/Source/Editor/SceneOutliner/Private/ActorMode.cpp
@@ -25,6 +25,7 @@
 #include "LevelInstance/LevelInstanceSubsystem.h"
 #include "LevelInstance/LevelInstanceEditorInstanceActor.h"
 #include "LevelEditor.h"
+#include "SceneOutlinerHelpers.h"
 #include "Logging/MessageLog.h"
 #include "DragAndDrop/ActorDragDropGraphEdOp.h"
 #include "DragAndDrop/FolderDragDropOp.h"
@@ -370,34 +371,24 @@ FFolder::FRootObject FActorMode::GetPasteTargetRootObject() const
 
 bool FActorMode::IsActorDisplayable(const SSceneOutliner* SceneOutliner, const AActor* Actor, bool bShowLevelInstanceContent)
 {
-	bool bIsActorDisplayable =	Actor &&
-								!SceneOutliner->GetSharedData().bOnlyShowFolders &&		// Don't show actors if we're only showing folders
-								Actor->IsEditable() &&									// Only show actors that are allowed to be selected and drawn in editor
-								Actor->IsListedInSceneOutliner();
+	const bool bIsActorDisplayable = 	FActorEditorUtils::IsActorDisplayable(Actor) &&
+										!SceneOutliner->GetSharedData().bOnlyShowFolders;		// Don't show actors if we're only showing folders
 	
-	if(bIsActorDisplayable)
+	if(bIsActorDisplayable && Actor->HasAnyFlags(RF_Transient))
 	{
-		if (Actor->HasAnyFlags(RF_Transient))
+		// Level Instance transient actors are shown based on passed in bShowLevelInstanceContent flag
+		if (Actor->IsInLevelInstance())
 		{
-			// Level Instance transient actors are shown based on passed in bShowLevelInstanceContent flag
-			if (Actor->IsInLevelInstance())
-			{
-				bIsActorDisplayable = bShowLevelInstanceContent;
-			}
-			else
-			{
-				// Don't show transient actors in non-play worlds, except if bShowTransient is true
-				bIsActorDisplayable = SceneOutliner->GetSharedData().bShowTransient || (Actor->GetWorld() && Actor->GetWorld()->IsPlayInEditor());
-			}
+			return bShowLevelInstanceContent;
+		}
+		else
+		{
+			// Don't show transient actors in non-play worlds, except if bShowTransient is true
+			return SceneOutliner->GetSharedData().bShowTransient || (Actor->GetWorld() && Actor->GetWorld()->IsPlayInEditor());
 		}
 	}
 
-	return	bIsActorDisplayable &&																	// Previous results
-			!Actor->IsTemplate() &&																	// Should never happen, but we never want CDOs displayed
-			!FActorEditorUtils::IsABuilderBrush(Actor) &&											// Don't show the builder brush
-			!Actor->IsA(AWorldSettings::StaticClass()) &&											// Don't show the WorldSettings actor, even though it is technically editable
-			IsValidChecked(Actor) &&																// We don't want to show actors that are about to go away
-			FLevelUtils::IsLevelVisible(Actor->GetLevel());											// Only show Actors whose level is visible
+	return	bIsActorDisplayable;
 }
 
 bool FActorMode::CanInteract(const ISceneOutlinerTreeItem& Item) const
diff --git a/Engine/Source/Runtime/Engine/Private/ActorEditorUtils.cpp b/Engine/Source/Runtime/Engine/Private/ActorEditorUtils.cpp
--- a/Engine/Source/Runtime/Engine/Private/ActorEditorUtils.cpp
+++ b/Engine/Source/Runtime/Engine/Private/ActorEditorUtils.cpp
@@ -2,11 +2,13 @@
 
 #include "ActorEditorUtils.h"
 #include "Engine/Level.h"
-#include "Engine/Level.h"
 #include "Engine/World.h"
 #include "Engine/Brush.h"
+#include "LevelUtils.h"
 
 #if WITH_EDITOR
+#include "GameFramework/Actor.h"
+#include "GameFramework/WorldSettings.h"
 #include "Kismet2/ComponentEditorUtils.h"
 #endif
 
@@ -149,6 +151,19 @@ namespace FActorEditorUtils
 
 		return true;
 	}
+#if WITH_EDITOR
+	bool IsActorDisplayable(const AActor* Actor)
+	{
+		return Actor &&
+			   IsValid(Actor) &&										// We don't want to show actors that are about to go away
+			   Actor->IsEditable() &&									// Only show actors that are allowed to be selected and drawn in editor
+			   Actor->IsListedInSceneOutliner() &&
+			   !Actor->IsTemplate() &&									// Should never happen, but we never want CDOs displayed
+			   !IsABuilderBrush(Actor) &&								// Don't show the builder brush
+			   !Actor->IsA(AWorldSettings::StaticClass()) &&	// Don't show the WorldSettings actor, even though it is technically editable
+			   FLevelUtils::IsLevelVisible(Actor->GetLevel());			// Only show Actors whose level is visible
+	}
+#endif // WITH_EDITOR
 
 }
 
diff --git a/Engine/Source/Runtime/Engine/Public/ActorEditorUtils.h b/Engine/Source/Runtime/Engine/Public/ActorEditorUtils.h
--- a/Engine/Source/Runtime/Engine/Public/ActorEditorUtils.h
+++ b/Engine/Source/Runtime/Engine/Public/ActorEditorUtils.h
@@ -75,5 +75,15 @@ namespace FActorEditorUtils
 	 * @return true if the actor name is valid, false if not.
 	 */ 
 	ENGINE_API bool ValidateActorLabel(FStringView InName, FText& OutErrorMessage);
+	
+#if WITH_EDITOR
+	/**
+	 * Return whether this actor is displayable in common UI (i.e. SceneOutliner), checks for things like
+	 * validity, if it is a BuilderBrush, etc.
+	 *
+	 * @return true if this actor is displayable in common UI
+	 */
+	ENGINE_API bool IsActorDisplayable(const AActor* InActor);
+#endif // WITH_EDITOR
 
 };
