[TEDS Viewers] Have table viewer UI rows initialize instead of creating them when rows come into view. Additionally, add the option to set 'Expand New Rows' option in the prefab editor

#Synced-CL 50817436

diff --git a/Engine/Plugins/Experimental/EditorDataStorageFeatures/Source/TedsTableViewer/Private/TedsTableViewerUtils.cpp b/Engine/Plugins/Experimental/EditorDataStorageFeatures/Source/TedsTableViewer/Private/TedsTableViewerUtils.cpp
--- a/Engine/Plugins/Experimental/EditorDataStorageFeatures/Source/TedsTableViewer/Private/TedsTableViewerUtils.cpp
+++ b/Engine/Plugins/Experimental/EditorDataStorageFeatures/Source/TedsTableViewer/Private/TedsTableViewerUtils.cpp
@@ -8,6 +8,7 @@
 #include "Elements/Columns/TypedElementIconOverrideColumns.h"
 #include "Elements/Columns/TypedElementMiscColumns.h"
 #include "Elements/Columns/TypedElementTypeInfoColumns.h"
+#include "Elements/Columns/TypedElementUIColumns.h"
 #include "Elements/Interfaces/TypedElementDataStorageUiInterface.h"
 #include "Styling/SlateIconFinder.h"
 #include "TedsTableViewerWidgetColumns.h"
@@ -224,6 +225,42 @@ namespace UE::Editor::DataStorage::TableViewerUtils
 		}
 		return ParentRow;
 	}
+	
+	void InitializeAllTableViewerUIRows(ICoreProvider* DataStorage, const TArray<TableViewerItemPtr>& Items, 
+		const FGetTableViewerRowMapping& GetTableViewerRowMapping, const FName& TableViewerIdentifier, 
+		const RowHandle TableViewerRowHandle, const bool bAddHierarchyTags)
+	{
+		// Have all rows initialize their UI Row counterparts
+		if (DataStorage)
+		{
+			const TableHandle RowWidgetTableHandle = DataStorage->FindTable(GetRowWidgetTableName());
+			const FHierarchyHandle WidgetHierarchyHandle = DataStorage->FindHierarchyByName(GetWidgetHierarchyName());
+
+			if (RowWidgetTableHandle != InvalidTableHandle && DataStorage->IsValidHierarchyHandle(WidgetHierarchyHandle))
+			{
+				for (const TableViewerItemPtr& Item : Items)
+				{
+					if (GetTableViewerRowMapping.IsBound())
+					{
+						const FString MappingKey = GetTableViewerRowMapping.Execute(TableViewerIdentifier, Item);
+						// If the row already existed we will use that, but if not, we want to create a row for each entry so
+						// changing the default expansion state doesn't affect rows that come into view but aren't new
+						if (const RowHandle ExistingRow = DataStorage->LookupMappedRow(TableViewerUIRowMappingDomain, FMapKey(MappingKey)); !DataStorage->IsRowAvailable(ExistingRow))
+						{
+							const RowHandle AddedRow = DataStorage->AddRow(RowWidgetTableHandle);
+				
+							DataStorage->SetParentRow(WidgetHierarchyHandle, AddedRow, TableViewerRowHandle);
+							DataStorage->MapRow(TableViewerUIRowMappingDomain, FMapKey(MappingKey), AddedRow);
+							if (bAddHierarchyTags)
+							{
+								DataStorage->AddColumns<FExpandedInUITag>(AddedRow);
+							}
+						}
+					}
+				}
+			}
+		}
+	}
 }
 
 void UTedsTableViewerFactory::RegisterTables(UE::Editor::DataStorage::ICoreProvider& DataStorage)
diff --git a/Engine/Plugins/Experimental/EditorDataStorageFeatures/Source/TedsTableViewer/Private/Widgets/Composite/STedsCompositeHierarchyViewer.cpp b/Engine/Plugins/Experimental/EditorDataStorageFeatures/Source/TedsTableViewer/Private/Widgets/Composite/STedsCompositeHierarchyViewer.cpp
--- a/Engine/Plugins/Experimental/EditorDataStorageFeatures/Source/TedsTableViewer/Private/Widgets/Composite/STedsCompositeHierarchyViewer.cpp
+++ b/Engine/Plugins/Experimental/EditorDataStorageFeatures/Source/TedsTableViewer/Private/Widgets/Composite/STedsCompositeHierarchyViewer.cpp
@@ -204,6 +204,7 @@ namespace UE::Editor::DataStorage
 			.ListSelectionMode(HierarchyViewerArgs._ListSelectionMode)
 			.OnSelectionChanged(HierarchyViewerArgs._OnSelectionChanged)
 			.GetCustomTableViewerRowMapping(HierarchyViewerArgs._GetCustomTableViewerRowMapping)
+			.ExpandNewRows(HierarchyViewerArgs._ExpandNewRows)
 			.PrimaryColumn(HierarchyViewerArgs._PrimaryColumn)
 			.AllowInvisibleItemSelection(HierarchyViewerArgs._AllowInvisibleItemSelection);
 		
@@ -463,6 +464,11 @@ namespace UE::Editor::DataStorage
 		HierarchyViewer->SetQueryStack(CompositeQueryNode);
 	}
 
+	void STedsCompositeHierarchyViewer::OnToggleExpandNewRows()
+	{
+		HierarchyViewer->bExpandNewRows = !HierarchyViewer->bExpandNewRows;
+	}
+
 } // namespace UE::Editor::DataStorage
 
 #undef LOCTEXT_NAMESPACE
diff --git a/Engine/Plugins/Experimental/EditorDataStorageFeatures/Source/TedsTableViewer/Private/Widgets/STedsHierarchyViewer.cpp b/Engine/Plugins/Experimental/EditorDataStorageFeatures/Source/TedsTableViewer/Private/Widgets/STedsHierarchyViewer.cpp
--- a/Engine/Plugins/Experimental/EditorDataStorageFeatures/Source/TedsTableViewer/Private/Widgets/STedsHierarchyViewer.cpp
+++ b/Engine/Plugins/Experimental/EditorDataStorageFeatures/Source/TedsTableViewer/Private/Widgets/STedsHierarchyViewer.cpp
@@ -37,6 +37,7 @@ namespace UE::Editor::DataStorage
 		HierarchyInterface = MoveTemp(InHierarchyInterface);
 		OnSelectionChanged = InArgs._OnSelectionChanged;
 		GetCustomTableViewerRowMapping = InArgs._GetCustomTableViewerRowMapping;
+		bExpandNewRows = InArgs._ExpandNewRows;
 		EmptyRowsMessage = InArgs._EmptyRowsMessage;
 		
 		IUiProvider::FPurposeID CellWidgetPurpose = InArgs._CellWidgetPurpose;
@@ -143,15 +144,31 @@ namespace UE::Editor::DataStorage
 			}
 		}
 
-		if (bExpansionStateChangeRequested.IsSet())
+		// Check if the model has been initialized yet, if so, run any Post initialization/deferred actions
+		if (!bModelInitialized)
 		{
-			if (TSharedPtr<QueryStack::IRowNode> AllRowsNode = Model->GetRowNode())
+			if (const TSharedPtr<QueryStack::IRowNode> RowNode = Model->GetRowNode())
 			{
-				FRowHandleArrayView AllRows = AllRowsNode->GetRows();
-				if (AllRows.Num())
+				if (RowNode->GetRevision() != QueryStack::INode::UninitializedRevisionId)
 				{
-					SetExpansionState_Internal(AllRows, bExpansionStateChangeRequested.GetValue());
-					bExpansionStateChangeRequested.Reset();
+					TableViewerUtils::InitializeAllTableViewerUIRows(
+						Model->GetDataStorageInterface(),
+						Model->GetItems(), 
+						TableViewerUtils::FGetTableViewerRowMapping::CreateSP(this, &SHierarchyViewer::GetTableViewerRowMapping), 
+						TableViewerIdentifier, 
+						GetWidgetRowHandle(),
+						true);
+					bModelInitialized = true;
+				}
+
+				if (bExpansionStateChangeRequested.IsSet())
+				{
+					FRowHandleArrayView AllRows = RowNode->GetRows();
+					if (AllRows.Num())
+					{
+						SetExpansionState_Internal(AllRows, bExpansionStateChangeRequested.GetValue());
+						bExpansionStateChangeRequested.Reset();
+					}
 				}
 			}
 		}
@@ -290,7 +307,7 @@ namespace UE::Editor::DataStorage
 	{
 		if (ICoreProvider* Storage = Model->GetDataStorageInterface())
 		{
-			const FString TableRowId = GetTableViewerRowMapping(Item);
+			const FString TableRowId = GetTableViewerRowMapping(TableViewerIdentifier, Item.RowHandle);
 			if (const RowHandle UIRow = Storage->LookupMappedRow(TableViewerUIRowMappingDomain, FMapKey(TableRowId)); Storage->IsRowAvailable(UIRow))
 			{
 				if (bExpanded)
@@ -416,34 +433,27 @@ namespace UE::Editor::DataStorage
 
 	void SHierarchyViewer::ExpandAll() const
 	{
-		if (TSharedPtr<QueryStack::IRowNode> RowNode = Model->GetRowNode())
+		// If the query stack hasn't been initialized yet, defer the expansion until next refresh
+		if (!bModelInitialized)
 		{
-			// If the query stack hasn't been initialized yet, defer the expansion until next refresh
-			if (RowNode->GetRevision() == QueryStack::INode::UninitializedRevisionId)
-			{
-				bExpansionStateChangeRequested = true;
-			}
-			else
-			{
-				SetExpansionState_Internal(RowNode->GetRows(), true);
-			}
-			
+			bExpansionStateChangeRequested = true;
+		}
+		else if (TSharedPtr<QueryStack::IRowNode> RowNode = Model->GetRowNode())
+		{
+			SetExpansionState_Internal(RowNode->GetRows(), true);
 		}
 	}
 
 	void SHierarchyViewer::CollapseAll() const
 	{
 		// If the query stack hasn't been initialized yet, defer the expansion until next refresh
-		if (TSharedPtr<QueryStack::IRowNode> RowNode = Model->GetRowNode())
+		if (!bModelInitialized)
 		{
-			if (RowNode->GetRevision() == QueryStack::INode::UninitializedRevisionId)
-			{
-				bExpansionStateChangeRequested = false;
-			}
-			else
-			{
-				SetExpansionState_Internal(RowNode->GetRows(), false);
-			}
+			bExpansionStateChangeRequested = false;
+		}
+		else if (TSharedPtr<QueryStack::IRowNode> RowNode = Model->GetRowNode())
+		{
+			SetExpansionState_Internal(RowNode->GetRows(), false);
 		}
 	}
 
@@ -470,27 +480,27 @@ namespace UE::Editor::DataStorage
 		return TreeView->IsItemVisible(InItem);
 	}
 	
-	FString SHierarchyViewer::GetTableViewerRowMapping(const TableViewerItemPtr InItem) const
+	FString SHierarchyViewer::GetTableViewerRowMapping(const FName InTableViewerIdentifier, const RowHandle InRow) const
 	{
 		if(GetCustomTableViewerRowMapping.IsBound())
 		{
-			if (FString RowMapping = GetCustomTableViewerRowMapping.Execute(TableViewerIdentifier, InItem.RowHandle); !RowMapping.IsEmpty())
+			if (FString RowMapping = GetCustomTableViewerRowMapping.Execute(InTableViewerIdentifier, InRow); !RowMapping.IsEmpty())
 			{
 				return RowMapping;
 			}
 		}
-		return TableViewerUtils::CreateRowMappingKey(TableViewerIdentifier, InItem);
+		return TableViewerUtils::CreateRowMappingKey(InTableViewerIdentifier, InRow);
 	}
-
+	
 	RowHandle SHierarchyViewer::AddOrFindTableViewerUIRow(const TableViewerItemPtr InItem) const
 	{
 		if (ICoreProvider* Storage = Model->GetDataStorageInterface())
 		{
-			const FString MappingKey = GetTableViewerRowMapping(InItem);
+			const FString MappingKey = GetTableViewerRowMapping(TableViewerIdentifier, InItem.RowHandle);
 			// If the row already existed but was removed from TreeView's SparseItemInfos (Searching, Filters, etc.), re-expand it
 			if (const RowHandle ExistingRow = Storage->LookupMappedRow(TableViewerUIRowMappingDomain, FMapKey(MappingKey)); Storage->IsRowAvailable(ExistingRow))
 			{
-				if(Storage->HasColumns<FExpandedInUITag>(ExistingRow))
+				if(Storage->HasColumns<FExpandedInUITag>(ExistingRow) && !TreeView->IsItemExpanded(InItem))
 				{
 					TreeView->SetItemExpansion(InItem, true);
 				}
@@ -503,12 +513,13 @@ namespace UE::Editor::DataStorage
 				Storage->SetParentRow(Storage->FindHierarchyByName(TableViewerUtils::GetWidgetHierarchyName()), AddedRow, GetWidgetRowHandle());
 				Storage->MapRow(TableViewerUIRowMappingDomain, FMapKey(MappingKey), AddedRow);
 
-				if (!TreeView->IsItemExpanded(InItem))
+				const bool bIsItemExpandedInTree = TreeView->IsItemExpanded(InItem);
+				if (bIsItemExpandedInTree != bExpandNewRows)
 				{
-					// Default to expanded if not already
-					TreeView->SetItemExpansion(InItem, true);
+					// Default to bExpandNewRows if not already
+					TreeView->SetItemExpansion(InItem, bExpandNewRows);
 				}
-				else
+				else if (bIsItemExpandedInTree && bExpandNewRows)
 				{
 					Storage->AddColumns<FExpandedInUITag>(AddedRow);
 				}
diff --git a/Engine/Plugins/Experimental/EditorDataStorageFeatures/Source/TedsTableViewer/Private/Widgets/STedsTableViewer.cpp b/Engine/Plugins/Experimental/EditorDataStorageFeatures/Source/TedsTableViewer/Private/Widgets/STedsTableViewer.cpp
--- a/Engine/Plugins/Experimental/EditorDataStorageFeatures/Source/TedsTableViewer/Private/Widgets/STedsTableViewer.cpp
+++ b/Engine/Plugins/Experimental/EditorDataStorageFeatures/Source/TedsTableViewer/Private/Widgets/STedsTableViewer.cpp
@@ -103,11 +103,31 @@ namespace UE::Editor::DataStorage
 		});
 
 		// Whenever the model changes, refresh the list to update the UI
-		Model->GetOnModelChanged().AddLambda([this]()
+		Model->GetOnModelChanged().AddRaw(this, &STedsTableViewer::OnModelChanged);
+	}
+
+	void STedsTableViewer::OnModelChanged()
+	{
+		// Check if the model hasn't been initialized yet, if so, run any Post initialization/deferred actions
+		if (!bModelInitialized)
 		{
-			ListView->RequestListRefresh();
-			CreateInternalWidget();
-		});
+			if (const TSharedPtr<QueryStack::IRowNode> RowNode = Model->GetRowNode())
+			{
+				if (RowNode->GetRevision() != QueryStack::INode::UninitializedRevisionId)
+				{
+					TableViewerUtils::InitializeAllTableViewerUIRows(
+						Model->GetDataStorageInterface(),
+						Model->GetItems(),
+						TableViewerUtils::FGetTableViewerRowMapping::CreateSP(this, &STedsTableViewer::GetTableViewerRowMapping), 
+						TableViewerIdentifier, 
+						GetWidgetRowHandle());
+					bModelInitialized = true;
+				}
+			}
+		}
+		
+		ListView->RequestListRefresh();
+		CreateInternalWidget();
 	}
 
 	void STedsTableViewer::OnDragEnter(const FGeometry& MyGeometry, const FDragDropEvent& DragDropEvent)
@@ -312,23 +332,23 @@ namespace UE::Editor::DataStorage
 		return ListView->IsItemVisible(InItem);
 	}
 	
-	FString STedsTableViewer::GetTableViewerRowMapping(const TableViewerItemPtr InItem) const
+	FString STedsTableViewer::GetTableViewerRowMapping(const FName InTableViewerIdentifier, const RowHandle InRow) const
 	{
 		if(GetCustomTableViewerRowMapping.IsBound())
 		{
-			if (FString RowMapping = GetCustomTableViewerRowMapping.Execute(TableViewerIdentifier, InItem.RowHandle); !RowMapping.IsEmpty())
+			if (FString RowMapping = GetCustomTableViewerRowMapping.Execute(InTableViewerIdentifier, InRow); !RowMapping.IsEmpty())
 			{
 				return RowMapping;
 			}
 		}
-		return TableViewerUtils::CreateRowMappingKey(TableViewerIdentifier, InItem);
+		return TableViewerUtils::CreateRowMappingKey(InTableViewerIdentifier, InRow);
 	}
 
 	RowHandle STedsTableViewer::AddOrFindTableViewerUIRow(const TableViewerItemPtr InItem) const
 	{
 		if (ICoreProvider* Storage = Model->GetDataStorageInterface())
 		{
-			const FString MappingKey = GetTableViewerRowMapping(InItem);
+			const FString MappingKey = GetTableViewerRowMapping(TableViewerIdentifier, InItem.RowHandle);
 			// If the row already existed but was removed from TreeView's SparseItemInfos (Searching, Filters, etc.), re-expand it
 			if (const RowHandle ExistingRow = Storage->LookupMappedRow(TableViewerUIRowMappingDomain, FMapKey(MappingKey)); Storage->IsRowAvailable(ExistingRow))
 			{
diff --git a/Engine/Plugins/Experimental/EditorDataStorageFeatures/Source/TedsTableViewer/Public/TedsTableViewerModel.h b/Engine/Plugins/Experimental/EditorDataStorageFeatures/Source/TedsTableViewer/Public/TedsTableViewerModel.h
--- a/Engine/Plugins/Experimental/EditorDataStorageFeatures/Source/TedsTableViewer/Public/TedsTableViewerModel.h
+++ b/Engine/Plugins/Experimental/EditorDataStorageFeatures/Source/TedsTableViewer/Public/TedsTableViewerModel.h
@@ -155,7 +155,7 @@ namespace UE::Editor::DataStorage
 		IUiProvider::FPurposeID HeaderWidgetPurpose;
 
 		// Cached revision ID for the query stack used to check when the table viewer needs a refresh
-		uint32 CachedRowQueryStackRevision = 0;
+		uint32 CachedRowQueryStackRevision = QueryStack::INode::UninitializedRevisionId;
 
 		// Delegate supplied by the widget to check if an item is visible in the UI currently
 		FIsItemVisible IsItemVisible;
diff --git a/Engine/Plugins/Experimental/EditorDataStorageFeatures/Source/TedsTableViewer/Public/TedsTableViewerUtils.h b/Engine/Plugins/Experimental/EditorDataStorageFeatures/Source/TedsTableViewer/Public/TedsTableViewerUtils.h
--- a/Engine/Plugins/Experimental/EditorDataStorageFeatures/Source/TedsTableViewer/Public/TedsTableViewerUtils.h
+++ b/Engine/Plugins/Experimental/EditorDataStorageFeatures/Source/TedsTableViewer/Public/TedsTableViewerUtils.h
@@ -23,6 +23,9 @@ struct FSlateBrush;
 // Util library for functions shared by the Teds Table Viewer and the Teds Outliner
 namespace UE::Editor::DataStorage::TableViewerUtils
 {
+	// Delegate fired to get a custom row mapping for the table viewer
+	DECLARE_DELEGATE_RetVal_TwoParams(FString, FGetTableViewerRowMapping, FName, RowHandle);
+	
 	TEDSTABLEVIEWER_API FName GetWidgetTableName();
 
 	// Get the name of the widget hierarchy for table viewer widgets
@@ -55,6 +58,11 @@ namespace UE::Editor::DataStorage::TableViewerUtils
 
 	// Walks up the UI Widget Hierarchy until finding the row with the FTableViewerTag
 	TEDSTABLEVIEWER_API RowHandle GetTableViewerUiRow(const ICoreProvider* DataStorage, const RowHandle ChildWidgetRow);
+
+	// Initializes all rows on a TableViewer for a QueryStack
+	TEDSTABLEVIEWER_API void InitializeAllTableViewerUIRows(ICoreProvider* DataStorage, const TArray<TableViewerItemPtr>& Items,
+		const FGetTableViewerRowMapping& GetTableViewerRowMapping, const FName& TableViewerIdentifier, 
+		const RowHandle TableViewerRowHandle, const bool bAddHierarchyTags = false);
 } // namespace UE::Editor::DataStorage::TableViewerUtils
 
 UCLASS()
diff --git a/Engine/Plugins/Experimental/EditorDataStorageFeatures/Source/TedsTableViewer/Public/Widgets/Composite/STedsCompositeHierarchyViewer.h b/Engine/Plugins/Experimental/EditorDataStorageFeatures/Source/TedsTableViewer/Public/Widgets/Composite/STedsCompositeHierarchyViewer.h
--- a/Engine/Plugins/Experimental/EditorDataStorageFeatures/Source/TedsTableViewer/Public/Widgets/Composite/STedsCompositeHierarchyViewer.h
+++ b/Engine/Plugins/Experimental/EditorDataStorageFeatures/Source/TedsTableViewer/Public/Widgets/Composite/STedsCompositeHierarchyViewer.h
@@ -175,6 +175,7 @@ namespace UE::Editor::DataStorage
 
 		// Menu actions
 		TEDSTABLEVIEWER_API void OnToggleShowFilteredParentHierarchy();
+		TEDSTABLEVIEWER_API void OnToggleExpandNewRows();
 		
 		// Set a handler that this viewer will forward its drop operations to.
 		TEDSTABLEVIEWER_API void SetDropHandler(TUniquePtr<FWidgetDropHandler> InDropHandler);
diff --git a/Engine/Plugins/Experimental/EditorDataStorageFeatures/Source/TedsTableViewer/Public/Widgets/STedsHierarchyViewer.h b/Engine/Plugins/Experimental/EditorDataStorageFeatures/Source/TedsTableViewer/Public/Widgets/STedsHierarchyViewer.h
--- a/Engine/Plugins/Experimental/EditorDataStorageFeatures/Source/TedsTableViewer/Public/Widgets/STedsHierarchyViewer.h
+++ b/Engine/Plugins/Experimental/EditorDataStorageFeatures/Source/TedsTableViewer/Public/Widgets/STedsHierarchyViewer.h
@@ -6,6 +6,7 @@
 #include "DataStorage/Handles.h"
 #include "TedsRowFilterNode.h"
 #include "TedsTableViewerModel.h"
+#include "TedsTableViewerUtils.h"
 #include "TedsQueryStackInterfaces.h"
 #include "TedsQueryStackExecutor.h"
 #include "Templates/Function.h"
@@ -61,13 +62,11 @@ namespace UE::Editor::DataStorage
 
 		// Delegate fired when the selection in the table viewer changes
 		DECLARE_DELEGATE_OneParam(FOnSelectionChanged, RowHandle)
-		
-		// Delegate fired to get a custom row mapping for the table viewer
-        DECLARE_DELEGATE_RetVal_TwoParams(FString, FGetTableViewerRowMapping, FName, RowHandle);
 
 		SLATE_BEGIN_ARGS(SHierarchyViewer)
 			: _CanSelectGeneratedColumn(true)
 			, _ListSelectionMode(ESelectionMode::Type::Single)
+			, _ExpandNewRows(true)
 			, _AllowInvisibleItemSelection(false)
 		{
 			
@@ -101,7 +100,10 @@ namespace UE::Editor::DataStorage
 
 		// Delegate called to define a custom row mapping for the Table Viewer Row Widget Rows
 		// Used in cases where you need the UI Rows to persist across RowHandle changes
-		SLATE_EVENT(FGetTableViewerRowMapping, GetCustomTableViewerRowMapping)
+		SLATE_EVENT(TableViewerUtils::FGetTableViewerRowMapping, GetCustomTableViewerRowMapping)
+
+		// Whether new rows will be expanded or collapsed when coming into view, rows will expand on initialization (if not persisted) then follow this set bool
+		SLATE_ARGUMENT(bool, ExpandNewRows)
 
 		// The message to show in place of the table viewer when there are no rows provided by the current query stack
 		// Empty = simply show the column headers instead of a message
@@ -172,7 +174,7 @@ namespace UE::Editor::DataStorage
 		
 	protected:
 		
-		FString GetTableViewerRowMapping(const TableViewerItemPtr InItem) const;
+		FString GetTableViewerRowMapping(const FName InTableViewerIdentifier, const RowHandle InRow) const;
 		
 		RowHandle AddOrFindTableViewerUIRow(const TableViewerItemPtr InItem) const;
 
@@ -224,7 +226,7 @@ namespace UE::Editor::DataStorage
 		FOnSelectionChanged OnSelectionChanged;
 		
 		// Delegate fired when getting the row mapping, will use the default 'TableId.DataRowHandle' if nothing is bound
-        FGetTableViewerRowMapping GetCustomTableViewerRowMapping;
+		TableViewerUtils::FGetTableViewerRowMapping GetCustomTableViewerRowMapping;
 
 		// Wrapper Teds Widget around our contents so we can use Teds columns to specify behavior
 		TSharedPtr<ITedsWidget> TedsWidget;
@@ -251,6 +253,12 @@ namespace UE::Editor::DataStorage
 		// If we need to update the expansion state of all rows the next time the model changes
 		mutable TOptional<bool> bExpansionStateChangeRequested;
 
+		// If the model hasn't been initialized yet, will check during OnModelChanged and let any PostModelInitialization functions run
+		bool bModelInitialized = false;
+
+		// If any newly viewed rows (after initialization) should be expanded by default
+		bool bExpandNewRows = true;
+
 		// Drag and Drop Handler
 		TUniquePtr<FWidgetDropHandler> DropHandler;
 	};
diff --git a/Engine/Plugins/Experimental/EditorDataStorageFeatures/Source/TedsTableViewer/Public/Widgets/STedsTableViewer.h b/Engine/Plugins/Experimental/EditorDataStorageFeatures/Source/TedsTableViewer/Public/Widgets/STedsTableViewer.h
--- a/Engine/Plugins/Experimental/EditorDataStorageFeatures/Source/TedsTableViewer/Public/Widgets/STedsTableViewer.h
+++ b/Engine/Plugins/Experimental/EditorDataStorageFeatures/Source/TedsTableViewer/Public/Widgets/STedsTableViewer.h
@@ -5,6 +5,7 @@
 #include "DataStorage/Handles.h"
 #include "TedsQueryStackInterfaces.h"
 #include "TedsTableViewerModel.h"
+#include "TedsTableViewerUtils.h"
 #include "Widgets/SCompoundWidget.h"
 #include "Widgets/Views/STableRow.h"
 #include "Widgets/Views/SListView.h"
@@ -46,9 +47,6 @@ namespace UE::Editor::DataStorage
 		
 		// Delegate fired when the selection in the table viewer changes
 		DECLARE_DELEGATE_OneParam(FOnSelectionChanged, RowHandle)
-		
-		// Delegate fired to get a custom row mapping for the table viewer
-		DECLARE_DELEGATE_RetVal_TwoParams(FString, FGetTableViewerRowMapping, FName, RowHandle);
 
 		SLATE_BEGIN_ARGS(STedsTableViewer)
 			: _CanSelectGeneratedColumn(true)
@@ -88,7 +86,7 @@ namespace UE::Editor::DataStorage
 
 		// Delegate called to define a custom row mapping for the Table Viewer Row Widget Rows
 		// Used in cases where you need the UI Rows to persist across RowHandle changes
-		SLATE_EVENT(FGetTableViewerRowMapping, GetCustomTableViewerRowMapping)
+		SLATE_EVENT(TableViewerUtils::FGetTableViewerRowMapping, GetCustomTableViewerRowMapping)
 
 		// The message to show in place of the table viewer when there are no rows provided by the current query stack
 		// Empty = simply show the column headers instead of a message
@@ -155,7 +153,7 @@ namespace UE::Editor::DataStorage
 
 	protected:
 	
-		FString GetTableViewerRowMapping(const TableViewerItemPtr InItem) const;
+		FString GetTableViewerRowMapping(const FName InTableViewerIdentifier, const RowHandle InRow) const;
 
 		RowHandle AddOrFindTableViewerUIRow(const TableViewerItemPtr InItem) const;
 
@@ -173,6 +171,8 @@ namespace UE::Editor::DataStorage
 
 		void AddWidgetColumns();
 
+		void OnModelChanged();
+
 	private:
 		// Unique identifier for this TableViewer
 		FName TableViewerIdentifier;
@@ -199,13 +199,16 @@ namespace UE::Editor::DataStorage
 		FOnSelectionChanged OnSelectionChanged;
 		
 		// Delegate fired when getting the row mapping, will use the default 'TableId.DataRowHandle' if nothing is bound
-		FGetTableViewerRowMapping GetCustomTableViewerRowMapping;
+		TableViewerUtils::FGetTableViewerRowMapping GetCustomTableViewerRowMapping;
 
 		// Wrapper Teds Widget around our contents so we can use Teds columns to specify behavior
 		TSharedPtr<ITedsWidget> TedsWidget;
 
 		// The message to show in place of the table viewer when there are no rows provided by the current query stack
 		TAttribute<FText> EmptyRowsMessage;
+		
+		// If the model hasn't been initialized yet, will check during OnModelChanged and let any PostModelInitialization functions run
+		bool bModelInitialized = false;
 
 		// Drag and Drop Handler
 		TUniquePtr<FWidgetDropHandler> DropHandler;
