Rework TEDS Outliner to select rows via the TypedElementSelectionInterface

Selection Interface is used predominantly through the engine to be able to add/remove items from a selection set.  The TEDS Outliner was not going through the same mechansim leading to missing elements in SelectionSets that are used to feed tooling that is not necessarily built around TEDS or its own selection representation.

Summary of changes
==============
* TedsOutliner can be passed SelectionSet.  When outliner row is selected, if associated TEDS RowHandle has a FTypedElementColumn, the handle in that column is added to the SelectionSet.  This allows any customizations to apply to the handle in the same way it is applied from selection via other tooling (ie. viewport).  TedsOutlinerImpl.cpp
* bSelectionDirty driven by selection change events.  Removal of queries to derive the event.
* TedsOutliner no longer responsible for adding SelectionColumn.  The Selection Interfaces do this where appropriate. Removes redundant code.
* TedsOutliner drives its own internal selection from configured SelectionSet

Additonal client changes:

Chaos Visual Debugger
================
CVD was already using TypedElements for selection with a small amount of TEDS usage to feed a TEDS outliner.
* Added TEDS observer to generate a TypedElementHandle for FChaosVDSceneParticle objects.
** See comment for caveat if CVD wants to support more types. Recommendation to have non-template API to derive type info.
* Added implementation of ITedsTypedElementBridgeInterface to StructDataElement to avoid downcast check to retrieve TEDS RowHandle feeding the scroll-to-selection mechansim in ChaosVDWorldOutliner.  Simplifies code too.
* Note that RowHandle representations of FChaosVDSceneParticle do _not_ get SelectionColumns as of this CL.  There was no usage of them so it would have constituted extra work.  Should there be a desire to query selection from TEDS using the SelectionColumns, suggest adding/removing in the base ChaosVDSelectionInterface.

Dataflow Editor
============
* Dataflow Outliner usage overrode how selection works in construction + simulation viewports and push outliner tree items directly to the outliner widget.  No SelectionSet or SelectionSetOverride necessary.

#Synced-CL 50829790

#Synced-CL 50829790

#Synced-CL 50829790

diff --git a/Engine/Plugins/ChaosVD/Source/ChaosVD/Private/ChaosVDWorldOutlinerMode.cpp b/Engine/Plugins/ChaosVD/Source/ChaosVD/Private/ChaosVDWorldOutlinerMode.cpp
--- a/Engine/Plugins/ChaosVD/Source/ChaosVD/Private/ChaosVDWorldOutlinerMode.cpp
+++ b/Engine/Plugins/ChaosVD/Source/ChaosVD/Private/ChaosVDWorldOutlinerMode.cpp
@@ -285,23 +285,13 @@ void FChaosVDWorldOutlinerMode::HandlePostSelectionChange(const UTypedElementSel
 		using namespace UE::Editor::DataStorage;
 		using namespace Chaos::VD::TypedElementDataUtil;
 		
-		RowHandle RowHandle = InvalidRowHandle;
-	
-		if (const FChaosVDSceneParticle* Particle = GetStructDataFromTypedElementHandle<FChaosVDSceneParticle>(SelectedParticlesHandles[0]))
-		{
-			RowHandle = Particle->GetTedsRowHandle();
-		}
-		else if (const ICompatibilityProvider* Compatibility = GetDataStorageFeature<ICompatibilityProvider>(CompatibilityFeatureName))
+		if (TTypedElement<ITedsTypedElementBridgeInterface> BridgeInterface =  UTypedElementRegistry::GetInstance()->GetElement<ITedsTypedElementBridgeInterface>(SelectedParticlesHandles[0]))
 		{
-			RowHandle = Compatibility->FindRowWithCompatibleObject(ActorElementDataUtil::GetActorFromHandle(SelectedParticlesHandles[0]));
-		}
-
-		// Add a selection column in TEDS
-		TedsOutlinerImpl->SetSelection({ RowHandle });
-
-		if (FSceneOutlinerTreeItemPtr TreeItem = SceneOutliner->GetTreeItem(RowHandle, true))
-		{
-			SceneOutliner->ScrollItemIntoView(TreeItem);
+			UE::Editor::DataStorage::RowHandle RowHandle = BridgeInterface.GetRowHandle();
+			if (FSceneOutlinerTreeItemPtr TreeItem = SceneOutliner->GetTreeItem(RowHandle, true))
+			{
+				SceneOutliner->ScrollItemIntoView(TreeItem);
+			}
 		}
 	}
 }
diff --git a/Engine/Plugins/ChaosVD/Source/ChaosVD/Private/ChaosVDWorldOutlinerTab.cpp b/Engine/Plugins/ChaosVD/Source/ChaosVD/Private/ChaosVDWorldOutlinerTab.cpp
--- a/Engine/Plugins/ChaosVD/Source/ChaosVD/Private/ChaosVDWorldOutlinerTab.cpp
+++ b/Engine/Plugins/ChaosVD/Source/ChaosVD/Private/ChaosVDWorldOutlinerTab.cpp
@@ -91,8 +91,12 @@ void FChaosVDWorldOutlinerTab::CreateWorldOutlinerWidget()
 		});
 		
 	Params.HierarchyData = FTedsOutlinerHierarchyData(FChaosVDTableRowParentColumn::StaticStruct(), RowHandleGetter, RowHandleSetter, ChildrenRowHandlesGetter);
-
-	Params.SelectionSetOverride = FName("CVDSelection");
+	
+	if (TSharedPtr<FChaosVDScene> CvdScene = GetChaosVDScene().Pin())
+	{
+		Params.SelectionSet = CvdScene->GetElementSelectionSet();
+	}
+	
 
 	InitOptions.ModeFactory = FCreateSceneOutlinerMode::CreateLambda([this, PlaybackController, &Params](SSceneOutliner* Outliner)
 	{
diff --git a/Engine/Plugins/ChaosVD/Source/ChaosVD/Private/TEDS/ChaosVDParticleEditorDataFactory.cpp b/Engine/Plugins/ChaosVD/Source/ChaosVD/Private/TEDS/ChaosVDParticleEditorDataFactory.cpp
--- a/Engine/Plugins/ChaosVD/Source/ChaosVD/Private/TEDS/ChaosVDParticleEditorDataFactory.cpp
+++ b/Engine/Plugins/ChaosVD/Source/ChaosVD/Private/TEDS/ChaosVDParticleEditorDataFactory.cpp
@@ -9,6 +9,7 @@
 #include "Actors/ChaosVDSolverInfoActor.h"
 #include "DataStorage/Handles.h"
 #include "Elements/Columns/TypedElementCompatibilityColumns.h"
+#include "Elements/Columns/TypedElementHandleColumn.h"
 #include "Elements/Columns/TypedElementLabelColumns.h"
 #include "Elements/Columns/TypedElementMiscColumns.h"
 #include "Elements/Columns/TypedElementTypeInfoColumns.h"
@@ -35,6 +36,7 @@ void UChaosVDParticleEditorDataFactory::RegisterTables(UE::Editor::DataStorage::
 
 	Registry->RegisterElementType<FStructTypedElementData, true>(NAME_CVD_StructDataElement);
 	Registry->RegisterElementInterface<ITypedElementSelectionInterface>(NAME_CVD_StructDataElement, NewObject<UChaosVDSelectionInterface>());
+	Registry->RegisterElementInterface<ITedsTypedElementBridgeInterface>(NAME_CVD_StructDataElement, NewObject<UChaosVDTedsTypedElementBridgeInterface>());
 
 	using namespace UE::Editor::DataStorage;
 
@@ -56,6 +58,32 @@ void UChaosVDParticleEditorDataFactory::RegisterQueries(UE::Editor::DataStorage:
 
 	using namespace UE::Editor::DataStorage::Queries;
 	using namespace UE::Editor::DataStorage;
+		
+		DataStorage.RegisterQuery(
+		Select(
+			TEXT("Generate TypedElement Handle for ChaosVDDataObject"),
+			FObserver::OnAdd<FChaosVDObjectDataTag>(),
+			[](IQueryContext& Context, RowHandle Row, const FTypedElementExternalObjectColumn& RawObject, const FTypedElementScriptStructTypeInfoColumn& TypeInfo)
+			{
+				if (RawObject.Object)
+				{
+					// Note: Templated AcquireTypedElementHandleForStruct<T> call to CreateTypedElementDataForStructData<T>
+					// implicitly determines the type from T
+					// Therefore it is necessary to explicitly downcast to FChaosVDSceneParticle so that the StructDataType::StaticStruct()
+					// resolves the correct type info.
+					if (TypeInfo.TypeInfo->IsChildOf(FChaosVDSceneParticle::StaticStruct()))
+					{
+						FChaosVDSceneParticle* Particle = static_cast<FChaosVDSceneParticle*>(RawObject.Object);
+						const bool bAllowCreate = true;
+						FTypedElementHandle Handle = Chaos::VD::TypedElementDataUtil::AcquireTypedElementHandleForStruct(Particle, bAllowCreate);
+						Context.AddColumn(Row, Compatibility::FTypedElementColumn
+						{
+							.Handle = MoveTemp(Handle)
+						});
+					}
+				}
+			})
+			.Compile());
 	
 	DataStorage.RegisterQuery(
 		Select(
diff --git a/Engine/Plugins/ChaosVD/Source/ChaosVD/Private/TEDS/ChaosVDSelectionInterface.cpp b/Engine/Plugins/ChaosVD/Source/ChaosVD/Private/TEDS/ChaosVDSelectionInterface.cpp
--- a/Engine/Plugins/ChaosVD/Source/ChaosVD/Private/TEDS/ChaosVDSelectionInterface.cpp
+++ b/Engine/Plugins/ChaosVD/Source/ChaosVD/Private/TEDS/ChaosVDSelectionInterface.cpp
@@ -2,5 +2,18 @@
 
 #include "TEDS/ChaosVDSelectionInterface.h"
 
+#include "ChaosVDSceneParticle.h"
+#include "TEDS/ChaosVDStructTypedElementData.h"
+
 #include UE_INLINE_GENERATED_CPP_BY_NAME(ChaosVDSelectionInterface)
 
+FTedsRowHandle UChaosVDTedsTypedElementBridgeInterface::GetRowHandle(const FTypedElementHandle& InElementHandle) const
+{
+	const FChaosVDSceneParticle* Particle = Chaos::VD::TypedElementDataUtil::GetStructDataFromTypedElementHandle<FChaosVDSceneParticle>(InElementHandle);
+	if (Particle)
+	{
+		return FTedsRowHandle(Particle->GetTedsRowHandle());
+	}
+	
+	return FTedsRowHandle();
+}
diff --git a/Engine/Plugins/ChaosVD/Source/ChaosVD/Public/TEDS/ChaosVDSelectionInterface.h b/Engine/Plugins/ChaosVD/Source/ChaosVD/Public/TEDS/ChaosVDSelectionInterface.h
--- a/Engine/Plugins/ChaosVD/Source/ChaosVD/Public/TEDS/ChaosVDSelectionInterface.h
+++ b/Engine/Plugins/ChaosVD/Source/ChaosVD/Public/TEDS/ChaosVDSelectionInterface.h
@@ -2,6 +2,7 @@
 
 #pragma once
 
+#include "Elements/Interfaces/TedsTypedElementBridgeInterface.h"
 #include "Elements/Interfaces/TypedElementSelectionInterface.h"
 #include "UObject/Object.h"
 #include "ChaosVDSelectionInterface.generated.h"
@@ -14,3 +15,11 @@ class UChaosVDSelectionInterface : public UObject, public ITypedElementSelection
 {
 	GENERATED_BODY()
 };
+
+UCLASS(MinimalAPI)
+class UChaosVDTedsTypedElementBridgeInterface : public UObject, public ITedsTypedElementBridgeInterface
+{
+	GENERATED_BODY()
+public:
+	virtual FTedsRowHandle GetRowHandle(const FTypedElementHandle& InElementHandle) const override;
+};
\ No newline at end of file
diff --git a/Engine/Plugins/Experimental/Dataflow/Source/DataflowEditor/Private/Dataflow/DataflowOutlinerView.cpp b/Engine/Plugins/Experimental/Dataflow/Source/DataflowEditor/Private/Dataflow/DataflowOutlinerView.cpp
--- a/Engine/Plugins/Experimental/Dataflow/Source/DataflowEditor/Private/Dataflow/DataflowOutlinerView.cpp
+++ b/Engine/Plugins/Experimental/Dataflow/Source/DataflowEditor/Private/Dataflow/DataflowOutlinerView.cpp
@@ -79,9 +79,6 @@ TSharedPtr<ISceneOutliner> FDataflowOutlinerView::CreateWidget()
 			Select()
 			.Where(TColumn<FDataflowSceneObjectTag>(GetEditorContent()->GetDataflowOwner().GetFName()))
 			.Compile()));
-	
-	// Empty selection set name is currently the level editor
-	Params.SelectionSetOverride = FName("DataflowSelection");
 
 	InitOptions.ModeFactory = FCreateSceneOutlinerMode::CreateLambda([this, &Params](SSceneOutliner* Outliner)
 	{
diff --git a/Engine/Plugins/Experimental/EditorDataStorageFeatures/Source/TedsOutliner/Private/TedsOutlinerImpl.cpp b/Engine/Plugins/Experimental/EditorDataStorageFeatures/Source/TedsOutliner/Private/TedsOutlinerImpl.cpp
--- a/Engine/Plugins/Experimental/EditorDataStorageFeatures/Source/TedsOutliner/Private/TedsOutlinerImpl.cpp
+++ b/Engine/Plugins/Experimental/EditorDataStorageFeatures/Source/TedsOutliner/Private/TedsOutlinerImpl.cpp
@@ -35,6 +35,7 @@
 #include "Columns/TedsOutlinerColumns.h"
 #include "Compatibility/SceneOutlinerTedsBridge.h"
 #include "Elements/Columns/TypedElementCompatibilityColumns.h"
+#include "Elements/Columns/TypedElementHandleColumn.h"
 #include "Elements/Columns/TypedElementOverrideColumns.h"
 #include "Elements/Columns/TypedElementSelectionColumns.h"
 #include "Elements/Columns/TypedElementSlateWidgetColumns.h"
@@ -248,7 +249,7 @@ FTedsOutlinerImpl::FTedsOutlinerImpl(const FTedsOutlinerParams& InParams, IScene
 	, LabelWidgetPurpose(InParams.LabelWidgetPurpose)
 	, InitialQueryDescription(InParams.QueryDescription)
 	, HierarchyData(InParams.HierarchyData)
-	, SelectionSetName(InParams.SelectionSetOverride)
+	, SelectionSet(InParams.SelectionSet)
 	, bForceShowParents(InParams.bForceShowParents)
 	, SceneOutlinerMode(InMode)
 	, SceneOutliner(InParams.SceneOutliner)
@@ -420,16 +421,30 @@ FTedsOutlinerImpl::FIsItemCompatible& FTedsOutlinerImpl::IsItemCompatible()
 
 void FTedsOutlinerImpl::SetSelection(const TArray<DataStorage::RowHandle>& InSelectedRows)
 {
-	if (!SelectionSetName.IsSet())
+	if (!SelectionSet)
 	{
 		return;
 	}
-	
+
+	// Clear the current selection first using the UTypedElementSelectionSet
 	ClearSelection();
 
+	// Select each row using the TypedElementSelectionInterface via UTypedElementSelectionSet
+	FTypedElementSelectionOptions SelectionOptions;
 	for(DataStorage::RowHandle Row : InSelectedRows)
 	{
-		Storage->AddColumn(Row, FTypedElementSelectionColumn{ .SelectionSet = SelectionSetName.GetValue() });
+		FTypedElementHandle ElementHandle;
+		// Get the TypedElementHandle from storage for this row
+		if (DataStorage::Compatibility::FTypedElementColumn* Column = Storage->GetColumn<DataStorage::Compatibility::FTypedElementColumn>(Row))
+		{
+			ElementHandle = Column->Handle;
+		}
+		
+		if (ElementHandle)
+		{
+			// Use the UTypedElementSelectionSet's selection interface to select this element
+			SelectionSet->SelectElement(ElementHandle, SelectionOptions);
+		}
 	}
 }
 	
@@ -681,6 +696,11 @@ void FTedsOutlinerImpl::RegenerateColumns(const Outliner::FTedsOutlinerColumnDes
 		MetaDataView, ColumnConstructor);
 }
 
+void FTedsOutlinerImpl::OnSelectionChangedHandler(const UTypedElementSelectionSet* InSelectionSet)
+{
+	bSelectionDirty = true;
+}
+
 DataStorage::RowHandle FTedsOutlinerImpl::GetOutlinerRowHandle() const
 {
 	return OutlinerRowHandle;
@@ -1019,57 +1039,12 @@ void FTedsOutlinerImpl::RecompileQueries()
 		UpdateParentQuery = Storage->RegisterQuery(MoveTemp(UpdateParentQueryDescription));
 	}
 
-	if (SelectionSetName.IsSet())
+	if (SelectionSet)
 	{
-		// Query to grab all selected rows
-		FQueryDescription SelectedRowsQueryDescription;
-		
-		if (bUsingQueryConditionsSyntax)
-		{
-			SelectedRowsQueryDescription = Select()
-				.Where(TColumn<FTypedElementSelectionColumn>())
-				.Compile();
-		}
-		else
+		if (!SelectionSetChangedDelegateHandle.IsValid())
 		{
-			SelectedRowsQueryDescription =
-				Select()
-				.Where()
-					.All(FTypedElementSelectionColumn::StaticStruct())
-				.Compile();
+			SelectionSetChangedDelegateHandle = SelectionSet->OnChanged().AddSP(this, &FTedsOutlinerImpl::OnSelectionChangedHandler);
 		}
-	
-		// Query to track when a row gets selected
-		FQueryDescription SelectionAddedQueryDescription =
-							Select(
-							TEXT("Row selected"),
-							FObserver::OnAdd<FTypedElementSelectionColumn>().SetExecutionMode(EExecutionMode::GameThread),
-							[this](IQueryContext& Context, DataStorage::RowHandle Row)
-							{
-								bSelectionDirty = true;
-							})
-							.Compile();
-
-		// Add the conditions from FinalQueryDescription to ensure we are the rows the user requested
-		MergeQueries(SelectionAddedQueryDescription, FinalQueryDescription);
-
-		// Query to track when a row gets deselected
-		FQueryDescription SelectionRemovedQueryDescription =
-							Select(
-							TEXT("Row deselected"),
-							FObserver::OnRemove<FTypedElementSelectionColumn>().SetExecutionMode(EExecutionMode::GameThread),
-							[this](IQueryContext& Context, DataStorage::RowHandle Row)
-							{
-								bSelectionDirty = true;
-							})
-							.Compile();
-
-		// Add the conditions from FinalQueryDescription to ensure we are the rows the user requested
-		MergeQueries(SelectionRemovedQueryDescription, FinalQueryDescription);
-
-		SelectedRowsQuery = Storage->RegisterQuery(MoveTemp(SelectedRowsQueryDescription));
-		SelectionAddedQuery = Storage->RegisterQuery(MoveTemp(SelectionAddedQueryDescription));
-		SelectionRemovedQuery = Storage->RegisterQuery(MoveTemp(SelectionRemovedQueryDescription));
 	}
 
 	// Query to track when the label of a row we are observing changes, to re-filter/re-search for the item
@@ -1106,6 +1081,11 @@ void FTedsOutlinerImpl::RecompileQueries()
 
 void FTedsOutlinerImpl::UnregisterQueries()
 {
+	if (SelectionSet && SelectionSetChangedDelegateHandle.IsValid())
+	{
+		SelectionSet->OnChanged().Remove(SelectionSetChangedDelegateHandle);
+		SelectionSetChangedDelegateHandle.Reset();
+	}
 	if (Storage)
 	{
 		Storage->UnregisterQuery(ChildRowHandleQuery);
@@ -1126,40 +1106,15 @@ void FTedsOutlinerImpl::UnregisterQueries()
 
 void FTedsOutlinerImpl::ClearSelection() const
 {
-	if (!SelectionSetName.IsSet())
+	if (!SelectionSet)
 	{
 		return;
 	}
 
-	using namespace UE::Editor::DataStorage::Queries;
-
-	TArray<RowHandle> RowsToRemoveSelectionColumn;
-
-	// Query to remove the selection column from all rows that belong to this selection set
-	DirectQueryCallback RowCollector = CreateDirectQueryCallbackBinding(
-	[this, &RowsToRemoveSelectionColumn](const IDirectQueryContext& Context, const RowHandle* RowHandles)
-	{
-		const TConstArrayView<RowHandle> Rows(RowHandles, Context.GetRowCount());
-
-		for(const RowHandle RowHandle : Rows)
-		{
-			if (const FTypedElementSelectionColumn* SelectionColumn = Storage->GetColumn<FTypedElementSelectionColumn>(RowHandle))
-			{
-				if (SelectionColumn->SelectionSet == SelectionSetName)
-				{
-					RowsToRemoveSelectionColumn.Add(RowHandle);
-				}
-			}
-		}
-	});
-
-	Storage->RunQuery(SelectedRowsQuery, RowCollector);
-
-	for(const RowHandle RowHandle : RowsToRemoveSelectionColumn)
-	{
-		Storage->RemoveColumn<FTypedElementSelectionColumn>(RowHandle);
-	}
-
+	// Use the UTypedElementSelectionSet to clear all selections
+	// This will properly use the TypedElementSelectionInterface for each element
+	FTypedElementSelectionOptions SelectionOptions;
+	SelectionSet->ClearSelection(SelectionOptions);
 }
 
 void FTedsOutlinerImpl::PostDataStorageUpdate()
@@ -1267,9 +1222,9 @@ DataStorage::ICompatibilityProvider* FTedsOutlinerImpl::GetStorageCompatibility(
 	return StorageCompatibility;
 }
 
-TOptional<FName> FTedsOutlinerImpl::GetSelectionSetName() const
+UTypedElementSelectionSet* FTedsOutlinerImpl::GetSelectionSet() const
 {
-	return SelectionSetName;
+	return SelectionSet;
 }
 
 FOnTedsOutlinerSelectionChanged& FTedsOutlinerImpl::OnSelectionChanged()
diff --git a/Engine/Plugins/Experimental/EditorDataStorageFeatures/Source/TedsOutliner/Private/TedsOutlinerMode.cpp b/Engine/Plugins/Experimental/EditorDataStorageFeatures/Source/TedsOutliner/Private/TedsOutlinerMode.cpp
--- a/Engine/Plugins/Experimental/EditorDataStorageFeatures/Source/TedsOutliner/Private/TedsOutlinerMode.cpp
+++ b/Engine/Plugins/Experimental/EditorDataStorageFeatures/Source/TedsOutliner/Private/TedsOutlinerMode.cpp
@@ -14,6 +14,7 @@
 #include "TedsOutlinerItem.h"
 #include "ToolMenus.h"
 #include "DragAndDrop/Widgets/WidgetDropHandler.h"
+#include "Elements/Columns/TypedElementHandleColumn.h"
 
 #include UE_INLINE_GENERATED_CPP_BY_NAME(TedsOutlinerMode)
 
@@ -69,20 +70,35 @@ void FTedsOutlinerMode::Rebuild()
 
 void FTedsOutlinerMode::OnSelectionChanged()
 {
-	TOptional<FName> SelectionSetName = TedsOutlinerImpl->GetSelectionSetName();
+	TWeakObjectPtr<UTypedElementSelectionSet> WeakSelectionSet = TedsOutlinerImpl->GetSelectionSet();
 	DataStorage::ICoreProvider* Storage = TedsOutlinerImpl->GetStorage();
-	
+
 	// The selection in TEDS was changed, update the outliner to respond
-	SceneOutliner->SetSelection([SelectionSetName, Storage](ISceneOutlinerTreeItem& InItem) -> bool
+	SceneOutliner->SetSelection([WeakSelectionSet, Storage](ISceneOutlinerTreeItem& InItem) -> bool
 	{
+		UTypedElementSelectionSet* SelectionSet = WeakSelectionSet.Get();
+		if (!SelectionSet)
+		{
+			return false;
+		}
+
 		if(const FTedsOutlinerTreeItem* TedsItem = InItem.CastTo<FTedsOutlinerTreeItem>())
 		{
 			const DataStorage::RowHandle RowHandle = TedsItem->GetRowHandle();
 
-			if(const FTypedElementSelectionColumn* SelectionColumn = Storage->GetColumn<FTypedElementSelectionColumn>(RowHandle))
+			FTypedElementHandle ElementHandle;
+			if (DataStorage::Compatibility::FTypedElementColumn* Column = Storage->GetColumn<DataStorage::Compatibility::FTypedElementColumn>(RowHandle))
+			{
+				ElementHandle = Column->Handle;
+			}
+						
+			bool bIsSelected = false;
+			if (ElementHandle)
 			{
-				return SelectionColumn->SelectionSet == SelectionSetName;
+				bIsSelected = SelectionSet->IsElementSelected(ElementHandle, FTypedElementIsSelectedOptions());
 			}
+			
+			return bIsSelected;
 		}
 		return false;
 	});
diff --git a/Engine/Plugins/Experimental/EditorDataStorageFeatures/Source/TedsOutliner/Private/TedsOutlinerModule.cpp b/Engine/Plugins/Experimental/EditorDataStorageFeatures/Source/TedsOutliner/Private/TedsOutlinerModule.cpp
--- a/Engine/Plugins/Experimental/EditorDataStorageFeatures/Source/TedsOutliner/Private/TedsOutlinerModule.cpp
+++ b/Engine/Plugins/Experimental/EditorDataStorageFeatures/Source/TedsOutliner/Private/TedsOutlinerModule.cpp
@@ -187,9 +187,15 @@ TSharedRef<SWidget> FTedsOutlinerModule::CreateLevelEditorTedsOutliner()
 	 *		return MobilityColumn.Mobility == EComponentMobility::Static;
 	 *	});
 	 */
-		
-	// Empty selection set name is currently the level editor
-	Params.SelectionSetOverride = FName();
+	
+	const FLevelEditorModule* const LevelEditorModule = FModuleManager::GetModulePtr<FLevelEditorModule>(TEXT("LevelEditor"));
+	if (LevelEditorModule)
+	{
+		if (auto LevelEditor = LevelEditorModule->GetLevelEditorInstance().Pin())
+		{
+			Params.SelectionSet = LevelEditor->GetMutableElementSelectionSet();
+		}
+	}
 	
 	TSharedRef<ISceneOutliner> TEDSOutlinerShared = CreateTedsOutliner(InitOptions, Params);
 	
diff --git a/Engine/Plugins/Experimental/EditorDataStorageFeatures/Source/TedsOutliner/Public/TedsOutlinerImpl.h b/Engine/Plugins/Experimental/EditorDataStorageFeatures/Source/TedsOutliner/Public/TedsOutlinerImpl.h
--- a/Engine/Plugins/Experimental/EditorDataStorageFeatures/Source/TedsOutliner/Public/TedsOutlinerImpl.h
+++ b/Engine/Plugins/Experimental/EditorDataStorageFeatures/Source/TedsOutliner/Public/TedsOutlinerImpl.h
@@ -10,6 +10,7 @@
 #include "Elements/Framework/TypedElementRowHandleArrayView.h"
 #include "Elements/Interfaces/TypedElementDataStorageUiInterface.h"
 #include "Elements/Framework/TypedElementQueryBuilder.h"
+#include "Elements/Framework/TypedElementSelectionSet.h"
 #include "TedsOutlinerFilter.h"
 
 #define UE_API TEDSOUTLINER_API
@@ -157,8 +158,8 @@ struct FTedsOutlinerParams
 	// parent-child relation in the tree view
 	TOptional<FTedsOutlinerHierarchyData> HierarchyData;
 
-	// The selection set to use for this Outliner, unset = don't propagate tree selection to the TEDS column
-	TOptional<FName> SelectionSetOverride;
+	// The selection set to use for this Outliner, nullptr = don't propagate tree selection through TypedElementSelectionInterface
+	UTypedElementSelectionSet* SelectionSet;
 
 	// The purpose to use when generating widgets for row/column pairs through TEDS UI
 	DataStorage::IUiProvider::FPurposeID CellWidgetPurpose;
@@ -200,9 +201,9 @@ public:
 	UE_API DataStorage::IUiProvider* GetStorageUI() const;
 	UE_API DataStorage::ICompatibilityProvider* GetStorageCompatibility() const;
 
-	UE_API TOptional<FName> GetSelectionSetName() const;
+	UE_API UTypedElementSelectionSet* GetSelectionSet() const;
 
-	// Delegate fired when the selection in TEDS changes, only if SelectionSetName is set
+	// Delegate fired when the selection changes, only if SelectionSet is set
 	UE_API FOnTedsOutlinerSelectionChanged& OnSelectionChanged();
 
 	// Delegate fired when the hierarchy changes due to item addition/removal/move
@@ -212,7 +213,8 @@ public:
 	DECLARE_DELEGATE_RetVal_OneParam(bool, FIsItemCompatible, const ISceneOutlinerTreeItem&)
 	UE_API FIsItemCompatible& IsItemCompatible();
 
-	// Update the selection in TEDS to the input rows, only if SelectionSetName is set
+	// Update the selection to the input rows, only if SelectionSet is set
+	// This will use the TypedElementSelectionInterface to select/deselect elements
 	UE_API void SetSelection(const TArray<DataStorage::RowHandle>& InSelectedRows);
 
 	// Helper function to create a label widget for a given row
@@ -279,6 +281,9 @@ protected:
 
 	UE_API void ClearColumns() const;
 	UE_API void RegenerateColumns(const FTedsOutlinerColumnDescription& ColumnDescription);
+	
+private:
+	void OnSelectionChangedHandler(const UTypedElementSelectionSet* InSelectionSet);
 
 protected:
 	// TEDS Storage Constructs
@@ -360,8 +365,14 @@ protected:
 
 	// Query to track when a row's label gets changed
 	DataStorage::QueryHandle LabelUpdateQuery = DataStorage::InvalidQueryHandle;
+
+	// Handle to keep track of the lifetime of the OnChanged delegate handler for the SelectionSet
+	// This handler is responsible for setting bSelectionDirty
+	FDelegateHandle SelectionSetChangedDelegateHandle;
+	
+	// The UTypedElementSelectionSet to use for propagating selection via TypedElementSelectionInterface
+	TObjectPtr<UTypedElementSelectionSet> SelectionSet;
 	
-	TOptional<FName> SelectionSetName;
 	bool bSelectionDirty = false;
 
 	// If true, parent nodes will remain visible if a child passes all filters even if the parent fails a filter
diff --git a/Engine/Plugins/Experimental/EditorDataStorageFeatures/Source/TedsOutliner/TedsOutliner.Build.cs b/Engine/Plugins/Experimental/EditorDataStorageFeatures/Source/TedsOutliner/TedsOutliner.Build.cs
--- a/Engine/Plugins/Experimental/EditorDataStorageFeatures/Source/TedsOutliner/TedsOutliner.Build.cs
+++ b/Engine/Plugins/Experimental/EditorDataStorageFeatures/Source/TedsOutliner/TedsOutliner.Build.cs
@@ -26,7 +26,8 @@ public class TedsOutliner : ModuleRules
 					"SceneOutliner",
 					"Slate",
 					"SlateCore",
-					"TypedElementFramework"
+					"TypedElementFramework",
+					"TypedElementRuntime"
 				});
 			PrivateDependencyModuleNames.AddRange(
 				new string[]
