[TEDS Hierarchy Viewer] Added a quick storage nullptr check

#Synced-CL 50837606

diff --git a/Engine/Plugins/Experimental/EditorDataStorageFeatures/Source/TedsTableViewer/Private/Widgets/STedsHierarchyViewer.cpp b/Engine/Plugins/Experimental/EditorDataStorageFeatures/Source/TedsTableViewer/Private/Widgets/STedsHierarchyViewer.cpp
--- a/Engine/Plugins/Experimental/EditorDataStorageFeatures/Source/TedsTableViewer/Private/Widgets/STedsHierarchyViewer.cpp
+++ b/Engine/Plugins/Experimental/EditorDataStorageFeatures/Source/TedsTableViewer/Private/Widgets/STedsHierarchyViewer.cpp
@@ -27,8 +27,8 @@ namespace UE::Editor::DataStorage
 
 	SHierarchyViewer::~SHierarchyViewer()
 	{
-		RemoveAllTableViewerUIRows();
-	}
+			RemoveAllTableViewerUIRows();
+		}
 	
 	void SHierarchyViewer::Construct(
 		const FArguments& InArgs,
@@ -324,7 +324,7 @@ namespace UE::Editor::DataStorage
 
 	void SHierarchyViewer::OnSetExpansionRecursive(TableViewerItemPtr Item, bool bExpanded) const
 	{
-		if (const ICoreProvider* Storage = Model->GetDataStorageInterface(); HierarchyInterface.IsValid())
+		if (const ICoreProvider* Storage = Model->GetDataStorageInterface(); Storage != nullptr && HierarchyInterface.IsValid())
 		{
 			TArray<RowHandle> ChildRowsToSet;
 			HierarchyInterface->WalkDepthFirst(*Storage, Item.RowHandle,
