Updated ContentBrowser::ExtensionUtils::IsFolderFavorite to take FContentBrowserItemPath instead of FString

#rnx

#Synced-CL 50845945

diff --git a/Engine/Plugins/Experimental/EditorDataStorageFeatures/Source/TedsAssetData/Private/Processors/AssetProcessors.cpp b/Engine/Plugins/Experimental/EditorDataStorageFeatures/Source/TedsAssetData/Private/Processors/AssetProcessors.cpp
--- a/Engine/Plugins/Experimental/EditorDataStorageFeatures/Source/TedsAssetData/Private/Processors/AssetProcessors.cpp
+++ b/Engine/Plugins/Experimental/EditorDataStorageFeatures/Source/TedsAssetData/Private/Processors/AssetProcessors.cpp
@@ -132,7 +132,7 @@ void UTedsAssetDataFactory::RegisterQueries(UE::Editor::DataStorage::ICoreProvid
 			[](IQueryContext& Context, RowHandle Row, const FAssetPathColumn_Experimental& AssetPathColumn)
 			{
 				// Only do work if the state in the world doesn't match TEDS already
-				if (!UE::Editor::ContentBrowser::ExtensionUtils::IsFolderFavorite(AssetPathColumn.Path.ToString()))
+				if (!UE::Editor::ContentBrowser::ExtensionUtils::IsFolderFavorite(FContentBrowserItemPath(AssetPathColumn.Path.ToString(), EContentBrowserPathType::Internal)))
 				{
 					Context.PushCommand(FSetFolderFavoriteCommand{.FolderPath = AssetPathColumn.Path.ToString(), .bFavorite = true});
 				}
@@ -151,7 +151,7 @@ void UTedsAssetDataFactory::RegisterQueries(UE::Editor::DataStorage::ICoreProvid
 				[](IQueryContext& Context, RowHandle Row, const FAssetPathColumn_Experimental& AssetPathColumn)
 				{
 					// Only do work if the state in the world doesn't match TEDS already
-					if (UE::Editor::ContentBrowser::ExtensionUtils::IsFolderFavorite(AssetPathColumn.Path.ToString()))
+					if (UE::Editor::ContentBrowser::ExtensionUtils::IsFolderFavorite(FContentBrowserItemPath(AssetPathColumn.Path.ToString(), EContentBrowserPathType::Internal)))
 					{
 						Context.PushCommand(FSetFolderFavoriteCommand{.FolderPath = AssetPathColumn.Path.ToString(), .bFavorite = false});
 					}
diff --git a/Engine/Source/Editor/ContentBrowser/Private/Experimental/ContentBrowserExtensionUtils.cpp b/Engine/Source/Editor/ContentBrowser/Private/Experimental/ContentBrowserExtensionUtils.cpp
--- a/Engine/Source/Editor/ContentBrowser/Private/Experimental/ContentBrowserExtensionUtils.cpp
+++ b/Engine/Source/Editor/ContentBrowser/Private/Experimental/ContentBrowserExtensionUtils.cpp
@@ -59,9 +59,9 @@ namespace UE::Editor::ContentBrowser::ExtensionUtils
 		ContentBrowserUtils::SetPathColor(FolderPath.ToString(), FolderColor);
 	}
 
-	bool IsFolderFavorite(const FString& FolderPath)
+	bool IsFolderFavorite(const FContentBrowserItemPath& ItemPath)
 	{
-		return ContentBrowserUtils::IsFavoriteFolder(FContentBrowserItemPath(FolderPath, EContentBrowserPathType::Internal));
+		return ContentBrowserUtils::IsFavoriteFolder(ItemPath);
 	}
 
 	void SetFolderFavorite(const FString& FolderPath, bool bIsFavorite)
diff --git a/Engine/Source/Editor/ContentBrowser/Public/Experimental/ContentBrowserExtensionUtils.h b/Engine/Source/Editor/ContentBrowser/Public/Experimental/ContentBrowserExtensionUtils.h
--- a/Engine/Source/Editor/ContentBrowser/Public/Experimental/ContentBrowserExtensionUtils.h
+++ b/Engine/Source/Editor/ContentBrowser/Public/Experimental/ContentBrowserExtensionUtils.h
@@ -28,7 +28,7 @@ namespace UE::Editor::ContentBrowser::ExtensionUtils
 
 	// Returns if this folder has been marked as a favorite folder
 	UE_EXPERIMENTAL(5.8, "The Content Browser extension utils are experimental hooks that may be removed in the future")
-	CONTENTBROWSER_API bool IsFolderFavorite(const FString& FolderPath);
+	CONTENTBROWSER_API bool IsFolderFavorite(const FContentBrowserItemPath& ItemPath);
 
 	// Set if a folder is favorite
 	UE_EXPERIMENTAL(5.8, "The Content Browser extension utils are experimental hooks that may be removed in the future")
