[TEDS Mapping] Change the mapping functions to return an MapKey instead of an FString and have the prefab editor use FSoftObjectPaths for improved mapping accross instances with renaming and reinstancing mapping properly

#Synced-CL 50934773

diff --git a/Engine/Plugins/Experimental/EditorDataStorageFeatures/Source/TedsTableViewer/Private/TedsTableViewerUtils.cpp b/Engine/Plugins/Experimental/EditorDataStorageFeatures/Source/TedsTableViewer/Private/TedsTableViewerUtils.cpp
--- a/Engine/Plugins/Experimental/EditorDataStorageFeatures/Source/TedsTableViewer/Private/TedsTableViewerUtils.cpp
+++ b/Engine/Plugins/Experimental/EditorDataStorageFeatures/Source/TedsTableViewer/Private/TedsTableViewerUtils.cpp
@@ -186,14 +186,14 @@ namespace UE::Editor::DataStorage::TableViewerUtils
 
 	}
 
-	FString CreateRowMappingKey(const FName& HierarchyIdentifier, const TableViewerItemPtr Item)
+	void CreateRowMappingKey(const FName& InHierarchyIdentifier, const TableViewerItemPtr InItem, FMapKey& OutMapKey)
 	{
-		return CreateRowMappingKey(HierarchyIdentifier, Item.RowHandle);
+		CreateRowMappingKey(InHierarchyIdentifier, InItem.RowHandle, OutMapKey);
 	}
 	
-	FString CreateRowMappingKey(const FName& HierarchyIdentifier, const RowHandle Item)
+	void CreateRowMappingKey(const FName& InHierarchyIdentifier, const RowHandle InItem, FMapKey& OutMapKey)
 	{
-		return FString::Printf(TEXT("%s.%llu"), *HierarchyIdentifier.ToString(), Item);
+		OutMapKey = FMapKey(FString::Printf(TEXT("%s.%llu"), *InHierarchyIdentifier.ToString(), InItem));
 	}
 	
 	RowHandle GetTableViewerRowUiRow(const ICoreProvider* DataStorage, const RowHandle ChildWidgetRow)
@@ -240,21 +240,26 @@ namespace UE::Editor::DataStorage::TableViewerUtils
 			{
 				for (const TableViewerItemPtr& Item : Items)
 				{
+					FMapKey MappingKey;
 					if (GetTableViewerRowMapping.IsBound())
 					{
-						const FString MappingKey = GetTableViewerRowMapping.Execute(TableViewerIdentifier, Item);
-						// If the row already existed we will use that, but if not, we want to create a row for each entry so
-						// changing the default expansion state doesn't affect rows that come into view but aren't new
-						if (const RowHandle ExistingRow = DataStorage->LookupMappedRow(TableViewerUIRowMappingDomain, FMapKey(MappingKey)); !DataStorage->IsRowAvailable(ExistingRow))
+						GetTableViewerRowMapping.Execute(TableViewerIdentifier, Item, MappingKey);
+					}
+					if (!MappingKey.IsSet())
+					{
+						CreateRowMappingKey(TableViewerIdentifier, Item, MappingKey);
+					}
+					// If the row already existed we will use that, but if not, we want to create a row for each entry so
+					// changing the default expansion state doesn't affect rows that come into view but aren't new
+					if (const RowHandle ExistingRow = DataStorage->LookupMappedRow(TableViewerUIRowMappingDomain, MappingKey); !DataStorage->IsRowAvailable(ExistingRow))
+					{
+						const RowHandle AddedRow = DataStorage->AddRow(RowWidgetTableHandle);
+			
+						DataStorage->SetParentRow(WidgetHierarchyHandle, AddedRow, TableViewerRowHandle);
+						DataStorage->MapRow(TableViewerUIRowMappingDomain, MappingKey, AddedRow);
+						if (bAddHierarchyTags)
 						{
-							const RowHandle AddedRow = DataStorage->AddRow(RowWidgetTableHandle);
-				
-							DataStorage->SetParentRow(WidgetHierarchyHandle, AddedRow, TableViewerRowHandle);
-							DataStorage->MapRow(TableViewerUIRowMappingDomain, FMapKey(MappingKey), AddedRow);
-							if (bAddHierarchyTags)
-							{
-								DataStorage->AddColumns<FExpandedInUITag>(AddedRow);
-							}
+							DataStorage->AddColumns<FExpandedInUITag>(AddedRow);
 						}
 					}
 				}
diff --git a/Engine/Plugins/Experimental/EditorDataStorageFeatures/Source/TedsTableViewer/Private/Widgets/STedsHierarchyViewer.cpp b/Engine/Plugins/Experimental/EditorDataStorageFeatures/Source/TedsTableViewer/Private/Widgets/STedsHierarchyViewer.cpp
--- a/Engine/Plugins/Experimental/EditorDataStorageFeatures/Source/TedsTableViewer/Private/Widgets/STedsHierarchyViewer.cpp
+++ b/Engine/Plugins/Experimental/EditorDataStorageFeatures/Source/TedsTableViewer/Private/Widgets/STedsHierarchyViewer.cpp
@@ -27,8 +27,8 @@ namespace UE::Editor::DataStorage
 
 	SHierarchyViewer::~SHierarchyViewer()
 	{
-			RemoveAllTableViewerUIRows();
-		}
+		RemoveAllTableViewerUIRows();
+	}
 	
 	void SHierarchyViewer::Construct(
 		const FArguments& InArgs,
@@ -307,8 +307,9 @@ namespace UE::Editor::DataStorage
 	{
 		if (ICoreProvider* Storage = Model->GetDataStorageInterface())
 		{
-			const FString TableRowId = GetTableViewerRowMapping(TableViewerIdentifier, Item.RowHandle);
-			if (const RowHandle UIRow = Storage->LookupMappedRow(TableViewerUIRowMappingDomain, FMapKey(TableRowId)); Storage->IsRowAvailable(UIRow))
+			FMapKey MappingKey;
+			GetTableViewerRowMapping(TableViewerIdentifier, Item.RowHandle, MappingKey);
+			if (const RowHandle UIRow = Storage->LookupMappedRow(TableViewerUIRowMappingDomain, MappingKey); Storage->IsRowAvailable(UIRow))
 			{
 				if (bExpanded)
 				{
@@ -480,25 +481,27 @@ namespace UE::Editor::DataStorage
 		return TreeView->IsItemVisible(InItem);
 	}
 	
-	FString SHierarchyViewer::GetTableViewerRowMapping(const FName InTableViewerIdentifier, const RowHandle InRow) const
+	void SHierarchyViewer::GetTableViewerRowMapping(const FName& InTableViewerIdentifier, const RowHandle InRow, FMapKey& OutMapKey) const
 	{
 		if(GetCustomTableViewerRowMapping.IsBound())
 		{
-			if (FString RowMapping = GetCustomTableViewerRowMapping.Execute(InTableViewerIdentifier, InRow); !RowMapping.IsEmpty())
+			GetCustomTableViewerRowMapping.Execute(InTableViewerIdentifier, InRow, OutMapKey);
+			if (OutMapKey.IsSet())
 			{
-				return RowMapping;
+				return;
 			}
 		}
-		return TableViewerUtils::CreateRowMappingKey(InTableViewerIdentifier, InRow);
+		TableViewerUtils::CreateRowMappingKey(InTableViewerIdentifier, InRow, OutMapKey);
 	}
 	
 	RowHandle SHierarchyViewer::AddOrFindTableViewerUIRow(const TableViewerItemPtr InItem) const
 	{
 		if (ICoreProvider* Storage = Model->GetDataStorageInterface())
 		{
-			const FString MappingKey = GetTableViewerRowMapping(TableViewerIdentifier, InItem.RowHandle);
+			FMapKey MappingKey;
+			GetTableViewerRowMapping(TableViewerIdentifier, InItem.RowHandle, MappingKey);
 			// If the row already existed but was removed from TreeView's SparseItemInfos (Searching, Filters, etc.), re-expand it
-			if (const RowHandle ExistingRow = Storage->LookupMappedRow(TableViewerUIRowMappingDomain, FMapKey(MappingKey)); Storage->IsRowAvailable(ExistingRow))
+			if (const RowHandle ExistingRow = Storage->LookupMappedRow(TableViewerUIRowMappingDomain, MappingKey); Storage->IsRowAvailable(ExistingRow))
 			{
 				if(Storage->HasColumns<FExpandedInUITag>(ExistingRow) && !TreeView->IsItemExpanded(InItem))
 				{
@@ -511,7 +514,7 @@ namespace UE::Editor::DataStorage
 				const RowHandle AddedRow = Storage->AddRow(Storage->FindTable(TableViewerUtils::GetRowWidgetTableName()));
 				
 				Storage->SetParentRow(Storage->FindHierarchyByName(TableViewerUtils::GetWidgetHierarchyName()), AddedRow, GetWidgetRowHandle());
-				Storage->MapRow(TableViewerUIRowMappingDomain, FMapKey(MappingKey), AddedRow);
+				Storage->MapRow(TableViewerUIRowMappingDomain, MappingKey, AddedRow);
 
 				const bool bIsItemExpandedInTree = TreeView->IsItemExpanded(InItem);
 				if (bIsItemExpandedInTree != bExpandNewRows)
diff --git a/Engine/Plugins/Experimental/EditorDataStorageFeatures/Source/TedsTableViewer/Private/Widgets/STedsTableViewer.cpp b/Engine/Plugins/Experimental/EditorDataStorageFeatures/Source/TedsTableViewer/Private/Widgets/STedsTableViewer.cpp
--- a/Engine/Plugins/Experimental/EditorDataStorageFeatures/Source/TedsTableViewer/Private/Widgets/STedsTableViewer.cpp
+++ b/Engine/Plugins/Experimental/EditorDataStorageFeatures/Source/TedsTableViewer/Private/Widgets/STedsTableViewer.cpp
@@ -332,25 +332,27 @@ namespace UE::Editor::DataStorage
 		return ListView->IsItemVisible(InItem);
 	}
 	
-	FString STedsTableViewer::GetTableViewerRowMapping(const FName InTableViewerIdentifier, const RowHandle InRow) const
+	void STedsTableViewer::GetTableViewerRowMapping(const FName& InTableViewerIdentifier, const RowHandle InRow, FMapKey& OutMapKey) const
 	{
 		if(GetCustomTableViewerRowMapping.IsBound())
 		{
-			if (FString RowMapping = GetCustomTableViewerRowMapping.Execute(InTableViewerIdentifier, InRow); !RowMapping.IsEmpty())
+			GetCustomTableViewerRowMapping.Execute(InTableViewerIdentifier, InRow, OutMapKey);
+			if (OutMapKey.IsSet())
 			{
-				return RowMapping;
+				return;
 			}
 		}
-		return TableViewerUtils::CreateRowMappingKey(InTableViewerIdentifier, InRow);
+		TableViewerUtils::CreateRowMappingKey(InTableViewerIdentifier, InRow, OutMapKey);
 	}
 
 	RowHandle STedsTableViewer::AddOrFindTableViewerUIRow(const TableViewerItemPtr InItem) const
 	{
 		if (ICoreProvider* Storage = Model->GetDataStorageInterface())
 		{
-			const FString MappingKey = GetTableViewerRowMapping(TableViewerIdentifier, InItem.RowHandle);
+			FMapKey MappingKey;
+			GetTableViewerRowMapping(TableViewerIdentifier, InItem.RowHandle, MappingKey);
 			// If the row already existed but was removed from TreeView's SparseItemInfos (Searching, Filters, etc.), re-expand it
-			if (const RowHandle ExistingRow = Storage->LookupMappedRow(TableViewerUIRowMappingDomain, FMapKey(MappingKey)); Storage->IsRowAvailable(ExistingRow))
+			if (const RowHandle ExistingRow = Storage->LookupMappedRow(TableViewerUIRowMappingDomain, MappingKey); Storage->IsRowAvailable(ExistingRow))
 			{
 				return ExistingRow;
 			}
@@ -359,7 +361,7 @@ namespace UE::Editor::DataStorage
 				const RowHandle AddedRow = Storage->AddRow(Storage->FindTable(TableViewerUtils::GetRowWidgetTableName()));
 				
 				Storage->SetParentRow(Storage->FindHierarchyByName(TableViewerUtils::GetWidgetHierarchyName()), AddedRow, GetWidgetRowHandle());
-				Storage->MapRow(TableViewerUIRowMappingDomain, FMapKey(MappingKey), AddedRow);
+				Storage->MapRow(TableViewerUIRowMappingDomain, MappingKey, AddedRow);
 				
 				return AddedRow;
 			}
diff --git a/Engine/Plugins/Experimental/EditorDataStorageFeatures/Source/TedsTableViewer/Public/TedsTableViewerUtils.h b/Engine/Plugins/Experimental/EditorDataStorageFeatures/Source/TedsTableViewer/Public/TedsTableViewerUtils.h
--- a/Engine/Plugins/Experimental/EditorDataStorageFeatures/Source/TedsTableViewer/Public/TedsTableViewerUtils.h
+++ b/Engine/Plugins/Experimental/EditorDataStorageFeatures/Source/TedsTableViewer/Public/TedsTableViewerUtils.h
@@ -24,7 +24,7 @@ struct FSlateBrush;
 namespace UE::Editor::DataStorage::TableViewerUtils
 {
 	// Delegate fired to get a custom row mapping for the table viewer
-	DECLARE_DELEGATE_RetVal_TwoParams(FString, FGetTableViewerRowMapping, FName, RowHandle);
+	DECLARE_DELEGATE_ThreeParams(FGetTableViewerRowMapping, const FName&, const RowHandle, FMapKey&);
 	
 	TEDSTABLEVIEWER_API FName GetWidgetTableName();
 
@@ -48,10 +48,10 @@ namespace UE::Editor::DataStorage::TableViewerUtils
 	TEDSTABLEVIEWER_API const FSlateBrush* GetIconForRow(ICoreProvider* DataStorage, RowHandle Row);
 
 	// Creates the default Row Map Key for the Table Row UI Rows
-	TEDSTABLEVIEWER_API FString CreateRowMappingKey(const FName& HierarchyIdentifier, const TableViewerItemPtr Item);
+	TEDSTABLEVIEWER_API void CreateRowMappingKey(const FName& InHierarchyIdentifier, const TableViewerItemPtr InItem, FMapKey& OutMapKey);
 	
 	// Creates the default Row Map Key for the Table Row UI Rows from the RowHandle directly
-    TEDSTABLEVIEWER_API FString CreateRowMappingKey(const FName& HierarchyIdentifier, const RowHandle Item);
+	TEDSTABLEVIEWER_API void CreateRowMappingKey(const FName& InHierarchyIdentifier, const RowHandle InItem, FMapKey& OutMapKey);
 
 	// Walks up the UI Widget Hierarchy until finding the row with the FTableViewerRowTag
 	TEDSTABLEVIEWER_API RowHandle GetTableViewerRowUiRow(const ICoreProvider* DataStorage, const RowHandle ChildWidgetRow);
diff --git a/Engine/Plugins/Experimental/EditorDataStorageFeatures/Source/TedsTableViewer/Public/Widgets/STedsHierarchyViewer.h b/Engine/Plugins/Experimental/EditorDataStorageFeatures/Source/TedsTableViewer/Public/Widgets/STedsHierarchyViewer.h
--- a/Engine/Plugins/Experimental/EditorDataStorageFeatures/Source/TedsTableViewer/Public/Widgets/STedsHierarchyViewer.h
+++ b/Engine/Plugins/Experimental/EditorDataStorageFeatures/Source/TedsTableViewer/Public/Widgets/STedsHierarchyViewer.h
@@ -50,7 +50,7 @@ namespace UE::Editor::DataStorage
 	class SHierarchyViewer : public SCompoundWidget, public ITableViewer
 	{
 	friend class STedsCompositeHierarchyViewer;
-	
+		
 	public:
 		
 		enum class EExpansionState : uint8
@@ -174,7 +174,7 @@ namespace UE::Editor::DataStorage
 		
 	protected:
 		
-		FString GetTableViewerRowMapping(const FName InTableViewerIdentifier, const RowHandle InRow) const;
+		void GetTableViewerRowMapping(const FName& InTableViewerIdentifier, const RowHandle InRow, FMapKey& OutMapKey) const;
 		
 		RowHandle AddOrFindTableViewerUIRow(const TableViewerItemPtr InItem) const;
 
diff --git a/Engine/Plugins/Experimental/EditorDataStorageFeatures/Source/TedsTableViewer/Public/Widgets/STedsTableViewer.h b/Engine/Plugins/Experimental/EditorDataStorageFeatures/Source/TedsTableViewer/Public/Widgets/STedsTableViewer.h
--- a/Engine/Plugins/Experimental/EditorDataStorageFeatures/Source/TedsTableViewer/Public/Widgets/STedsTableViewer.h
+++ b/Engine/Plugins/Experimental/EditorDataStorageFeatures/Source/TedsTableViewer/Public/Widgets/STedsTableViewer.h
@@ -153,7 +153,7 @@ namespace UE::Editor::DataStorage
 
 	protected:
 	
-		FString GetTableViewerRowMapping(const FName InTableViewerIdentifier, const RowHandle InRow) const;
+		void GetTableViewerRowMapping(const FName& InTableViewerIdentifier, const RowHandle InRow, FMapKey& OutMapKey) const;
 
 		RowHandle AddOrFindTableViewerUIRow(const TableViewerItemPtr InItem) const;
 
