Dataflow : fix ununitialized variable that cause crashs in Dataflow
- fix relates to CL 50923524

#Synced-CL 50953304

diff --git a/Engine/Plugins/Experimental/EditorDataStorageFeatures/Source/TedsOutliner/Private/TedsOutlinerImpl.cpp b/Engine/Plugins/Experimental/EditorDataStorageFeatures/Source/TedsOutliner/Private/TedsOutlinerImpl.cpp
--- a/Engine/Plugins/Experimental/EditorDataStorageFeatures/Source/TedsOutliner/Private/TedsOutlinerImpl.cpp
+++ b/Engine/Plugins/Experimental/EditorDataStorageFeatures/Source/TedsOutliner/Private/TedsOutlinerImpl.cpp
@@ -228,6 +228,7 @@ FTedsOutlinerParams::FTedsOutlinerParams(SSceneOutliner* InSceneOutliner)
 	, bForceShowParents(true)
 	, bUseDefaultObservers(true)
 	, HierarchyData(FTedsOutlinerHierarchyData::GetDefaultHierarchyData())
+	, SelectionSet(nullptr)
 {
 	CellWidgetPurpose =
 		DataStorage::IUiProvider::FPurposeID(DataStorage::IUiProvider::FPurposeInfo(
