[TEDS Viewers] Added an argument for destroying the TableViewerRowUIRows on destruction of the viewer.
* Added this argument to the prefab editor with a CVar disabling it for now

#Synced-CL 50970343

diff --git a/Engine/Plugins/Experimental/EditorDataStorageFeatures/Source/TedsTableViewer/Private/Widgets/Composite/STedsCompositeHierarchyViewer.cpp b/Engine/Plugins/Experimental/EditorDataStorageFeatures/Source/TedsTableViewer/Private/Widgets/Composite/STedsCompositeHierarchyViewer.cpp
--- a/Engine/Plugins/Experimental/EditorDataStorageFeatures/Source/TedsTableViewer/Private/Widgets/Composite/STedsCompositeHierarchyViewer.cpp
+++ b/Engine/Plugins/Experimental/EditorDataStorageFeatures/Source/TedsTableViewer/Private/Widgets/Composite/STedsCompositeHierarchyViewer.cpp
@@ -203,6 +203,7 @@ namespace UE::Editor::DataStorage
 			.ColumnMetaData(HierarchyViewerArgs._ColumnMetaData)
 			.ListSelectionMode(HierarchyViewerArgs._ListSelectionMode)
 			.OnSelectionChanged(HierarchyViewerArgs._OnSelectionChanged)
+			.PersistUIRowsOnClose(HierarchyViewerArgs._PersistUIRowsOnClose)
 			.GetCustomTableViewerRowMapping(HierarchyViewerArgs._GetCustomTableViewerRowMapping)
 			.ExpandNewRows(HierarchyViewerArgs._ExpandNewRows)
 			.PrimaryColumn(HierarchyViewerArgs._PrimaryColumn)
diff --git a/Engine/Plugins/Experimental/EditorDataStorageFeatures/Source/TedsTableViewer/Private/Widgets/STedsHierarchyViewer.cpp b/Engine/Plugins/Experimental/EditorDataStorageFeatures/Source/TedsTableViewer/Private/Widgets/STedsHierarchyViewer.cpp
--- a/Engine/Plugins/Experimental/EditorDataStorageFeatures/Source/TedsTableViewer/Private/Widgets/STedsHierarchyViewer.cpp
+++ b/Engine/Plugins/Experimental/EditorDataStorageFeatures/Source/TedsTableViewer/Private/Widgets/STedsHierarchyViewer.cpp
@@ -27,7 +27,10 @@ namespace UE::Editor::DataStorage
 
 	SHierarchyViewer::~SHierarchyViewer()
 	{
-		RemoveAllTableViewerUIRows();
+		if (!bPersistUIRowsOnClose)
+		{
+			RemoveAllTableViewerUIRows();
+		}
 	}
 	
 	void SHierarchyViewer::Construct(
@@ -36,6 +39,7 @@ namespace UE::Editor::DataStorage
 	{
 		HierarchyInterface = MoveTemp(InHierarchyInterface);
 		OnSelectionChanged = InArgs._OnSelectionChanged;
+		bPersistUIRowsOnClose = InArgs._PersistUIRowsOnClose;
 		GetCustomTableViewerRowMapping = InArgs._GetCustomTableViewerRowMapping;
 		bExpandNewRows = InArgs._ExpandNewRows;
 		EmptyRowsMessage = InArgs._EmptyRowsMessage;
@@ -151,13 +155,7 @@ namespace UE::Editor::DataStorage
 			{
 				if (RowNode->GetRevision() != QueryStack::INode::UninitializedRevisionId)
 				{
-					TableViewerUtils::InitializeAllTableViewerUIRows(
-						Model->GetDataStorageInterface(),
-						Model->GetItems(), 
-						TableViewerUtils::FGetTableViewerRowMapping::CreateSP(this, &SHierarchyViewer::GetTableViewerRowMapping), 
-						TableViewerIdentifier, 
-						GetWidgetRowHandle(),
-						true);
+					InitializeAllTableViewerUIRows();
 					bModelInitialized = true;
 				}
 
@@ -370,11 +368,25 @@ namespace UE::Editor::DataStorage
 
 	void SHierarchyViewer::SetTableViewerIdentifier(const FName& NewIdentifier)
 	{
-		// This will remove any pre-existing row info with the previous name
-		RemoveAllTableViewerUIRows();
+		// Only use this function when it is necessary to get a name after initialization, the rows associated with the
+		// previous name will be removed regardless of the bRemoveTableViewerRowsOnClose argument and re-initialized.
+		// If the model hasn't been initialized yet, there are no rows to remove so we can skip this
+		if (bModelInitialized)
+		{
+			// Remove while set to the old identifier
+			RemoveAllTableViewerUIRows();
+		}
+ 
 		TableViewerIdentifier = NewIdentifier.IsNone() ?
 			FName("TEDSHierarchyViewerID", InstanceCounter) :
 			NewIdentifier;
+		
+		// If the model hasn't been initialized yet, we can skip initializing the new rows as it will be done later
+		if (bModelInitialized)
+		{
+			// We want to reinitialize the current rows as table viewer row UI rows under the new mapping
+			InitializeAllTableViewerUIRows();
+		}
 	}
 
 	void SHierarchyViewer::AddCustomRowWidget(const TSharedRef<FTedsTableViewerColumn>& InColumn)
@@ -534,6 +546,17 @@ namespace UE::Editor::DataStorage
 		return InvalidRowHandle;
 	}
 
+	void SHierarchyViewer::InitializeAllTableViewerUIRows() const
+	{
+		TableViewerUtils::InitializeAllTableViewerUIRows(
+			Model->GetDataStorageInterface(),
+			Model->GetItems(), 
+			TableViewerUtils::FGetTableViewerRowMapping::CreateSP(this, &SHierarchyViewer::GetTableViewerRowMapping), 
+			TableViewerIdentifier, 
+			GetWidgetRowHandle(),
+			true);
+	}
+
 	void SHierarchyViewer::RemoveAllTableViewerUIRows() const
 	{
 		if(TedsWidget.IsValid() && Model.IsValid())
diff --git a/Engine/Plugins/Experimental/EditorDataStorageFeatures/Source/TedsTableViewer/Private/Widgets/STedsTableViewer.cpp b/Engine/Plugins/Experimental/EditorDataStorageFeatures/Source/TedsTableViewer/Private/Widgets/STedsTableViewer.cpp
--- a/Engine/Plugins/Experimental/EditorDataStorageFeatures/Source/TedsTableViewer/Private/Widgets/STedsTableViewer.cpp
+++ b/Engine/Plugins/Experimental/EditorDataStorageFeatures/Source/TedsTableViewer/Private/Widgets/STedsTableViewer.cpp
@@ -25,13 +25,17 @@ namespace UE::Editor::DataStorage
 
 	STedsTableViewer::~STedsTableViewer()
 	{
-		RemoveAllTableViewerUIRows();
+		if (!bPersistUIRowsOnClose)
+		{
+			RemoveAllTableViewerUIRows();
+		}
 	}
 	
 	void STedsTableViewer::Construct(const FArguments& InArgs)
 	{
 		OnSelectionChanged = InArgs._OnSelectionChanged;
 		GetCustomTableViewerRowMapping = InArgs._GetCustomTableViewerRowMapping;
+		bPersistUIRowsOnClose = InArgs._PersistUIRowsOnClose;
 		EmptyRowsMessage = InArgs._EmptyRowsMessage;
 		ItemHeight = InArgs._ItemHeight;
 		ItemPadding = InArgs._ItemPadding;
@@ -115,12 +119,7 @@ namespace UE::Editor::DataStorage
 			{
 				if (RowNode->GetRevision() != QueryStack::INode::UninitializedRevisionId)
 				{
-					TableViewerUtils::InitializeAllTableViewerUIRows(
-						Model->GetDataStorageInterface(),
-						Model->GetItems(),
-						TableViewerUtils::FGetTableViewerRowMapping::CreateSP(this, &STedsTableViewer::GetTableViewerRowMapping), 
-						TableViewerIdentifier, 
-						GetWidgetRowHandle());
+					InitializeAllTableViewerUIRows();
 					bModelInitialized = true;
 				}
 			}
@@ -254,11 +253,25 @@ namespace UE::Editor::DataStorage
 
 	void STedsTableViewer::SetTableViewerIdentifier(const FName& NewIdentifier)
 	{
-		// This will remove any pre-existing row info with the previous name
-		RemoveAllTableViewerUIRows();
+		// Only use this function when it is necessary to get a name after initialization, the rows associated with the
+		// previous name will be removed regardless of the bRemoveTableViewerRowsOnClose argument and re-initialized.
+		// If the model hasn't been initialized yet, there are no rows to remove so we can skip this
+		if (bModelInitialized)
+		{
+			// Remove while set to the old identifier
+			RemoveAllTableViewerUIRows();
+		}
+ 
 		TableViewerIdentifier = NewIdentifier.IsNone() ?
 			FName("TEDSTableViewerID", InstanceCounter) :
 			NewIdentifier;
+		
+		// If the model hasn't been initialized yet, we can skip initializing the new rows as it will be done later
+		if (bModelInitialized)
+		{
+			// We want to reinitialize the current rows as table viewer row UI rows under the new mapping
+			InitializeAllTableViewerUIRows();
+		}
 	}
 
 	void STedsTableViewer::SetQueryStack(TSharedPtr<QueryStack::IRowNode> InRowQueryStack)
@@ -370,6 +383,16 @@ namespace UE::Editor::DataStorage
 		return InvalidRowHandle;
 	}
 
+	void STedsTableViewer::InitializeAllTableViewerUIRows() const
+	{
+		TableViewerUtils::InitializeAllTableViewerUIRows(
+			Model->GetDataStorageInterface(),
+			Model->GetItems(), 
+			TableViewerUtils::FGetTableViewerRowMapping::CreateSP(this, &STedsTableViewer::GetTableViewerRowMapping), 
+			TableViewerIdentifier, 
+			GetWidgetRowHandle());
+	}
+
 	void STedsTableViewer::RemoveAllTableViewerUIRows() const
 	{
 		if(TedsWidget.IsValid() && Model.IsValid())
diff --git a/Engine/Plugins/Experimental/EditorDataStorageFeatures/Source/TedsTableViewer/Public/Widgets/STedsHierarchyViewer.h b/Engine/Plugins/Experimental/EditorDataStorageFeatures/Source/TedsTableViewer/Public/Widgets/STedsHierarchyViewer.h
--- a/Engine/Plugins/Experimental/EditorDataStorageFeatures/Source/TedsTableViewer/Public/Widgets/STedsHierarchyViewer.h
+++ b/Engine/Plugins/Experimental/EditorDataStorageFeatures/Source/TedsTableViewer/Public/Widgets/STedsHierarchyViewer.h
@@ -66,6 +66,7 @@ namespace UE::Editor::DataStorage
 		SLATE_BEGIN_ARGS(SHierarchyViewer)
 			: _CanSelectGeneratedColumn(true)
 			, _ListSelectionMode(ESelectionMode::Type::Single)
+			, _PersistUIRowsOnClose(false)
 			, _ExpandNewRows(true)
 			, _AllowInvisibleItemSelection(false)
 		{
@@ -98,6 +99,9 @@ namespace UE::Editor::DataStorage
 		// The selection mode for the table viewer (single/multi etc)
 		SLATE_ARGUMENT(ESelectionMode::Type, ListSelectionMode)
 
+		// Persists all TableViewer UI rows from this viewer past destruction
+		SLATE_ARGUMENT(bool, PersistUIRowsOnClose);
+
 		// Delegate called to define a custom row mapping for the Table Viewer Row Widget Rows
 		// Used in cases where you need the UI Rows to persist across RowHandle changes
 		SLATE_EVENT(TableViewerUtils::FGetTableViewerRowMapping, GetCustomTableViewerRowMapping)
@@ -178,6 +182,8 @@ namespace UE::Editor::DataStorage
 		
 		RowHandle AddOrFindTableViewerUIRow(const TableViewerItemPtr InItem) const;
 
+		void InitializeAllTableViewerUIRows() const;
+
 		void RemoveAllTableViewerUIRows() const;
 		
 		TSharedRef<ITableRow> MakeTableRowWidget(TableViewerItemPtr InItem, const TSharedRef<STableViewBase>& OwnerTable) const;
@@ -224,6 +230,10 @@ namespace UE::Editor::DataStorage
 
 		// Delegate fired when the selection changes
 		FOnSelectionChanged OnSelectionChanged;
+
+		// If we want to persist all Table Viewer Rows and their mapping to this table ID when the viewer is destructed
+		// (keeping them around will persist its data while the editor is open)
+		bool bPersistUIRowsOnClose = false;
 		
 		// Delegate fired when getting the row mapping, will use the default 'TableId.DataRowHandle' if nothing is bound
 		TableViewerUtils::FGetTableViewerRowMapping GetCustomTableViewerRowMapping;
diff --git a/Engine/Plugins/Experimental/EditorDataStorageFeatures/Source/TedsTableViewer/Public/Widgets/STedsTableViewer.h b/Engine/Plugins/Experimental/EditorDataStorageFeatures/Source/TedsTableViewer/Public/Widgets/STedsTableViewer.h
--- a/Engine/Plugins/Experimental/EditorDataStorageFeatures/Source/TedsTableViewer/Public/Widgets/STedsTableViewer.h
+++ b/Engine/Plugins/Experimental/EditorDataStorageFeatures/Source/TedsTableViewer/Public/Widgets/STedsTableViewer.h
@@ -51,6 +51,7 @@ namespace UE::Editor::DataStorage
 		SLATE_BEGIN_ARGS(STedsTableViewer)
 			: _CanSelectGeneratedColumn(true)
 			, _ListSelectionMode(ESelectionMode::Type::Single)
+			, _PersistUIRowsOnClose(false)
 			, _ItemHeight()
 			, _ItemPadding(FMargin(0.f))
 		{
@@ -84,6 +85,9 @@ namespace UE::Editor::DataStorage
 		// The selection mode for the table viewer (single/multi etc)
 		SLATE_ARGUMENT(ESelectionMode::Type, ListSelectionMode)
 
+		// Persists all TableViewer Row UI rows from this viewer past destruction
+		SLATE_ARGUMENT(bool, PersistUIRowsOnClose);
+
 		// Delegate called to define a custom row mapping for the Table Viewer Row Widget Rows
 		// Used in cases where you need the UI Rows to persist across RowHandle changes
 		SLATE_EVENT(TableViewerUtils::FGetTableViewerRowMapping, GetCustomTableViewerRowMapping)
@@ -157,6 +161,8 @@ namespace UE::Editor::DataStorage
 
 		RowHandle AddOrFindTableViewerUIRow(const TableViewerItemPtr InItem) const;
 
+		void InitializeAllTableViewerUIRows() const;
+
 		void RemoveAllTableViewerUIRows() const;
 		
 		TSharedRef<ITableRow> MakeTableRowWidget(TableViewerItemPtr InItem, const TSharedRef<STableViewBase>& OwnerTable) const;
@@ -198,6 +204,10 @@ namespace UE::Editor::DataStorage
 		// Delegate fired when the selection changes
 		FOnSelectionChanged OnSelectionChanged;
 		
+		// If we want to persist all Table Viewer Rows and their mapping to this table ID when the viewer is destructed
+		// (keeping them around will persist its data while the editor is open)
+		bool bPersistUIRowsOnClose = false;
+        		
 		// Delegate fired when getting the row mapping, will use the default 'TableId.DataRowHandle' if nothing is bound
 		TableViewerUtils::FGetTableViewerRowMapping GetCustomTableViewerRowMapping;
 
