Favorites Outliner Column for the Level Editor and Prefab Editor

#Synced-CL 50983539

diff --git a/Engine/Plugins/Experimental/EditorDataStorageFeatures/Source/TedsOutliner/Private/TedsOutlinerModule.cpp b/Engine/Plugins/Experimental/EditorDataStorageFeatures/Source/TedsOutliner/Private/TedsOutlinerModule.cpp
--- a/Engine/Plugins/Experimental/EditorDataStorageFeatures/Source/TedsOutliner/Private/TedsOutlinerModule.cpp
+++ b/Engine/Plugins/Experimental/EditorDataStorageFeatures/Source/TedsOutliner/Private/TedsOutlinerModule.cpp
@@ -33,6 +33,7 @@ namespace UE::Editor::Outliner
 	namespace Private
 	{
 		static bool bUseNewRevisionControlWidgets = false;
+		static bool bUseFavoritesWidgets = false;
 	} // namespace Private
 
 	void RefreshLevelEditorTedsOutliner(bool bAlwaysInvoke)
@@ -57,6 +58,15 @@ namespace UE::Editor::Outliner
 			RefreshLevelEditorTedsOutliner(false);
 		}));
 
+	static FAutoConsoleVariableRef CVarUseFavoritesWidgets(
+		TEXT("TEDS.UI.UseFavoritesWidgets.LevelEditorOutliner"),
+		Private::bUseFavoritesWidgets,
+		TEXT("Use unified favorites in the Level Editor Outliner (requires TEDS-Outliner to be enabled)")
+		, FConsoleVariableDelegate::CreateLambda([](IConsoleVariable*)
+			{
+				RefreshLevelEditorTedsOutliner(false);
+			}));
+
 	// CVar to summon the TEDS-Outliner as a separate tab
 	static FAutoConsoleCommand OpenTableViewerConsoleCommand(
 		TEXT("TEDS.UI.OpenTedsOutliner"),
@@ -130,6 +140,11 @@ FTedsOutlinerColumnDescription FTedsOutlinerModule::GetLevelEditorTedsOutlinerCo
 		FLevelColumn::StaticStruct(),
 		FVisibleInEditorColumn::StaticStruct()});
 
+	if (CVarUseFavoritesWidgets->GetBool())
+	{
+		OutColumns.Add(FFavoriteTag::StaticStruct());
+	}
+
 	if (CVarUseNewRevisionControlWidgets->GetBool())
 	{
 		OutColumns.Add(FTypedElementPackageReference::StaticStruct());
diff --git a/Engine/Plugins/Experimental/EditorDataStorageFeatures/Source/TedsOutliner/Private/Widgets/OutlinerFavoriteWidget.cpp b/Engine/Plugins/Experimental/EditorDataStorageFeatures/Source/TedsOutliner/Private/Widgets/OutlinerFavoriteWidget.cpp
new file mode 100644
--- /dev/null
+++ b/Engine/Plugins/Experimental/EditorDataStorageFeatures/Source/TedsOutliner/Private/Widgets/OutlinerFavoriteWidget.cpp
@@ -0,0 +1,284 @@
+// Copyright Epic Games, Inc. All Rights Reserved.
+
+#include "Widgets/OutlinerFavoriteWidget.h"
+
+#include "ScopedTransaction.h"
+#include "Columns/SlateHeaderColumns.h"
+#include "Columns/TedsActorComponentCompatibilityColumns.h"
+#include "Elements/Columns/TypedElementCompatibilityColumns.h"
+#include "Elements/Columns/TypedElementFolderColumns.h"
+#include "Elements/Columns/TypedElementLabelColumns.h"
+#include "Elements/Columns/TypedElementMiscColumns.h"
+#include "Elements/Framework/TypedElementAttributeBinding.h"
+#include "Styling/AppStyle.h"
+#include "Styling/StyleColors.h"
+#include "Widgets/Images/SImage.h"
+#include "Widgets/Input/SButton.h"
+#include "Widgets/Layout/SBox.h"
+#include "Widgets/Text/STextBlock.h"
+
+#include UE_INLINE_GENERATED_CPP_BY_NAME(OutlinerFavoriteWidget)
+
+#define LOCTEXT_NAMESPACE "OutlinerFavoriteWidget"
+
+namespace UE::Editor::DataStorage::Private
+{
+	//
+	// Column Sorter
+	//
+
+	class FFavoriteWidgetSorter final : public FColumnSorterInterface
+	{
+	public:
+		virtual int32 Compare(const ICoreProvider& Storage, RowHandle Left, RowHandle Right) const override
+		{
+			return 0;
+		}
+
+		virtual FPrefixInfo CalculatePrefix(const ICoreProvider& Storage, RowHandle Row, uint32 ByteIndex) const override
+		{
+			return CreateSortPrefix(ByteIndex, static_cast<int32>(!Storage.HasColumns<FFavoriteTag>(Row)));
+		}
+
+		virtual ESortType GetSortType() const override
+		{
+			return ESortType::FixedSize64;
+		}
+
+		virtual FText GetShortName() const override
+		{
+			return LOCTEXT("FavoriteWidgetSorter", "Favorite Sorter");
+		}
+	};
+
+	//
+	// Widgets
+	//
+
+	/** Widget responsible for managing the favorite state of a TEDS row handle */
+	class STedsFavoriteWidget : public SImage
+	{
+	public:
+		SLATE_BEGIN_ARGS(STedsFavoriteWidget) {}
+		SLATE_END_ARGS()
+
+		/** Construct this widget */
+		void Construct(const FArguments& InArgs, ICoreProvider* InDataStorage, RowHandle InTargetRow)
+		{
+			DataStorage = InDataStorage;
+			TargetRow = InTargetRow;
+
+			SImage::Construct(
+				SImage::FArguments()
+				.ColorAndOpacity(this, &STedsFavoriteWidget::GetForegroundColor)
+				.Image(this, &STedsFavoriteWidget::GetBrush)
+			);
+
+			FavoriteHoveredBrush = FAppStyle::GetBrush("Icons.Star");
+			FavoriteNotHoveredBrush = FAppStyle::GetBrush("Icons.Star");
+			NotFavoriteHoveredBrush = FAppStyle::GetBrush("PropertyWindow.Favorites_Disabled");
+			NotFavoriteNotHoveredBrush = FAppStyle::GetBrush("PropertyWindow.Favorites_Disabled");
+		}
+
+	private:
+		FReply HandleClick() const
+		{
+			if (!IsEnabled())
+			{
+				return FReply::Unhandled();
+			}
+
+			ToggleFavorite();
+
+			return FReply::Handled();
+		}
+
+		virtual FReply OnMouseButtonDoubleClick(const FGeometry& InMyGeometry, const FPointerEvent& InMouseEvent) override
+		{
+			return HandleClick();
+		}
+
+		virtual FReply OnMouseButtonDown(const FGeometry& MyGeometry, const FPointerEvent& MouseEvent) override
+		{
+			if (MouseEvent.GetEffectingButton() == EKeys::LeftMouseButton)
+			{
+				return HandleClick();
+			}
+
+			return FReply::Unhandled();
+		}
+
+		virtual FReply OnMouseButtonUp(const FGeometry& MyGeometry, const FPointerEvent& MouseEvent) override
+		{
+			if (MouseEvent.GetEffectingButton() == EKeys::LeftMouseButton)
+			{
+				return FReply::Handled();
+			}
+
+			return FReply::Unhandled();
+		}
+
+		const FSlateBrush* GetBrush() const
+		{
+			if (IsFavorite())
+			{
+				return IsHovered() ? FavoriteHoveredBrush : FavoriteNotHoveredBrush;
+			}
+
+			return IsHovered() ? NotFavoriteHoveredBrush : NotFavoriteNotHoveredBrush;
+		}
+
+		virtual FSlateColor GetForegroundColor() const override
+		{
+			if (!IsFavorite() && !IsHovered())
+			{
+				return FLinearColor::Transparent;
+			}
+
+			FLinearColor Color =
+				FAppStyle::Get().GetSlateColor("Colors.Foreground")
+				.GetSpecifiedColor();
+
+			if (IsHovered())
+			{
+				constexpr float Darken = .8f;
+				Color *= Darken;
+			}
+
+			return Color;
+		}
+
+		bool IsFavorite() const
+		{
+			return DataStorage->HasColumns<FFavoriteTag>(TargetRow);
+		}
+
+		void ToggleFavorite() const
+		{
+			DataStorage->AddColumn<FTypedElementSyncFromWorldTag>(TargetRow);
+
+			if (DataStorage->HasColumns<FFavoriteTag>(TargetRow))
+			{
+				DataStorage->RemoveColumn<FFavoriteTag>(TargetRow);
+			}
+			else
+			{
+				DataStorage->AddColumn<FFavoriteTag>(TargetRow);
+			}
+		}
+
+	private:
+		ICoreProvider* DataStorage;
+		RowHandle TargetRow;
+
+		/** Favorite brushes for the various states */
+		const FSlateBrush* FavoriteHoveredBrush;
+		const FSlateBrush* FavoriteNotHoveredBrush;
+		const FSlateBrush* NotFavoriteHoveredBrush;
+		const FSlateBrush* NotFavoriteNotHoveredBrush;
+	};
+
+	//
+	// Helpers
+	//
+
+	TSharedRef<SWidget> MakeFavoritesWidgetContainer(const TSharedRef<SWidget>& FavoritesWidget)
+	{
+		return SNew(SBox)
+			.VAlign(VAlign_Center)
+			.HAlign(HAlign_Center)
+			.HeightOverride(18.f)
+			.WidthOverride(18.f)
+			[
+				FavoritesWidget
+			];
+	}
+}
+
+
+void UOutlinerFavoriteWidgetFactory::RegisterWidgetConstructors(UE::Editor::DataStorage::ICoreProvider& DataStorage,
+                                                                UE::Editor::DataStorage::IUiProvider& DataStorageUi) const
+{
+	using namespace UE::Editor::DataStorage::Queries;
+
+	FConditions SupportedColumns = TColumn<FFavoriteTag>() || TColumn<FTypedElementUObjectColumn>() || TColumn<FFolderTag>();
+
+	DataStorageUi.RegisterWidgetFactory<FOutlinerFavoriteWidgetConstructor>(
+		DataStorageUi.FindPurpose(IUiProvider::FPurposeInfo("General", "Cell", NAME_None).GeneratePurposeID()), SupportedColumns);
+	DataStorageUi.RegisterWidgetFactory<FOutlinerFavoriteHeaderConstructor>(
+		DataStorageUi.FindPurpose(IUiProvider::FPurposeInfo("General", "Header", NAME_None).GeneratePurposeID()), SupportedColumns);
+}
+
+
+//
+// Outliner Favorite Header Widget
+//
+
+FOutlinerFavoriteHeaderConstructor::FOutlinerFavoriteHeaderConstructor()
+	: FSimpleWidgetConstructor(StaticStruct())
+{
+}
+
+TSharedPtr<SWidget> FOutlinerFavoriteHeaderConstructor::CreateWidget(UE::Editor::DataStorage::ICoreProvider* DataStorage,
+	UE::Editor::DataStorage::IUiProvider* DataStorageUi,
+	UE::Editor::DataStorage::RowHandle TargetRow, UE::Editor::DataStorage::RowHandle WidgetRow, const UE::Editor::DataStorage::FMetaDataView& Arguments)
+{
+	using namespace UE::Editor::DataStorage;
+	DataStorage->AddColumn(WidgetRow, FHeaderWidgetSizeColumn
+		{
+			.ColumnSizeMode = EColumnSizeMode::Fixed,
+			.Width = 24.0f
+		});
+
+	return Private::MakeFavoritesWidgetContainer(
+		SNew(SImage)
+		.ColorAndOpacity(FAppStyle::Get().GetSlateColor("Colors.Foreground"))
+		.Image(FAppStyle::Get().GetBrush("Icons.Star"))
+		.ToolTipText(LOCTEXT("SceneOutlinerFavoriteHeaderTooltip", "Favorite"))
+	);
+}
+
+
+//
+// Outliner Favorite Widget
+//
+
+FOutlinerFavoriteWidgetConstructor::FOutlinerFavoriteWidgetConstructor()
+	: FSimpleWidgetConstructor(StaticStruct())
+{
+}
+
+TSharedPtr<SWidget> FOutlinerFavoriteWidgetConstructor::CreateWidget(
+	UE::Editor::DataStorage::ICoreProvider* DataStorage,
+	UE::Editor::DataStorage::IUiProvider* DataStorageUi,
+	UE::Editor::DataStorage::RowHandle TargetRow,
+	UE::Editor::DataStorage::RowHandle WidgetRow,
+	const UE::Editor::DataStorage::FMetaDataView& Arguments)
+{
+
+	using namespace UE::Editor::DataStorage;
+	if (!DataStorage->IsRowAvailable(TargetRow))
+	{
+		return SNew(STextBlock)
+			.Text(LOCTEXT("MissingRowReferenceColumn", "Unable to retrieve row reference."));
+	}
+
+	return Private::MakeFavoritesWidgetContainer(
+		SNew(Private::STedsFavoriteWidget, DataStorage, TargetRow)
+	);
+}
+
+TArray<TSharedPtr<const UE::Editor::DataStorage::FColumnSorterInterface>> FOutlinerFavoriteWidgetConstructor::ConstructColumnSorters(
+	UE::Editor::DataStorage::ICoreProvider* DataStorage, UE::Editor::DataStorage::IUiProvider* DataStorageUi,
+	const UE::Editor::DataStorage::FMetaDataView& Arguments)
+{
+	using namespace UE::Editor::DataStorage;
+	return TArray<TSharedPtr<const FColumnSorterInterface>>(
+		{
+			MakeShared<Private::FFavoriteWidgetSorter>()
+		});
+}
+
+#undef LOCTEXT_NAMESPACE
+
+
diff --git a/Engine/Plugins/Experimental/EditorDataStorageFeatures/Source/TedsOutliner/Private/Widgets/OutlinerFavoriteWidget.h b/Engine/Plugins/Experimental/EditorDataStorageFeatures/Source/TedsOutliner/Private/Widgets/OutlinerFavoriteWidget.h
new file mode 100644
--- /dev/null
+++ b/Engine/Plugins/Experimental/EditorDataStorageFeatures/Source/TedsOutliner/Private/Widgets/OutlinerFavoriteWidget.h
@@ -0,0 +1,61 @@
+// Copyright Epic Games, Inc. All Rights Reserved.
+
+#pragma once
+
+#include "Elements/Columns/TypedElementMiscColumns.h"
+#include "Elements/Interfaces/TypedElementDataStorageFactory.h"
+#include "Elements/Interfaces/TypedElementDataStorageInterface.h"
+#include "Elements/Interfaces/TypedElementDataStorageUiInterface.h"
+#include "Input/Reply.h"
+#include "Widgets/Images/SImage.h"
+
+#include "OutlinerFavoriteWidget.generated.h"
+
+UCLASS()
+class UOutlinerFavoriteWidgetFactory : public UEditorDataStorageFactory
+{
+	GENERATED_BODY()
+
+public:
+	~UOutlinerFavoriteWidgetFactory() override = default;
+
+	void RegisterWidgetConstructors(UE::Editor::DataStorage::ICoreProvider& DataStorage,
+		UE::Editor::DataStorage::IUiProvider& DataStorageUi) const override;
+};
+
+USTRUCT()
+struct FOutlinerFavoriteWidgetConstructor : public FSimpleWidgetConstructor
+{
+	GENERATED_BODY()
+
+public:
+	TEDSOUTLINER_API FOutlinerFavoriteWidgetConstructor();
+	~FOutlinerFavoriteWidgetConstructor() override = default;
+
+	TEDSOUTLINER_API virtual TSharedPtr<SWidget> CreateWidget(
+		UE::Editor::DataStorage::ICoreProvider* DataStorage,
+		UE::Editor::DataStorage::IUiProvider* DataStorageUi,
+		UE::Editor::DataStorage::RowHandle TargetRow,
+		UE::Editor::DataStorage::RowHandle WidgetRow,
+		const UE::Editor::DataStorage::FMetaDataView& Arguments) override;
+	TEDSOUTLINER_API virtual TArray<TSharedPtr<const UE::Editor::DataStorage::FColumnSorterInterface>> ConstructColumnSorters(
+		UE::Editor::DataStorage::ICoreProvider* DataStorage, UE::Editor::DataStorage::IUiProvider* DataStorageUi,
+		const UE::Editor::DataStorage::FMetaDataView& Arguments) override;
+};
+
+USTRUCT()
+struct FOutlinerFavoriteHeaderConstructor : public FSimpleWidgetConstructor
+{
+	GENERATED_BODY()
+
+public:
+	TEDSOUTLINER_API FOutlinerFavoriteHeaderConstructor();
+	~FOutlinerFavoriteHeaderConstructor() override = default;
+
+	TEDSOUTLINER_API virtual TSharedPtr<SWidget> CreateWidget(
+		UE::Editor::DataStorage::ICoreProvider* DataStorage,
+		UE::Editor::DataStorage::IUiProvider* DataStorageUi,
+		UE::Editor::DataStorage::RowHandle TargetRow,
+		UE::Editor::DataStorage::RowHandle WidgetRow,
+		const UE::Editor::DataStorage::FMetaDataView& Arguments) override;
+};
\ No newline at end of file
