[TEDS Viewers] Have TableViewerRowUIRows purge persisted inactive rows after a certain limit is reached (15000) to save memory

#Synced-CL 51011294

diff --git a/Engine/Plugins/Experimental/EditorDataStorageFeatures/Source/TedsTableViewer/Private/TedsTableViewerUtils.cpp b/Engine/Plugins/Experimental/EditorDataStorageFeatures/Source/TedsTableViewer/Private/TedsTableViewerUtils.cpp
--- a/Engine/Plugins/Experimental/EditorDataStorageFeatures/Source/TedsTableViewer/Private/TedsTableViewerUtils.cpp
+++ b/Engine/Plugins/Experimental/EditorDataStorageFeatures/Source/TedsTableViewer/Private/TedsTableViewerUtils.cpp
@@ -4,6 +4,7 @@
 
 #include "GameFramework/Actor.h"
 #include "DataStorage/Debug/Log.h"
+#include "DataStorage/Features.h"
 #include "Elements/Columns/TypedElementCompatibilityColumns.h"
 #include "Elements/Columns/TypedElementIconOverrideColumns.h"
 #include "Elements/Columns/TypedElementMiscColumns.h"
@@ -20,6 +21,7 @@ namespace UE::Editor::DataStorage::TableViewerUtils
 	static FName TableViewerWidgetTableName("Editor_TableViewerWidgetTable");
 	static FName TableViewerWidgetHierarchyName("TableViewerWidgetHierarchy");
 	static FName TableViewerRowWidgetTableName("Editor_TableViewerRowWidgetTable");
+	static uint32 MaxTableViewerRowWidgetRows = 15000;
 
 	FName GetWidgetTableName()
 	{
@@ -225,6 +227,44 @@ namespace UE::Editor::DataStorage::TableViewerUtils
 		}
 		return ParentRow;
 	}
+
+	// Helper to get all UI Rows for a given TableViewer
+	void GetAllTableViewerUIRows(const ICoreProvider* DataStorage, const RowHandle TableViewerRowHandle, TArray<RowHandle>& OutRows)
+	{
+		if (DataStorage)
+		{
+			DataStorage->WalkDepthFirst(DataStorage->FindHierarchyByName(GetWidgetHierarchyName()), TableViewerRowHandle,
+			[&OutRows](const ICoreProvider& Context, RowHandle, const RowHandle Target)
+			{
+				if (Context.HasColumns<FTedsTableViewerRowTag>(Target))
+				{
+					OutRows.Add(Target);
+				}
+			});
+		}
+	}
+
+	void RemoveAllTableViewerUIRows(ICoreProvider* DataStorage, const RowHandle TableViewerRowHandle)
+	{
+		// Remove any rows attached to this table and the dynamic table
+		if (DataStorage)
+		{
+			TArray<RowHandle> ChildUIRows;
+			GetAllTableViewerUIRows(DataStorage, TableViewerRowHandle, ChildUIRows);
+			DataStorage->BatchRemoveRows(ChildUIRows);
+		}
+	}
+
+	void DeactivateAllTableViewerUIRows(ICoreProvider* DataStorage, const RowHandle TableViewerRowHandle)
+	{
+		// Deactivate any rows attached to this table and the dynamic table
+		if (DataStorage)
+		{
+			TArray<RowHandle> ChildUIRows;
+			GetAllTableViewerUIRows(DataStorage, TableViewerRowHandle, ChildUIRows);
+			DataStorage->BatchAddRemoveColumns(ChildUIRows, { FTedsTableViewerInactiveRowTag::StaticStruct() }, {});
+		}
+	}
 	
 	void InitializeAllTableViewerUIRows(ICoreProvider* DataStorage, const TArray<TableViewerItemPtr>& Items, 
 		const FGetTableViewerRowMapping& GetTableViewerRowMapping, const FName& TableViewerIdentifier, 
@@ -251,7 +291,11 @@ namespace UE::Editor::DataStorage::TableViewerUtils
 					}
 					// If the row already existed we will use that, but if not, we want to create a row for each entry so
 					// changing the default expansion state doesn't affect rows that come into view but aren't new
-					if (const RowHandle ExistingRow = DataStorage->LookupMappedRow(TableViewerUIRowMappingDomain, MappingKey); !DataStorage->IsRowAvailable(ExistingRow))
+					if (const RowHandle ExistingRow = DataStorage->LookupMappedRow(TableViewerUIRowMappingDomain, MappingKey); DataStorage->IsRowAvailable(ExistingRow))
+					{
+						DataStorage->RemoveColumns<FTedsTableViewerInactiveRowTag>(ExistingRow);
+					}
+					else
 					{
 						const RowHandle AddedRow = DataStorage->AddRow(RowWidgetTableHandle);
 			
@@ -268,6 +312,50 @@ namespace UE::Editor::DataStorage::TableViewerUtils
 	}
 }
 
+void UTedsTableViewerFactory::RegisterQueries(UE::Editor::DataStorage::ICoreProvider& DataStorage)
+{
+	using namespace UE::Editor::DataStorage;
+	using namespace UE::Editor::DataStorage::Queries;
+
+	struct FRemoveAllInactiveTableViewerRowWidgetRows
+	{
+		void operator()()
+		{
+			if (Storage)
+			{
+				Storage->RemoveAllRowsWith<FTableViewerRowTag, FTedsTableViewerInactiveRowTag>();
+			}
+		}
+		ICoreProvider* Storage = nullptr;
+	};
+
+	const QueryHandle TableViewerRowWidgetCount = DataStorage.RegisterQuery(
+		Count()
+		.Where()
+			.All<FTableViewerRowTag>()
+		.Compile()
+	);
+	
+	DataStorage.RegisterQuery(
+		Select(
+			TEXT("Check if the max TableViewer Rows has been reached and command to purge all inactive rows"),
+			FObserver::OnAdd<FTableViewerTag>(),
+			[](IQueryContext& Context, RowHandle Row)
+			{
+				const FQueryResult QueryResult = Context.RunSubquery(0);
+				if (QueryResult.Count > TableViewerUtils::MaxTableViewerRowWidgetRows)
+				{
+					Context.PushCommand(FRemoveAllInactiveTableViewerRowWidgetRows
+					{
+						.Storage = GetMutableDataStorageFeature<ICoreProvider>(StorageFeatureName)
+					});
+				}
+			})
+		.DependsOn()
+			.SubQuery(TableViewerRowWidgetCount)
+		.Compile());
+}
+
 void UTedsTableViewerFactory::RegisterTables(UE::Editor::DataStorage::ICoreProvider& DataStorage)
 {
 	using namespace UE::Editor::DataStorage;
diff --git a/Engine/Plugins/Experimental/EditorDataStorageFeatures/Source/TedsTableViewer/Private/Widgets/STedsHierarchyViewer.cpp b/Engine/Plugins/Experimental/EditorDataStorageFeatures/Source/TedsTableViewer/Private/Widgets/STedsHierarchyViewer.cpp
--- a/Engine/Plugins/Experimental/EditorDataStorageFeatures/Source/TedsTableViewer/Private/Widgets/STedsHierarchyViewer.cpp
+++ b/Engine/Plugins/Experimental/EditorDataStorageFeatures/Source/TedsTableViewer/Private/Widgets/STedsHierarchyViewer.cpp
@@ -27,9 +27,16 @@ namespace UE::Editor::DataStorage
 
 	SHierarchyViewer::~SHierarchyViewer()
 	{
-		if (!bPersistUIRowsOnClose)
+		if(TedsWidget.IsValid() && Model.IsValid())
 		{
-			RemoveAllTableViewerUIRows();
+			if (bPersistUIRowsOnClose)
+			{
+				TableViewerUtils::DeactivateAllTableViewerUIRows(Model->GetDataStorageInterface(), TedsWidget->GetRowHandle());
+			}
+			else
+			{
+				TableViewerUtils::RemoveAllTableViewerUIRows(Model->GetDataStorageInterface(), TedsWidget->GetRowHandle());
+			}
 		}
 	}
 	
@@ -371,10 +378,10 @@ namespace UE::Editor::DataStorage
 		// Only use this function when it is necessary to get a name after initialization, the rows associated with the
 		// previous name will be removed regardless of the bRemoveTableViewerRowsOnClose argument and re-initialized.
 		// If the model hasn't been initialized yet, there are no rows to remove so we can skip this
-		if (bModelInitialized)
+		if (bModelInitialized && TedsWidget.IsValid() && Model.IsValid())
 		{
-			// Remove while set to the old identifier
-			RemoveAllTableViewerUIRows();
+			// Remove while set to the old identifier rather than deactivating rows
+			TableViewerUtils::RemoveAllTableViewerUIRows(Model->GetDataStorageInterface(), GetWidgetRowHandle());
 		}
  
 		TableViewerIdentifier = NewIdentifier.IsNone() ?
@@ -519,6 +526,7 @@ namespace UE::Editor::DataStorage
 				{
 					TreeView->SetItemExpansion(InItem, true);
 				}
+				Storage->RemoveColumns<FTedsTableViewerInactiveRowTag>(ExistingRow);
 				return ExistingRow;
 			}
 			else
@@ -557,28 +565,6 @@ namespace UE::Editor::DataStorage
 			true);
 	}
 
-	void SHierarchyViewer::RemoveAllTableViewerUIRows() const
-	{
-		if(TedsWidget.IsValid() && Model.IsValid())
-		{
-			// Remove any rows attached to this table and the dynamic table
-			if (ICoreProvider* Storage = Model->GetDataStorageInterface())
-			{
-				TArray<RowHandle> ChildUIRows;
-				
-				Storage->WalkDepthFirst(Storage->FindHierarchyByName(TableViewerUtils::GetWidgetHierarchyName()), TedsWidget->GetRowHandle(),
-				[&ChildUIRows](const ICoreProvider& Context, RowHandle, const RowHandle Target)
-				{
-					if (Context.HasColumns<FTedsTableViewerRowTag>(Target))
-					{
-						ChildUIRows.Add(Target);
-					}
-				});
-				Storage->BatchRemoveRows(ChildUIRows);
-			}
-		}
-	}
-
 	TSharedRef<ITableRow> SHierarchyViewer::MakeTableRowWidget(TableViewerItemPtr InItem, const TSharedRef<STableViewBase>& OwnerTable) const
 	{
 		const RowHandle TableViewerRowUIRow = AddOrFindTableViewerUIRow(InItem);
diff --git a/Engine/Plugins/Experimental/EditorDataStorageFeatures/Source/TedsTableViewer/Private/Widgets/STedsTableViewer.cpp b/Engine/Plugins/Experimental/EditorDataStorageFeatures/Source/TedsTableViewer/Private/Widgets/STedsTableViewer.cpp
--- a/Engine/Plugins/Experimental/EditorDataStorageFeatures/Source/TedsTableViewer/Private/Widgets/STedsTableViewer.cpp
+++ b/Engine/Plugins/Experimental/EditorDataStorageFeatures/Source/TedsTableViewer/Private/Widgets/STedsTableViewer.cpp
@@ -25,9 +25,16 @@ namespace UE::Editor::DataStorage
 
 	STedsTableViewer::~STedsTableViewer()
 	{
-		if (!bPersistUIRowsOnClose)
+		if(TedsWidget.IsValid() && Model.IsValid())
 		{
-			RemoveAllTableViewerUIRows();
+			if (bPersistUIRowsOnClose)
+			{
+				TableViewerUtils::DeactivateAllTableViewerUIRows(Model->GetDataStorageInterface(), TedsWidget->GetRowHandle());
+			}
+			else
+			{
+				TableViewerUtils::RemoveAllTableViewerUIRows(Model->GetDataStorageInterface(), TedsWidget->GetRowHandle());
+			}
 		}
 	}
 	
@@ -256,10 +263,10 @@ namespace UE::Editor::DataStorage
 		// Only use this function when it is necessary to get a name after initialization, the rows associated with the
 		// previous name will be removed regardless of the bRemoveTableViewerRowsOnClose argument and re-initialized.
 		// If the model hasn't been initialized yet, there are no rows to remove so we can skip this
-		if (bModelInitialized)
+		if (bModelInitialized && TedsWidget.IsValid() && Model.IsValid())
 		{
-			// Remove while set to the old identifier
-			RemoveAllTableViewerUIRows();
+			// Remove while set to the old identifier rather than deactivating rows
+			TableViewerUtils::RemoveAllTableViewerUIRows(Model->GetDataStorageInterface(), GetWidgetRowHandle());
 		}
  
 		TableViewerIdentifier = NewIdentifier.IsNone() ?
@@ -367,6 +374,7 @@ namespace UE::Editor::DataStorage
 			// If the row already existed but was removed from TreeView's SparseItemInfos (Searching, Filters, etc.), re-expand it
 			if (const RowHandle ExistingRow = Storage->LookupMappedRow(TableViewerUIRowMappingDomain, MappingKey); Storage->IsRowAvailable(ExistingRow))
 			{
+				Storage->RemoveColumns<FTedsTableViewerInactiveRowTag>(ExistingRow);
 				return ExistingRow;
 			}
 			else
@@ -393,28 +401,6 @@ namespace UE::Editor::DataStorage
 			GetWidgetRowHandle());
 	}
 
-	void STedsTableViewer::RemoveAllTableViewerUIRows() const
-	{
-		if(TedsWidget.IsValid() && Model.IsValid())
-		{
-			// Remove any rows attached to this table and the dynamic table
-			if (ICoreProvider* Storage = Model->GetDataStorageInterface())
-			{
-				TArray<RowHandle> ChildUIRows;
-				
-				Storage->WalkDepthFirst(Storage->FindHierarchyByName(TableViewerUtils::GetWidgetHierarchyName()), TedsWidget->GetRowHandle(),
-				[&ChildUIRows](const ICoreProvider& Context, RowHandle, const RowHandle Target)
-				{
-					if (Context.HasColumns<FTedsTableViewerRowTag>(Target))
-					{
-						ChildUIRows.Add(Target);
-					}
-				});
-				Storage->BatchRemoveRows(ChildUIRows);
-			}
-		}
-	}
-
 	TSharedRef<ITableRow> STedsTableViewer::MakeTableRowWidget(TableViewerItemPtr InItem, const TSharedRef<STableViewBase>& OwnerTable) const
 	{
 		const RowHandle TableViewerRowUIRow = AddOrFindTableViewerUIRow(InItem);
diff --git a/Engine/Plugins/Experimental/EditorDataStorageFeatures/Source/TedsTableViewer/Public/TedsTableViewerUtils.h b/Engine/Plugins/Experimental/EditorDataStorageFeatures/Source/TedsTableViewer/Public/TedsTableViewerUtils.h
--- a/Engine/Plugins/Experimental/EditorDataStorageFeatures/Source/TedsTableViewer/Public/TedsTableViewerUtils.h
+++ b/Engine/Plugins/Experimental/EditorDataStorageFeatures/Source/TedsTableViewer/Public/TedsTableViewerUtils.h
@@ -59,6 +59,12 @@ namespace UE::Editor::DataStorage::TableViewerUtils
 	// Walks up the UI Widget Hierarchy until finding the row with the FTableViewerTag
 	TEDSTABLEVIEWER_API RowHandle GetTableViewerUiRow(const ICoreProvider* DataStorage, const RowHandle ChildWidgetRow);
 
+	// Finds all the UI Rows for the provided TableViewer and removes the rows
+	TEDSTABLEVIEWER_API void RemoveAllTableViewerUIRows(ICoreProvider* DataStorage, const RowHandle TableViewerRowHandle);
+
+	// Finds all the UI Rows for the provided TableViewer and deactivates them (adds the FTedsTableViewerInactiveRowTag)
+	TEDSTABLEVIEWER_API void DeactivateAllTableViewerUIRows(ICoreProvider* DataStorage, const RowHandle TableViewerRowHandle);
+
 	// Initializes all rows on a TableViewer for a QueryStack
 	TEDSTABLEVIEWER_API void InitializeAllTableViewerUIRows(ICoreProvider* DataStorage, const TArray<TableViewerItemPtr>& Items,
 		const FGetTableViewerRowMapping& GetTableViewerRowMapping, const FName& TableViewerIdentifier, 
@@ -73,6 +79,7 @@ class UTedsTableViewerFactory : public UEditorDataStorageFactory
 public:
 	~UTedsTableViewerFactory() override = default;
 
+	void RegisterQueries(UE::Editor::DataStorage::ICoreProvider& DataStorage) override;
 	void RegisterTables(UE::Editor::DataStorage::ICoreProvider& DataStorage) override;
 	void RegisterHierarchies(UE::Editor::DataStorage::ICoreProvider& DataStorage) override;
 };
diff --git a/Engine/Plugins/Experimental/EditorDataStorageFeatures/Source/TedsTableViewer/Public/TedsTableViewerWidgetColumns.h b/Engine/Plugins/Experimental/EditorDataStorageFeatures/Source/TedsTableViewer/Public/TedsTableViewerWidgetColumns.h
--- a/Engine/Plugins/Experimental/EditorDataStorageFeatures/Source/TedsTableViewer/Public/TedsTableViewerWidgetColumns.h
+++ b/Engine/Plugins/Experimental/EditorDataStorageFeatures/Source/TedsTableViewer/Public/TedsTableViewerWidgetColumns.h
@@ -122,6 +122,16 @@ struct FTedsTableViewerRowTag : public FEditorDataStorageTag
 	GENERATED_BODY()
 };
 
+
+/**
+ * Tag on to indicate that this UI row is currently kept around but is not currently being used as a widget
+ */
+USTRUCT(meta = (DisplayName = "Inactive UI"))
+struct FTedsTableViewerInactiveRowTag : public FEditorDataStorageTag
+{
+	GENERATED_BODY()
+};
+
 namespace UE::Editor::DataStorage
 {
 	using FTableViewerTag = FTedsTableViewerTag;
diff --git a/Engine/Plugins/Experimental/EditorDataStorageFeatures/Source/TedsTableViewer/Public/Widgets/STedsHierarchyViewer.h b/Engine/Plugins/Experimental/EditorDataStorageFeatures/Source/TedsTableViewer/Public/Widgets/STedsHierarchyViewer.h
--- a/Engine/Plugins/Experimental/EditorDataStorageFeatures/Source/TedsTableViewer/Public/Widgets/STedsHierarchyViewer.h
+++ b/Engine/Plugins/Experimental/EditorDataStorageFeatures/Source/TedsTableViewer/Public/Widgets/STedsHierarchyViewer.h
@@ -183,8 +183,6 @@ namespace UE::Editor::DataStorage
 		RowHandle AddOrFindTableViewerUIRow(const TableViewerItemPtr InItem) const;
 
 		void InitializeAllTableViewerUIRows() const;
-
-		void RemoveAllTableViewerUIRows() const;
 		
 		TSharedRef<ITableRow> MakeTableRowWidget(TableViewerItemPtr InItem, const TSharedRef<STableViewBase>& OwnerTable) const;
 
diff --git a/Engine/Plugins/Experimental/EditorDataStorageFeatures/Source/TedsTableViewer/Public/Widgets/STedsTableViewer.h b/Engine/Plugins/Experimental/EditorDataStorageFeatures/Source/TedsTableViewer/Public/Widgets/STedsTableViewer.h
--- a/Engine/Plugins/Experimental/EditorDataStorageFeatures/Source/TedsTableViewer/Public/Widgets/STedsTableViewer.h
+++ b/Engine/Plugins/Experimental/EditorDataStorageFeatures/Source/TedsTableViewer/Public/Widgets/STedsTableViewer.h
@@ -162,8 +162,6 @@ namespace UE::Editor::DataStorage
 		RowHandle AddOrFindTableViewerUIRow(const TableViewerItemPtr InItem) const;
 
 		void InitializeAllTableViewerUIRows() const;
-
-		void RemoveAllTableViewerUIRows() const;
 		
 		TSharedRef<ITableRow> MakeTableRowWidget(TableViewerItemPtr InItem, const TSharedRef<STableViewBase>& OwnerTable) const;
 
