[Unified Favorites] Outliner 2.0 Favorites Context Menu Button

[Added] Options section with the toggle for the favorite state

#Synced-CL 50983539

diff --git a/Engine/Plugins/Experimental/EditorDataStorageFeatures/Source/TedsOutliner/Private/TedsOutlinerItem.cpp b/Engine/Plugins/Experimental/EditorDataStorageFeatures/Source/TedsOutliner/Private/TedsOutlinerItem.cpp
--- a/Engine/Plugins/Experimental/EditorDataStorageFeatures/Source/TedsOutliner/Private/TedsOutlinerItem.cpp
+++ b/Engine/Plugins/Experimental/EditorDataStorageFeatures/Source/TedsOutliner/Private/TedsOutlinerItem.cpp
@@ -87,6 +87,55 @@ void FTedsOutlinerTreeItem::GenerateContextMenu(UToolMenu* Menu, SSceneOutliner&
 	using namespace UE::Editor::DataStorage;
 	ICoreProvider* DataStorageInterface = GetMutableDataStorageFeature<ICoreProvider>(StorageFeatureName);
 
+	FToolMenuSection& Section = Menu->FindOrAddSection("Options", LOCTEXT("Options Section", "Options"));
+	// Favorites
+	{
+		if (DataStorageInterface->HasColumns<FFavoriteTag>(RowHandle))
+		{
+			Section.AddMenuEntry
+			(
+				"RemoveFromFavorites",
+				LOCTEXT("RemoveFromFavorites_Title", "Remove From Favorites"),
+				LOCTEXT("RemoveFromFavorites_ToolTip", "Removes the item from Favorites"),
+				FSlateIcon(FAppStyle::GetAppStyleSetName(), "PropertyWindow.Favorites_Disabled"),
+				FUIAction
+				(
+					FExecuteAction::CreateLambda([this, DataStorageInterface]()
+						{
+							if (!DataStorageInterface)
+							{
+								return;
+							}
+							DataStorageInterface->AddColumn<FTypedElementSyncFromWorldTag>(RowHandle);
+							DataStorageInterface->RemoveColumn<FFavoriteTag>(RowHandle);
+						})
+				)
+			);
+		}
+		else
+		{
+			Section.AddMenuEntry
+			(
+				"AddToFavorites",
+				LOCTEXT("AddToFavorites_Title", "Add To Favorites"),
+				LOCTEXT("AddToFavorites_ToolTip", "Adds the item to Favorites"),
+				FSlateIcon(FAppStyle::GetAppStyleSetName(), "Icons.Star"),
+				FUIAction
+				(
+					FExecuteAction::CreateLambda([this, DataStorageInterface]()
+						{
+							if (!DataStorageInterface)
+							{
+								return;
+							}
+							DataStorageInterface->AddColumn<FTypedElementSyncFromWorldTag>(RowHandle);
+							DataStorageInterface->AddColumn<FFavoriteTag>(RowHandle);
+						})
+				)
+			);
+		}
+	}
+
 	if (FTedsOutlinerContextMenuColumn* Column = DataStorageInterface->GetColumn<FTedsOutlinerContextMenuColumn>(RowHandle))
 	{
 		if (Column->OnCreateContextMenu.IsBound())
