[Unified Favorites] Favorite and Not Favorite filters for Outliner 2.0

[Added] 2 filters in the "Other" section of the Outliner 2.0 Filters

#Synced-CL 50983539

diff --git a/Engine/Plugins/Experimental/EditorDataStorageFeatures/Source/TedsOutliner/Private/TedsOutlinerModule.cpp b/Engine/Plugins/Experimental/EditorDataStorageFeatures/Source/TedsOutliner/Private/TedsOutlinerModule.cpp
--- a/Engine/Plugins/Experimental/EditorDataStorageFeatures/Source/TedsOutliner/Private/TedsOutlinerModule.cpp
+++ b/Engine/Plugins/Experimental/EditorDataStorageFeatures/Source/TedsOutliner/Private/TedsOutlinerModule.cpp
@@ -202,6 +202,7 @@ TSharedRef<SWidget> FTedsOutlinerModule::CreateLevelEditorTedsOutliner()
 	 *		return MobilityColumn.Mobility == EComponentMobility::Static;
 	 *	});
 	 */
+	AddFavoriteFilters(Params);
 	
 	const FLevelEditorModule* const LevelEditorModule = FModuleManager::GetModulePtr<FLevelEditorModule>(TEXT("LevelEditor"));
 	if (LevelEditorModule)
@@ -217,6 +218,25 @@ TSharedRef<SWidget> FTedsOutlinerModule::CreateLevelEditorTedsOutliner()
 	return TEDSOutlinerShared;
 }
 
+void FTedsOutlinerModule::AddFavoriteFilters(FTedsOutlinerParams& Params) const
+{
+	const TSharedRef<FFilterCategory> FavoriteFiltersCategory = MakeShared<FFilterCategory>(LOCTEXT("FavoriteFilters", "Favorites"), LOCTEXT("FavoriteFiltersTooltip", "Filter by Favorite state"));
+
+	// Create the filter
+	Params.Filters.Emplace(
+		"FilterByFavorites",
+		LOCTEXT("FavoriteFilterDisplayName", "Favorite"),
+		LOCTEXT("FavoriteFilterDisplayNameToolTip", "Filter by Favorites"),
+		FName("Icons.Star"),
+		FavoriteFiltersCategory,
+		Storage->RegisterQuery(
+			Select()
+			.Where()
+			.All<FFavoriteTag>()
+			.Compile())
+		);
+}
+
 TSharedRef<SDockTab> FTedsOutlinerModule::OpenLevelEditorTedsOutliner(const FSpawnTabArgs& SpawnTabArgs)
 {
 	return SNew(SDockTab)
diff --git a/Engine/Plugins/Experimental/EditorDataStorageFeatures/Source/TedsOutliner/Public/TedsOutlinerModule.h b/Engine/Plugins/Experimental/EditorDataStorageFeatures/Source/TedsOutliner/Public/TedsOutlinerModule.h
--- a/Engine/Plugins/Experimental/EditorDataStorageFeatures/Source/TedsOutliner/Public/TedsOutlinerModule.h
+++ b/Engine/Plugins/Experimental/EditorDataStorageFeatures/Source/TedsOutliner/Public/TedsOutlinerModule.h
@@ -56,6 +56,7 @@ private:
 	void UnregisterLevelEditorTedsOutlinerTab();
 	TSharedRef<SDockTab> OpenLevelEditorTedsOutliner(const FSpawnTabArgs& SpawnTabArgs);
 	TSharedRef<SWidget> CreateLevelEditorTedsOutliner();
+	void AddFavoriteFilters(FTedsOutlinerParams& Params) const;
 	
 private:
 	FName TedsOutlinerTabName;
