Fixes for TedsViewer

#Synced-CL 51194092

diff --git a/Engine/Plugins/Experimental/EditorDataStorageFeatures/Source/TedsTableViewer/Private/TedsTableViewerModel.cpp b/Engine/Plugins/Experimental/EditorDataStorageFeatures/Source/TedsTableViewer/Private/TedsTableViewerModel.cpp
--- a/Engine/Plugins/Experimental/EditorDataStorageFeatures/Source/TedsTableViewer/Private/TedsTableViewerModel.cpp
+++ b/Engine/Plugins/Experimental/EditorDataStorageFeatures/Source/TedsTableViewer/Private/TedsTableViewerModel.cpp
@@ -17,12 +17,18 @@ namespace UE::Editor::DataStorage
 {
 	FTedsTableViewerModel::FTedsTableViewerModel(const FArguments& Args)
 		: RequestedTedsColumns(Args.RequestedColumns)
+		, GenericMetaData(Args.GenericMetaData)
 		, ColumnMetaData(Args.ColumnMetaData)
 		, CellWidgetPurpose(Args.CellWidgetPurpose)
 		, HeaderWidgetPurpose(Args.HeaderWidgetPurpose)
 		, IsItemVisible(Args.IsItemVisibleDelegate)
 		, PrimaryColumn(Args.PrimaryColumn)
 	{
+		if (!GenericMetaData.IsValid())
+		{
+			GenericMetaData = MakeShared<FMetaData>();
+		}
+
 		Storage = GetMutableDataStorageFeature<ICoreProvider>(StorageFeatureName);
 		StorageUi = GetMutableDataStorageFeature<IUiProvider>(UiFeatureName);
 
@@ -241,7 +247,7 @@ namespace UE::Editor::DataStorage
 				// Currently we are passing in metadata for all columns to individual columns, we might want to only pass in the metadata for the matched columns
 				// for each one
 				TSharedRef<FTedsTableViewerColumn> Column = MakeShared<FTedsTableViewerColumn>(NameId, CellConstructor, MoveTemp(MatchedColumnsCopy), HeaderConstructor,
-					FMetaData(), ColumnMetaData);
+					*GenericMetaData, ColumnMetaData);
 			
 				Column->SetIsRowVisibleDelegate(FTedsTableViewerColumn::FIsRowVisible::CreateRaw(this, &FTedsTableViewerModel::IsRowVisible));
 				Column->SetSortDelegates(FTedsTableViewerColumn::FSortHandler
diff --git a/Engine/Plugins/Experimental/EditorDataStorageFeatures/Source/TedsTableViewer/Private/Widgets/Composite/STedsCompositeHierarchyViewer.cpp b/Engine/Plugins/Experimental/EditorDataStorageFeatures/Source/TedsTableViewer/Private/Widgets/Composite/STedsCompositeHierarchyViewer.cpp
--- a/Engine/Plugins/Experimental/EditorDataStorageFeatures/Source/TedsTableViewer/Private/Widgets/Composite/STedsCompositeHierarchyViewer.cpp
+++ b/Engine/Plugins/Experimental/EditorDataStorageFeatures/Source/TedsTableViewer/Private/Widgets/Composite/STedsCompositeHierarchyViewer.cpp
@@ -200,6 +200,7 @@ namespace UE::Editor::DataStorage
 			.EmptyRowsMessage(HierarchyViewerArgs._EmptyRowsMessage)
 			.CellWidgetPurpose(HierarchyViewerArgs._CellWidgetPurpose)
 			.HeaderWidgetPurpose(HierarchyViewerArgs._HeaderWidgetPurpose)
+			.GenericMetaData(HierarchyViewerArgs._GenericMetaData)
 			.ColumnMetaData(HierarchyViewerArgs._ColumnMetaData)
 			.ListSelectionMode(HierarchyViewerArgs._ListSelectionMode)
 			.OnSelectionChanged(HierarchyViewerArgs._OnSelectionChanged)
diff --git a/Engine/Plugins/Experimental/EditorDataStorageFeatures/Source/TedsTableViewer/Private/Widgets/STedsHierarchyViewer.cpp b/Engine/Plugins/Experimental/EditorDataStorageFeatures/Source/TedsTableViewer/Private/Widgets/STedsHierarchyViewer.cpp
--- a/Engine/Plugins/Experimental/EditorDataStorageFeatures/Source/TedsTableViewer/Private/Widgets/STedsHierarchyViewer.cpp
+++ b/Engine/Plugins/Experimental/EditorDataStorageFeatures/Source/TedsTableViewer/Private/Widgets/STedsHierarchyViewer.cpp
@@ -75,6 +75,7 @@ namespace UE::Editor::DataStorage
 		Model = MakeShared<FTedsTableViewerModel>(FTedsTableViewerModel::FArguments{
 			.RowQueryStack = InArgs._AllNodeProvider,
 			.RequestedColumns = InArgs._Columns,
+			.GenericMetaData = InArgs._GenericMetaData,
 			.ColumnMetaData = InArgs._ColumnMetaData,
 			.CellWidgetPurpose = CellWidgetPurpose,
 			.HeaderWidgetPurpose = HeaderWidgetPurpose,
diff --git a/Engine/Plugins/Experimental/EditorDataStorageFeatures/Source/TedsTableViewer/Public/TedsTableViewerModel.h b/Engine/Plugins/Experimental/EditorDataStorageFeatures/Source/TedsTableViewer/Public/TedsTableViewerModel.h
--- a/Engine/Plugins/Experimental/EditorDataStorageFeatures/Source/TedsTableViewer/Public/TedsTableViewerModel.h
+++ b/Engine/Plugins/Experimental/EditorDataStorageFeatures/Source/TedsTableViewer/Public/TedsTableViewerModel.h
@@ -42,6 +42,7 @@ namespace UE::Editor::DataStorage
 		{
 			TSharedPtr<QueryStack::IRowNode> RowQueryStack;
 			TArray<TWeakObjectPtr<const UScriptStruct>> RequestedColumns;
+			TSharedPtr<FMetaData> GenericMetaData;
 			TArray<TSharedPtr<FColumnMetaData>> ColumnMetaData;
 			IUiProvider::FPurposeID CellWidgetPurpose; 
 			IUiProvider::FPurposeID HeaderWidgetPurpose;
@@ -147,6 +148,9 @@ namespace UE::Editor::DataStorage
 		// The initial TEDS columns the widget was requested to display
 		TArray<TWeakObjectPtr<const UScriptStruct>> RequestedTedsColumns;
 
+		// Generic Metadata for all widgets
+		TSharedPtr<FMetaData> GenericMetaData;
+
 		// Metadata for RequestedTedsColumns
 		TArray<TSharedPtr<FColumnMetaData>> ColumnMetaData;
 
diff --git a/Engine/Plugins/Experimental/EditorDataStorageFeatures/Source/TedsTableViewer/Public/Widgets/STedsHierarchyViewer.h b/Engine/Plugins/Experimental/EditorDataStorageFeatures/Source/TedsTableViewer/Public/Widgets/STedsHierarchyViewer.h
--- a/Engine/Plugins/Experimental/EditorDataStorageFeatures/Source/TedsTableViewer/Public/Widgets/STedsHierarchyViewer.h
+++ b/Engine/Plugins/Experimental/EditorDataStorageFeatures/Source/TedsTableViewer/Public/Widgets/STedsHierarchyViewer.h
@@ -81,6 +81,9 @@ namespace UE::Editor::DataStorage
 		// The Columns that this table viewer will display
 		SLATE_ARGUMENT(TArray<TWeakObjectPtr<const UScriptStruct>>, Columns)
 
+		// Generic Metadata about all widgets table viewer is displaying
+		SLATE_ARGUMENT(TSharedPtr<FMetaData>, GenericMetaData)
+
 		// Metadata about the columns the table viewer is displaying
 		SLATE_ARGUMENT(TArray<TSharedPtr<FColumnMetaData>>, ColumnMetaData)
 
