[TEDS Package] Turn the Package sync back on

diff --git a/Engine/Plugins/Experimental/EditorDataStorageFeatures/Source/TedsRevisionControl/Private/Queries/ObjectPackagePathToColumnQueries.cpp b/Engine/Plugins/Experimental/EditorDataStorageFeatures/Source/TedsRevisionControl/Private/Queries/ObjectPackagePathToColumnQueries.cpp
--- a/Engine/Plugins/Experimental/EditorDataStorageFeatures/Source/TedsRevisionControl/Private/Queries/ObjectPackagePathToColumnQueries.cpp
+++ b/Engine/Plugins/Experimental/EditorDataStorageFeatures/Source/TedsRevisionControl/Private/Queries/ObjectPackagePathToColumnQueries.cpp
@@ -19,7 +19,7 @@ namespace UE::Editor::RevisionControl::Private
 {
 	using namespace UE::Editor::DataStorage;
 
-	static bool bAutoSyncPackageInformation = false;
+	static bool bAutoSyncPackageInformation = true;
     FAutoConsoleVariableRef CVarSyncPackageInformation(
     	TEXT("TEDS.RevisionControl.SyncPackageInformation"),
     	bAutoSyncPackageInformation,
