TEDS: Enable Actor parent queries, these were disabled in the past due to an invalid memory access issue in TEDS hierarchies that has since been fixed

#Synced-CL 51094648

diff --git a/Engine/Plugins/Experimental/EditorDataStorageFeatures/Source/TedsActorCompatibility/Private/Compatibility/TedsActorParentQueries.cpp b/Engine/Plugins/Experimental/EditorDataStorageFeatures/Source/TedsActorCompatibility/Private/Compatibility/TedsActorParentQueries.cpp
--- a/Engine/Plugins/Experimental/EditorDataStorageFeatures/Source/TedsActorCompatibility/Private/Compatibility/TedsActorParentQueries.cpp
+++ b/Engine/Plugins/Experimental/EditorDataStorageFeatures/Source/TedsActorCompatibility/Private/Compatibility/TedsActorParentQueries.cpp
@@ -15,7 +15,7 @@
 namespace UE::Editor::DataStorage::ActorParentQueries::Private
 {
 	static FName ActorEditorHierarchyName(TEXT("ActorEditorHierarchy"));
-	static bool bAddParentColumnToActors = false;
+	static bool bAddParentColumnToActors = true;
 
 	// Cvar to add parenting info for actors to TEDS
 	static FAutoConsoleVariableRef CvarAddParentColumnToActors(
