[TEDS Outliner] Add the Unsaved Column to the TEDS Outliner

#Synced-CL 51340897

diff --git a/Engine/Plugins/Experimental/EditorDataStorageFeatures/Source/TedsOutliner/Private/Compatibility/SceneOutlinerTedsBridge.cpp b/Engine/Plugins/Experimental/EditorDataStorageFeatures/Source/TedsOutliner/Private/Compatibility/SceneOutlinerTedsBridge.cpp
--- a/Engine/Plugins/Experimental/EditorDataStorageFeatures/Source/TedsOutliner/Private/Compatibility/SceneOutlinerTedsBridge.cpp
+++ b/Engine/Plugins/Experimental/EditorDataStorageFeatures/Source/TedsOutliner/Private/Compatibility/SceneOutlinerTedsBridge.cpp
@@ -36,6 +36,7 @@ void USceneOutlinerTedsBridgeFactory::PreRegister(UE::Editor::DataStorage::ICore
 	TEDSToOutlinerDefaultColumnMapping.Add(FTedsActorMobilityColumn::StaticStruct(), FSceneOutlinerBuiltInColumnTypes::Mobility());
 	TEDSToOutlinerDefaultColumnMapping.Add(FTypedElementLabelColumn::StaticStruct(), FSceneOutlinerBuiltInColumnTypes::Label());
 	TEDSToOutlinerDefaultColumnMapping.Add(FLevelColumn::StaticStruct(), FSceneOutlinerBuiltInColumnTypes::Level());
+	TEDSToOutlinerDefaultColumnMapping.Add(FTedsPrimaryPackageObjectTag::StaticStruct(), FSceneOutlinerBuiltInColumnTypes::Unsaved());
 }
 
 void USceneOutlinerTedsBridgeFactory::RegisterTables(UE::Editor::DataStorage::ICoreProvider& DataStorage)
diff --git a/Engine/Plugins/Experimental/EditorDataStorageFeatures/Source/TedsOutliner/Private/TedsOutlinerModule.cpp b/Engine/Plugins/Experimental/EditorDataStorageFeatures/Source/TedsOutliner/Private/TedsOutlinerModule.cpp
--- a/Engine/Plugins/Experimental/EditorDataStorageFeatures/Source/TedsOutliner/Private/TedsOutlinerModule.cpp
+++ b/Engine/Plugins/Experimental/EditorDataStorageFeatures/Source/TedsOutliner/Private/TedsOutlinerModule.cpp
@@ -138,7 +138,8 @@ FTedsOutlinerColumnDescription FTedsOutlinerModule::GetLevelEditorTedsOutlinerCo
 		FUObjectIdNameColumn::StaticStruct(),
 		FTedsActorMobilityColumn::StaticStruct(),
 		FLevelColumn::StaticStruct(),
-		FVisibleInEditorColumn::StaticStruct()});
+		FVisibleInEditorColumn::StaticStruct(),
+		FTedsPrimaryPackageObjectTag::StaticStruct()});
 
 	if (CVarUseFavoritesWidgets->GetBool())
 	{
diff --git a/Engine/Plugins/Experimental/EditorDataStorageFeatures/Source/TedsOutliner/Private/Widgets/OutlinerUnsavedWidget.cpp b/Engine/Plugins/Experimental/EditorDataStorageFeatures/Source/TedsOutliner/Private/Widgets/OutlinerUnsavedWidget.cpp
--- a/Engine/Plugins/Experimental/EditorDataStorageFeatures/Source/TedsOutliner/Private/Widgets/OutlinerUnsavedWidget.cpp
+++ b/Engine/Plugins/Experimental/EditorDataStorageFeatures/Source/TedsOutliner/Private/Widgets/OutlinerUnsavedWidget.cpp
@@ -53,6 +53,12 @@ TSharedPtr<SWidget> FOutlinerUnsavedHeaderConstructor::CreateWidget(UE::Editor::
 		.ToolTipText(LOCTEXT("SceneOutlinerUnsavedHeaderTooltip", "Unsaved"));
 }
 
+FText FOutlinerUnsavedHeaderConstructor::CreateWidgetDisplayNameText(UE::Editor::DataStorage::ICoreProvider* DataStorage,
+	UE::Editor::DataStorage::RowHandle Row) const
+{
+	return LOCTEXT("OutlinerUnsavedHeaderDisplayName", "Unsaved");
+}
+
 TSharedPtr<SWidget> FOutlinerUnsavedWidgetConstructor::CreateWidget(
 	UE::Editor::DataStorage::ICoreProvider* DataStorage,
 	UE::Editor::DataStorage::IUiProvider* DataStorageUi,
@@ -70,7 +76,7 @@ TSharedPtr<SWidget> FOutlinerUnsavedWidgetConstructor::CreateWidget(
 	FAttributeBinder Binder(TargetRow, DataStorage);
 
 	return SNew(SBox)
-			.HAlign(HAlign_Center)
+			.HAlign(HAlign_Left)
 			.VAlign(VAlign_Center)
 			[
 				SNew(SImage)
diff --git a/Engine/Plugins/Experimental/EditorDataStorageFeatures/Source/TedsOutliner/Public/Widgets/OutlinerUnsavedWidget.h b/Engine/Plugins/Experimental/EditorDataStorageFeatures/Source/TedsOutliner/Public/Widgets/OutlinerUnsavedWidget.h
--- a/Engine/Plugins/Experimental/EditorDataStorageFeatures/Source/TedsOutliner/Public/Widgets/OutlinerUnsavedWidget.h
+++ b/Engine/Plugins/Experimental/EditorDataStorageFeatures/Source/TedsOutliner/Public/Widgets/OutlinerUnsavedWidget.h
@@ -52,4 +52,7 @@ public:
 		UE::Editor::DataStorage::RowHandle TargetRow,
 		UE::Editor::DataStorage::RowHandle WidgetRow, 
 		const UE::Editor::DataStorage::FMetaDataView& Arguments) override;
+		
+	TEDSOUTLINER_API virtual FText CreateWidgetDisplayNameText(
+		UE::Editor::DataStorage::ICoreProvider* DataStorage, UE::Editor::DataStorage::RowHandle Row = UE::Editor::DataStorage::InvalidRowHandle) const override;
 };
