[TEDS Outliner] Refactor TedsOutlinerFilters and add a custom Scene Outliner World Filter

#Synced-CL 51342861

diff --git a/Engine/Plugins/Experimental/Dataflow/Source/DataflowEditor/Private/Dataflow/DataflowOutlinerView.cpp b/Engine/Plugins/Experimental/Dataflow/Source/DataflowEditor/Private/Dataflow/DataflowOutlinerView.cpp
--- a/Engine/Plugins/Experimental/Dataflow/Source/DataflowEditor/Private/Dataflow/DataflowOutlinerView.cpp
+++ b/Engine/Plugins/Experimental/Dataflow/Source/DataflowEditor/Private/Dataflow/DataflowOutlinerView.cpp
@@ -4,6 +4,7 @@
 #include "Dataflow/DataflowEditorPreviewSceneBase.h"
 #include "Dataflow/DataflowOutlinerMode.h"
 #include "TedsOutlinerImpl.h"
+#include "TedsOutlinerFilter.h"
 #include "Compatibility/SceneOutlinerTedsBridge.h"
 #include "Dataflow/DataflowContent.h"
 #include "DataStorage/Features.h"
@@ -55,30 +56,30 @@ TSharedPtr<ISceneOutliner> FDataflowOutlinerView::CreateWidget()
 	Params.ColumnDescription = ColumnDescription;
 	
 	// Add outliner filter queries
-	Params.Filters.Emplace("Dataflow Construction",
+	Params.Filters.Add(MakeShared<FTedsOutlinerFilter>("Dataflow Construction",
 		LOCTEXT("FilterDataflowConstructionDisplayName", "Dataflow Construction"),
 		Storage->RegisterQuery(
 			Select()
 			.Where(TColumn<FDataflowConstructionObjectTag>())
-			.Compile()));
-	Params.Filters.Emplace("Dataflow Simulation",
+			.Compile())));
+	Params.Filters.Add(MakeShared<FTedsOutlinerFilter>("Dataflow Simulation",
 		LOCTEXT("FilterDataflowSimulationDisplayName", "Dataflow Simulation"),
 		Storage->RegisterQuery(
 			Select()
 			.Where(TColumn<FDataflowSimulationObjectTag>())
-			.Compile()));
-	Params.Filters.Emplace("Dataflow Elements",
+			.Compile())));
+	Params.Filters.Add(MakeShared<FTedsOutlinerFilter>("Dataflow Elements",
 		LOCTEXT("FilterDataflowElementsDisplayName", "Dataflow Elements"),
 		Storage->RegisterQuery(
 			Select()
 			.Where(TColumn<FDataflowSceneStructTag>(GetEditorContent()->GetDataflowOwner().GetFName()))
-			.Compile()));
-	Params.Filters.Emplace("Dataflow Components",
+			.Compile())));
+	Params.Filters.Add(MakeShared<FTedsOutlinerFilter>("Dataflow Components",
 		LOCTEXT("FilterDataflowComponentsDisplayName", "Dataflow Components"),
 		Storage->RegisterQuery(
 			Select()
 			.Where(TColumn<FDataflowSceneObjectTag>(GetEditorContent()->GetDataflowOwner().GetFName()))
-			.Compile()));
+			.Compile())));
 
 	InitOptions.ModeFactory = FCreateSceneOutlinerMode::CreateLambda([this, &Params](SSceneOutliner* Outliner)
 	{
diff --git a/Engine/Plugins/Experimental/EditorDataStorageFeatures/Source/TedsOutliner/Private/TedsOutlinerFilter.cpp b/Engine/Plugins/Experimental/EditorDataStorageFeatures/Source/TedsOutliner/Private/TedsOutlinerFilter.cpp
--- a/Engine/Plugins/Experimental/EditorDataStorageFeatures/Source/TedsOutliner/Private/TedsOutlinerFilter.cpp
+++ b/Engine/Plugins/Experimental/EditorDataStorageFeatures/Source/TedsOutliner/Private/TedsOutlinerFilter.cpp
@@ -3,6 +3,7 @@
 #include "TedsOutlinerFilter.h"
 
 #include "DataStorage/Features.h"
+#include "Elements/Columns/TypedElementCompatibilityColumns.h"
 #include "Elements/Columns/TypedElementTypeInfoColumns.h"
 #include "Elements/Framework/TypedElementQueryContext.h"
 #include "Styling/SlateIconFinder.h"
@@ -13,98 +14,53 @@
 
 namespace UE::Editor::Outliner
 {
-FTedsFilterData::FTedsFilterData(const FName& InFilterName, const FText& InFilterDisplayName, const FText& InFilterToolTip, const FName& InFilterIconName,
-	const TSharedPtr<FFilterCategory>& InCategory, const QueryHandle& InFilterQuery)
-	: FilterName(InFilterName)
-	, FilterDisplayName(InFilterDisplayName)
-	, FilterToolTip(InFilterToolTip)
-	, FilterIconName(InFilterIconName)
-	, FilterCategory(InCategory)
-{
-	FilterQuery.Set<QueryHandle>(Helpers::CheckValidFilterQueryHandle(InFilterQuery) ? InFilterQuery : InvalidQueryHandle);
-}
-
-FTedsFilterData::FTedsFilterData(const FName& InFilterName, const FText& InFilterDisplayName, const QueryHandle& InFilterQuery)
-	: FTedsFilterData(
-		InFilterName, 
-		InFilterDisplayName, 
-		FText::FromString(InFilterName.ToString()), 
-		FName(), 
-		nullptr, 
-		InFilterQuery)
-{}
-
-FTedsFilterData::FTedsFilterData(const FName& InFilterName, const FText& InFilterDisplayName, const FText& InFilterToolTip, const FName& InFilterIconName,
-	const TSharedPtr<FFilterCategory>& InCategory, const TQueryFunction<bool>& InFilterQuery)
-	: FilterName(InFilterName)
-	, FilterDisplayName(InFilterDisplayName)
-	, FilterToolTip(InFilterToolTip)
-	, FilterIconName(InFilterIconName)
-	, FilterCategory(InCategory)
-{
-	FilterQuery.Set<TQueryFunction<bool>>(InFilterQuery);
-}
-
-FTedsFilterData::FTedsFilterData(const FName& InFilterName, const FText& InFilterDisplayName, const TQueryFunction<bool>& InFilterQuery)
-	: FTedsFilterData(
-		InFilterName, 
-		InFilterDisplayName, 
-		FText::FromString(InFilterName.ToString()), 
-		FName(), 
-		nullptr, 
-		InFilterQuery)
-{}
-
-FTedsFilterData::FTedsFilterData(const FTedsFilterData&) = default;
-FTedsFilterData::~FTedsFilterData() = default;
-	
-FTedsOutlinerFilter::FTedsOutlinerFilter(const FTedsFilterData& InTedsFilterData, const TSharedPtr<FTedsOutlinerImpl>& InTedsOutlinerImpl)
-	: FFilterBase(InTedsFilterData.FilterCategory)
-	, FilterName(InTedsFilterData.FilterName)
-	, FilterDisplayName(InTedsFilterData.FilterDisplayName)
-	, FilterToolTip(InTedsFilterData.FilterToolTip)
-	, FilterIconName(InTedsFilterData.FilterIconName)
-	, bIsClassFilter(false)
-	, TedsOutlinerImpl(InTedsOutlinerImpl)
-	, FilterQuery(InTedsFilterData.FilterQuery)
-{}
-
 FTedsOutlinerFilter::FTedsOutlinerFilter(const FName& InFilterName, const FText& InFilterDisplayName, const FText& InFilterToolTip, const FName& InFilterIconName, 
-	const TSharedPtr<FFilterCategory>& InCategory, const TSharedPtr<FTedsOutlinerImpl>& InTedsOutlinerImpl, 
-	const QueryHandle& InFilterQuery)
+	const TSharedPtr<FFilterCategory>& InCategory, const QueryHandle& InFilterQuery, const bool bInteractiveFilter)
 	: FFilterBase(InCategory)
 	, FilterName(InFilterName)
 	, FilterDisplayName(InFilterDisplayName)
 	, FilterToolTip(InFilterToolTip)
 	, FilterIconName(InFilterIconName)
 	, bIsClassFilter(false)
-	, TedsOutlinerImpl(InTedsOutlinerImpl)
+	, bInteractiveFilter(bInteractiveFilter)
 {
 	FilterQuery.Set<QueryHandle>(Helpers::CheckValidFilterQueryHandle(InFilterQuery) ? InFilterQuery : InvalidQueryHandle);
 }
 
+FTedsOutlinerFilter::FTedsOutlinerFilter(const FName& InFilterName, const FText& InFilterDisplayName, const QueryHandle& InFilterQuery,
+	const bool bInteractiveFilter)
+	: FTedsOutlinerFilter(InFilterName, InFilterDisplayName, FText::FromString(InFilterName.ToString()), FName(), 
+		nullptr, InFilterQuery, bInteractiveFilter)
+{}
+
 FTedsOutlinerFilter::FTedsOutlinerFilter(const FName& InFilterName, const FText& InFilterDisplayName, const FText& InFilterToolTip, 
-	const FName& InFilterIconName, const TSharedPtr<FFilterCategory>& InCategory, const TSharedPtr<FTedsOutlinerImpl>& InTedsOutlinerImpl,
-	const TQueryFunction<bool>& InFilterQuery)
+	const FName& InFilterIconName, const TSharedPtr<FFilterCategory>& InCategory, const TQueryFunction<bool>& InFilterQuery,
+	const bool bInteractiveFilter)
 	: FFilterBase(InCategory)
 	, FilterName(InFilterName)
 	, FilterDisplayName(InFilterDisplayName)
 	, FilterToolTip(InFilterToolTip)
 	, FilterIconName(InFilterIconName)
 	, bIsClassFilter(false)
-	, TedsOutlinerImpl(InTedsOutlinerImpl)
+	, bInteractiveFilter(bInteractiveFilter)
 {
 	FilterQuery.Set<TQueryFunction<bool>>(InFilterQuery);
 }
 
-FTedsOutlinerFilter::FTedsOutlinerFilter(const UClass* InClass, const TSharedPtr<FFilterCategory>& InCategory, const TSharedPtr<FTedsOutlinerImpl>& InTedsOutlinerImpl)
+FTedsOutlinerFilter::FTedsOutlinerFilter(const FName& InFilterName, const FText& InFilterDisplayName, const TQueryFunction<bool>& InFilterQuery,
+	const bool bInteractiveFilter)
+	: FTedsOutlinerFilter(InFilterName, InFilterDisplayName, FText::FromString(InFilterName.ToString()), FName(), 
+		nullptr, InFilterQuery, bInteractiveFilter)
+{}
+
+FTedsOutlinerFilter::FTedsOutlinerFilter(const UClass* InClass, const TSharedPtr<FFilterCategory>& InCategory, const bool bInteractiveFilter)
 	: FFilterBase(InCategory)
 	, FilterName(InClass->GetFName())
 	, FilterDisplayName(InClass->GetDisplayNameText())
 	, FilterToolTip(FText::Format(LOCTEXT("FilterClassTooltip", "Filter by {0}"), InClass->GetDisplayNameText()))
 	, FilterIconName(FSlateIconFinder::FindIconForClass(InClass).GetStyleName())
 	, bIsClassFilter(true)
-	, TedsOutlinerImpl(InTedsOutlinerImpl)
+	, bInteractiveFilter(bInteractiveFilter)
 {
 	FilterQuery.Set<TQueryFunction<bool>>(
 		Queries::BuildQueryFunction<bool>([InClass](TQueryContext<RowBatchInfo> Context, TResult<bool>& Result, TConstBatch<FTypedElementClassTypeInfoColumn> TypeInfoColumns)
@@ -115,6 +71,16 @@ FTedsOutlinerFilter::FTedsOutlinerFilter(const UClass* InClass, const TSharedPtr
 			}, TypeInfoColumns);
 		}));
 }
+
+void FTedsOutlinerFilter::SetSceneOutlinerImpl(const TSharedPtr<FTedsOutlinerImpl>& InTedsOutlinerImpl)
+{
+	TedsOutlinerImpl = InTedsOutlinerImpl;
+}
+
+void FTedsOutlinerFilter::SetCategory(const TSharedPtr<FFilterCategory>& InCategory)
+{
+	FilterCategory = InCategory;
+}
 	
 FString FTedsOutlinerFilter::GetName() const
 {
diff --git a/Engine/Plugins/Experimental/EditorDataStorageFeatures/Source/TedsOutliner/Private/TedsOutlinerImpl.cpp b/Engine/Plugins/Experimental/EditorDataStorageFeatures/Source/TedsOutliner/Private/TedsOutlinerImpl.cpp
--- a/Engine/Plugins/Experimental/EditorDataStorageFeatures/Source/TedsOutliner/Private/TedsOutlinerImpl.cpp
+++ b/Engine/Plugins/Experimental/EditorDataStorageFeatures/Source/TedsOutliner/Private/TedsOutlinerImpl.cpp
@@ -324,28 +324,43 @@ void FTedsOutlinerImpl::CreateFilterQueries()
 					// Create the filter
 					TSharedRef<FTedsOutlinerFilter> TedsFilter = MakeShared<FTedsOutlinerFilter>(Struct->GetFName(), Struct->GetDisplayNameText(), 
 						Struct->GetDisplayNameText(), FName(), Struct->IsChildOf(TedsColumn) ? TedsColumnFilterCategory : TedsTagFilterCategory, 
-						AsShared(), FilterQuery);
+						FilterQuery, true);
+					TedsFilter->SetSceneOutlinerImpl(AsShared());
 					SceneOutliner->AddFilterToFilterBar(TedsFilter);
 				}
 			}
 		});
 	}
 
-	TSharedRef<FFilterCategory> CustomClassFiltersCategory = MakeShared<FFilterCategory>(LOCTEXT("CustomTypeFilters", "Custom Type Filters"), LOCTEXT("CustomTypeFiltersTooltip", "Filter by custom class types"));
-	for(UClass* Class : CreationParams.ClassFilters)
-	{
-		SceneOutliner->AddFilterToFilterBar(MakeShared<FTedsOutlinerFilter>(Class, CustomClassFiltersCategory, AsShared()));
-	}
-	
-	TSharedRef<FFilterCategory> CustomFiltersCategory = MakeShared<FFilterCategory>(LOCTEXT("TedsFilters", "TEDS Custom Filters"), LOCTEXT("TedsFiltersTooltip", "Filter by custom TEDS queries"));
+	const TSharedRef<FFilterCategory> CustomClassFiltersCategory = MakeShared<FFilterCategory>(LOCTEXT("CustomTypeFilters", "Custom Type Filters"), LOCTEXT("CustomTypeFiltersTooltip", "Filter by custom class types"));
+	const TSharedRef<FFilterCategory> CustomFiltersCategory = MakeShared<FFilterCategory>(LOCTEXT("TedsFilters", "TEDS Custom Filters"), LOCTEXT("TedsFiltersTooltip", "Filter by custom TEDS queries"));
 		
-	for(FTedsFilterData& Filter : CreationParams.Filters)
+	for(TSharedPtr<FTedsOutlinerFilter>& Filter : CreationParams.Filters)
 	{
-		if (!Filter.FilterCategory)
+		if (Filter)
 		{
-			Filter.FilterCategory = CustomFiltersCategory;
+			if (!Filter->GetCategory())
+			{
+				if(Filter->IsClassFilter())
+				{
+					Filter->SetCategory(CustomClassFiltersCategory);
+				}
+				else
+				{
+					Filter->SetCategory(CustomFiltersCategory);
+				}
+			}
+
+			Filter->SetSceneOutlinerImpl(AsShared());
+			if (Filter->IsInteractiveFilter())
+			{
+				SceneOutliner->AddFilterToFilterBar(Filter.ToSharedRef());
+			}
+			else
+			{
+				Filter->ActiveStateChanged(true);
+			}
 		}
-		SceneOutliner->AddFilterToFilterBar(MakeShared<FTedsOutlinerFilter>(Filter, AsShared()));
 	}
 }
 
@@ -366,6 +381,13 @@ void FTedsOutlinerImpl::Init()
 	}));
 }
 
+void FTedsOutlinerImpl::FullRefresh() const
+{
+	FSceneOutlinerHierarchyChangedData EventData;
+	EventData.Type = FSceneOutlinerHierarchyChangedData::FullRefresh;
+	HierarchyChangedEvent.Broadcast(EventData);
+}
+
 void FTedsOutlinerImpl::OnMonitoredRowsAdded(DataStorage::FRowHandleArrayView InRows)
 {
 	// Refresh queries so that the added rows can pass through the filters if applied
diff --git a/Engine/Plugins/Experimental/EditorDataStorageFeatures/Source/TedsOutliner/Private/TedsOutlinerModule.cpp b/Engine/Plugins/Experimental/EditorDataStorageFeatures/Source/TedsOutliner/Private/TedsOutlinerModule.cpp
--- a/Engine/Plugins/Experimental/EditorDataStorageFeatures/Source/TedsOutliner/Private/TedsOutlinerModule.cpp
+++ b/Engine/Plugins/Experimental/EditorDataStorageFeatures/Source/TedsOutliner/Private/TedsOutlinerModule.cpp
@@ -2,9 +2,13 @@
 
 #include "TedsOutlinerModule.h"
 
+#include "Engine/Engine.h"
+#include "Engine/World.h"
 #include "LevelEditor.h"
+#include "SceneOutlinerHelpers.h"
 #include "SceneOutlinerPublicTypes.h"
 #include "TedsOutlinerMode.h"
+#include "TedsOutlinerFilter.h"
 #include "WorkspaceMenuStructureModule.h"
 #include "WorkspaceMenuStructure.h"
 #include "Elements/Columns/TypedElementCompatibilityColumns.h"
@@ -132,14 +136,14 @@ FTedsOutlinerColumnDescription FTedsOutlinerModule::GetLevelEditorTedsOutlinerCo
 	using namespace UE::Editor::DataStorage::Columns;
 		
 	TArray<TWeakObjectPtr<const UScriptStruct>> OutColumns({
+		FVisibleInEditorColumn::StaticStruct(),
+		FTedsPrimaryPackageObjectTag::StaticStruct(),
 		FTypedElementClassTypeInfoColumn::StaticStruct(),
-		FAlertColumn::StaticStruct(),
-		FChildAlertColumn::StaticStruct(),
-		FUObjectIdNameColumn::StaticStruct(),
+		FLevelColumn::StaticStruct(), // PackageShortName
 		FTedsActorMobilityColumn::StaticStruct(),
-		FLevelColumn::StaticStruct(),
-		FVisibleInEditorColumn::StaticStruct(),
-		FTedsPrimaryPackageObjectTag::StaticStruct()});
+		FUObjectIdNameColumn::StaticStruct(),
+		FAlertColumn::StaticStruct(),
+		FChildAlertColumn::StaticStruct()});
 
 	if (CVarUseFavoritesWidgets->GetBool())
 	{
@@ -185,7 +189,8 @@ TSharedRef<SWidget> FTedsOutlinerModule::CreateLevelEditorTedsOutliner()
 	Params.QueryDescription = OutlinerQueryDescription;
 	Params.ColumnDescription = GetLevelEditorTedsOutlinerColumnDescription();
 	
-	Params.ClassFilters.Add(AActor::StaticClass());
+	Params.Filters.Add(MakeShared<FTedsOutlinerFilter>(AActor::StaticClass()));
+	Params.Filters.Add(MakeShared<FTedsSceneOutlinerWorldFilter>(TEXT("WorldFilter")));
 	
 	/** Example ways to create filters using query descriptions and functions
 	 * Note QueryDescriptions should always be used over QueryFunctions when possible as they are more performant
@@ -224,7 +229,7 @@ void FTedsOutlinerModule::AddFavoriteFilters(FTedsOutlinerParams& Params) const
 	const TSharedRef<FFilterCategory> FavoriteFiltersCategory = MakeShared<FFilterCategory>(LOCTEXT("FavoriteFilters", "Favorites"), LOCTEXT("FavoriteFiltersTooltip", "Filter by Favorite state"));
 
 	// Create the filter
-	Params.Filters.Emplace(
+	Params.Filters.Add(MakeShared<FTedsOutlinerFilter>(
 		"FilterByFavorites",
 		LOCTEXT("FavoriteFilterDisplayName", "Favorite"),
 		LOCTEXT("FavoriteFilterDisplayNameToolTip", "Filter by Favorites"),
@@ -235,7 +240,7 @@ void FTedsOutlinerModule::AddFavoriteFilters(FTedsOutlinerParams& Params) const
 			.Where()
 			.All<FFavoriteTag>()
 			.Compile())
-		);
+		));
 }
 
 TSharedRef<SDockTab> FTedsOutlinerModule::OpenLevelEditorTedsOutliner(const FSpawnTabArgs& SpawnTabArgs)
@@ -278,6 +283,76 @@ FName FTedsOutlinerModule::GetTedsOutlinerTabName()
 {
 	return TedsOutlinerTabName;
 }
+
+FTedsSceneOutlinerWorldFilter::FTedsSceneOutlinerWorldFilter(const FName& InFilterName, const TWeakObjectPtr<UWorld>& InSpecifiedWorldToDisplay)
+		: FTedsOutlinerFilter(
+          	InFilterName,
+          	FText(), 
+          	Queries::BuildQueryFunction<bool>([this](TQueryContext<RowBatchInfo> Context, TResult<bool>& Result, TConstBatch<FTypedElementWorldColumn> WorldColumns)
+			  {
+				  QueryFunction(Context, Result, WorldColumns);
+			  }),
+          	false)
+		, SpecifiedWorldToDisplay(InSpecifiedWorldToDisplay)
+{
+	BuildWorldFilter();
+	if (GEditor)
+	{
+		GEditor->OnLevelActorListChanged().AddRaw(this, &FTedsSceneOutlinerWorldFilter::BuildWorldFilter);
+	}
+}
+
+FTedsSceneOutlinerWorldFilter::~FTedsSceneOutlinerWorldFilter()
+{
+	if (GEditor)
+	{
+		GEditor->OnLevelActorListChanged().RemoveAll(this);
+	}
+}
+
+void FTedsSceneOutlinerWorldFilter::BuildWorldFilter()
+{
+	if (TedsOutlinerImpl.IsValid())
+	{
+		// Store a temporary World Ptr so we can check if the newly chosen world is the same as
+		// the previous, if so, we don't have to regenerate anything since nothing changed
+		TWeakObjectPtr<UWorld> TempWorld;
+		// This function can eventually accept a manually selected UserChosenWorld as well when this is added to options
+		SceneOutliner::FSceneOutlinerHelpers::ChooseRepresentingWorld(TempWorld, SpecifiedWorldToDisplay);
+		if (TempWorld.IsValid() && TempWorld != RepresentingWorld)
+		{
+			RepresentingWorld = TempWorld;
+
+			if (const TSharedPtr<FTedsOutlinerImpl> TedsOutlinerImplPin = TedsOutlinerImpl.Pin())
+			{
+				// Perform a FullRefresh here since the Outliner does this when Filters change but we have to do it manually 
+				// since this filter isn't managed by the FilterBar.
+				TedsOutlinerImplPin->FullRefresh();
+			}
+		}
+	}
+}
+
+void FTedsSceneOutlinerWorldFilter::QueryFunction(TQueryContext<RowBatchInfo> Context, TResult<bool>& Result, TConstBatch<FTypedElementWorldColumn> WorldColumns)
+{
+	// We always want this to be valid, so if not, find a new world to filter for
+	if (!RepresentingWorld.IsValid())
+	{
+		BuildWorldFilter();
+	}
+	
+	Context.ForEachRow([this, &Result](const FTypedElementWorldColumn& WorldColumn)
+		  {
+			  if (WorldColumn.World.IsValid() && RepresentingWorld.IsValid())
+			  {
+				  Result.Add(WorldColumn.World.Get() == RepresentingWorld.Get());
+			  }
+			  else
+			  {
+				  Result.Add(false);
+			  }
+		  }, WorldColumns);
+}
 } // namespace UE::Editor::Outliner
 
 IMPLEMENT_MODULE(UE::Editor::Outliner::FTedsOutlinerModule, TedsOutliner);
diff --git a/Engine/Plugins/Experimental/EditorDataStorageFeatures/Source/TedsOutliner/Public/TedsOutlinerFilter.h b/Engine/Plugins/Experimental/EditorDataStorageFeatures/Source/TedsOutliner/Public/TedsOutlinerFilter.h
--- a/Engine/Plugins/Experimental/EditorDataStorageFeatures/Source/TedsOutliner/Public/TedsOutlinerFilter.h
+++ b/Engine/Plugins/Experimental/EditorDataStorageFeatures/Source/TedsOutliner/Public/TedsOutlinerFilter.h
@@ -14,81 +14,123 @@ namespace UE::Editor::Outliner
 {
 	using namespace UE::Editor::DataStorage::Queries;
 	class FTedsOutlinerImpl;
-
-	struct FTedsFilterData
+		
+	class TEDSOUTLINER_API FTedsOutlinerFilter : public FFilterBase<SceneOutliner::FilterBarType>
 	{
-		TEDSOUTLINER_API FTedsFilterData(const FName& InFilterName, const FText& InFilterDisplayName, const FText& InFilterToolTip, const FName& InFilterIconName,
-			const TSharedPtr<FFilterCategory>& InCategory, const QueryHandle& InFilterQuery);
-
-		TEDSOUTLINER_API FTedsFilterData(const FName& InFilterName, const FText& InFilterDisplayName, const QueryHandle& InFilterQuery);
+		friend class FTedsOutlinerImpl;
 		
-		TEDSOUTLINER_API FTedsFilterData(const FName& InFilterName, const FText& InFilterDisplayName, const FText& InFilterToolTip, const FName& InFilterIconName,
-			const TSharedPtr<FFilterCategory>& InCategory, const TQueryFunction<bool>& InFilterQuery);
-
-		TEDSOUTLINER_API FTedsFilterData(const FName& InFilterName, const FText& InFilterDisplayName, const TQueryFunction<bool>& InFilterQuery);
-
-		/*
+	public:
+		/**
+		 * Constructor for a TEDS Outliner Filter
+		 * @param InFilterName FName Identifier of the Filter
+		 * @param InFilterDisplayName Display Name to show on the Add Filter Menu and toggle button
+		 * @param InFilterToolTip Tooltip to describe the function of the Filter
+		 * @param InFilterIconName Name of an icon in the App Style Set to display in the Add Filter Menu
+		 * @param InCategory Group to place the filter in the Add Filter Menu
+		 * @param InFilterQuery QueryHandle of a registered query to filter for (must be the same syntax as the Outliner Query)
+		 * @param bInteractiveFilter Whether this filter is visible in the Add Filter Menu and toggleable or is hidden and enabled by default
+		 */
+		FTedsOutlinerFilter(const FName& InFilterName, const FText& InFilterDisplayName, const FText& InFilterToolTip, const FName& InFilterIconName, 
+		                    const TSharedPtr<FFilterCategory>& InCategory, const QueryHandle& InFilterQuery, const bool bInteractiveFilter = true);
+
+		/**
+		 * Constructor for a TEDS Outliner Filter
+		 * @param InFilterName FName Identifier of the Filter
+		 * @param InFilterDisplayName Display Name to show on the Add Filter Menu and toggle button
+		 * @param InFilterQuery QueryHandle of a registered query to filter for (must be the same syntax as the Outliner Query)
+		 * @param bInteractiveFilter Whether this filter is visible in the Add Filter Menu and toggleable or is hidden and enabled by default
+		 */
+		FTedsOutlinerFilter(const FName& InFilterName, const FText& InFilterDisplayName, const QueryHandle& InFilterQuery,
+			const bool bInteractiveFilter = true);
+
+		/**
+		 * Constructor for a TEDS Outliner Filter
+		 * @param InFilterName FName Identifier of the Filter
+		 * @param InFilterDisplayName Display Name to show on the Add Filter Menu and toggle button
+		 * @param InFilterToolTip Tooltip to describe the function of the Filter
+		 * @param InFilterIconName Name of an icon in the App Style Set to display in the Add Filter Menu
+		 * @param InCategory Group to place the filter in the Add Filter Menu
+		 * @param InFilterQuery QueryFunction to filter for, Batching should be used for improved performance if possible
+		 * @param bInteractiveFilter Whether this filter is visible in the Add Filter Menu and toggleable or is hidden and enabled by default
+		 */
+		FTedsOutlinerFilter(const FName& InFilterName, const FText& InFilterDisplayName, const FText& InFilterToolTip, const FName& InFilterIconName,
+			const TSharedPtr<FFilterCategory>& InCategory, const TQueryFunction<bool>& InFilterQuery, const bool bInteractiveFilter = true);
+
+		/**
+		 * Constructor for a TEDS Outliner Filter
+		 * @param InFilterName FName Identifier of the Filter
+		 * @param InFilterDisplayName Display Name to show on the Add Filter Menu and toggle button
+		 * @param InFilterQuery QueryFunction to filter for, Batching should be used for improved performance if possible
+		 * @param bInteractiveFilter Whether this filter is visible in the Add Filter Menu and toggleable or is hidden and enabled by default
+		 */
+		FTedsOutlinerFilter(const FName& InFilterName, const FText& InFilterDisplayName, const TQueryFunction<bool>& InFilterQuery,
+			const bool bInteractiveFilter = true);
+		
+		/**
+		 * Constructor for a TEDS Outliner Filter
 		 * Require that the template type cannot be one of the pre-defined FilterQuery types so it will pass to the defined
 		 * constructor instead of this one.
+		 * @param InFilterName FName Identifier of the Filter
+		 * @param InFilterDisplayName Display Name to show on the Add Filter Menu and toggle button
+		 * @param InFilterToolTip Tooltip to describe the function of the Filter
+		 * @param InFilterIconName Name of an icon in the App Style Set to display in the Add Filter Menu
+		 * @param InCategory Group to place the filter in the Add Filter Menu
+		 * @param InFilterQuery Function callback that checks the value of a column(s) and returns a bool of whether that row passes the filter
+		 * @param bInteractiveFilter Whether this filter is visible in the Add Filter Menu and toggleable or is hidden and enabled by default
 		 */
 		template< typename FilterFunction UE_REQUIRES(
 			!std::is_same_v<std::remove_cvref_t<FilterFunction>, QueryHandle> &&
-			!std::is_same_v<std::remove_cvref_t<FilterFunction>, TQueryFunction<bool>>) >
-		FTedsFilterData(const FName& InFilterName, const FText& InFilterDisplayName, const FText& InFilterToolTip, const FName& InFilterIconName,
-			const TSharedPtr<FFilterCategory>& InCategory, FilterFunction&& InFilterQuery)
-			: FTedsFilterData(
+			!std::is_same_v<std::remove_cvref_t<FilterFunction>, Queries::TQueryFunction<bool>>) >
+		FTedsOutlinerFilter(const FName& InFilterName, const FText& InFilterDisplayName, const FText& InFilterToolTip, const FName& InFilterIconName,
+			const TSharedPtr<FFilterCategory>& InCategory, FilterFunction&& InFilterQuery, const bool bInteractiveFilter = true)
+			: FTedsOutlinerFilter(
 				InFilterName, 
 				InFilterDisplayName, 
 				InFilterToolTip, 
 				InFilterIconName, 
 				InCategory,
-				Queries::BuildQueryFunction<bool>(Forward<FilterFunction>(InFilterQuery)))
+				Queries::BuildQueryFunction<bool>(Forward<FilterFunction>(InFilterQuery)),
+				bInteractiveFilter)
 		{}
 
-		/*
+		/**
+		 * Constructor for a TEDS Outliner Filter
 		 * Require that the template type cannot be one of the pre-defined FilterQuery types so it will pass to the defined
 		 * constructor instead of this one.
+		 * @param InFilterName FName Identifier of the Filter
+		 * @param InFilterDisplayName Display Name to show on the Add Filter Menu and toggle button
+		 * @param InFilterQuery Function callback that checks the value of a column(s) and returns a bool of whether that row passes the filter
+		 * @param bInteractiveFilter Whether this filter is visible in the Add Filter Menu and toggleable or is hidden and enabled by default
 		 */
 		template< typename FilterFunction UE_REQUIRES(
 			!std::is_same_v<std::remove_cvref_t<FilterFunction>, QueryHandle> &&
-			!std::is_same_v<std::remove_cvref_t<FilterFunction>, TQueryFunction<bool>>) >
-	    FTedsFilterData(const FName& InFilterName, const FText& InFilterDisplayName, FilterFunction&& InFilterQuery)
-			: FTedsFilterData(
+			!std::is_same_v<std::remove_cvref_t<FilterFunction>, Queries::TQueryFunction<bool>>) >
+		FTedsOutlinerFilter(const FName& InFilterName, const FText& InFilterDisplayName, FilterFunction&& InFilterQuery,
+			const bool bInteractiveFilter = true)
+			: FTedsOutlinerFilter(
 				InFilterName, 
 				InFilterDisplayName, 
 				FText::FromString(InFilterName.ToString()), 
 				FName(), 
-				nullptr, 
-				Queries::BuildQueryFunction<bool>(Forward<FilterFunction>(InFilterQuery)))
+				nullptr,
+				Queries::BuildQueryFunction<bool>(Forward<FilterFunction>(InFilterQuery)),
+				bInteractiveFilter)
 		{}
 
-		TEDSOUTLINER_API FTedsFilterData(const FTedsFilterData&);
-		TEDSOUTLINER_API ~FTedsFilterData(); 
-		
-		FName FilterName;
-		FText FilterDisplayName;
-		FText FilterToolTip;
-		FName FilterIconName;
-		
-		// The Category can be set as a nullptr in FTedsFilterData to go to the default 'Other Filters' category when adding via the 
-		// FTedsOutlinerParams.Filters list, otherwise the filter will not display if it's not given a category
-		TSharedPtr<FFilterCategory> FilterCategory;
-		TVariant<QueryHandle, TQueryFunction<bool>> FilterQuery;
-	};
-		
-	class FTedsOutlinerFilter : public FFilterBase<SceneOutliner::FilterBarType>
-	{
-	public:
-		FTedsOutlinerFilter(const FTedsFilterData& InTedsFilterData, const TSharedPtr<FTedsOutlinerImpl>& InTedsOutlinerImpl);
+		/**
+		 * Constructor for a TEDS Outliner Class Filter
+		 * @param InClass UClass to filter for on the TypeInfo Column
+		 * @param InCategory Group to place the filter in the Add Filter Menu
+		 * @param bInteractiveFilter Whether this filter is visible in the Add Filter Menu and toggleable or is hidden and enabled by default
+		 */
+		FTedsOutlinerFilter(const UClass* InClass, const TSharedPtr<FFilterCategory>& InCategory = nullptr, const bool bInteractiveFilter = true);
 
-		FTedsOutlinerFilter(const FName& InFilterName, const FText& InFilterDisplayName, const FText& InFilterToolTip, const FName& InFilterIconName, 
-			const TSharedPtr<FFilterCategory>& InCategory, const TSharedPtr<FTedsOutlinerImpl>& InTedsOutlinerImpl, const QueryHandle& InFilterQuery);
+		void SetCategory(const TSharedPtr<FFilterCategory>& InCategory);
 		
-		FTedsOutlinerFilter(const FName& InFilterName, const FText& InFilterDisplayName, const FText& InFilterToolTip, const FName& InFilterIconName,
-			const TSharedPtr<FFilterCategory>& InCategory, const TSharedPtr<FTedsOutlinerImpl>& InTedsOutlinerImpl, const TQueryFunction<bool>& InFilterQuery);
-
-		FTedsOutlinerFilter(const UClass* InClass, const TSharedPtr<FFilterCategory>& InCategory, const TSharedPtr<FTedsOutlinerImpl>& InTedsOutlinerImpl);
+		bool IsClassFilter() const { return bIsClassFilter; }
 
+		bool IsInteractiveFilter() const { return bInteractiveFilter; }
+		
 		/** Returns the system name for this filter */
 		virtual FString GetName() const override;
 
@@ -123,13 +165,16 @@ namespace UE::Editor::Outliner
 		virtual bool PassesFilter(SceneOutliner::FilterBarType InItem ) const override;
 
 	protected:
-
+		/** Sets a SceneOutlinerImpl to be used by this Filter */
+		void SetSceneOutlinerImpl(const TSharedPtr<FTedsOutlinerImpl>& InTedsOutlinerImpl);
+	
 		FName FilterName;
 		FText FilterDisplayName;
 		FText FilterToolTip;
 		FName FilterIconName;
 
 		const bool bIsClassFilter;
+		const bool bInteractiveFilter;
 
 		TWeakPtr<FTedsOutlinerImpl> TedsOutlinerImpl;
 		TVariant<QueryHandle, TQueryFunction<bool>> FilterQuery;
diff --git a/Engine/Plugins/Experimental/EditorDataStorageFeatures/Source/TedsOutliner/Public/TedsOutlinerImpl.h b/Engine/Plugins/Experimental/EditorDataStorageFeatures/Source/TedsOutliner/Public/TedsOutlinerImpl.h
--- a/Engine/Plugins/Experimental/EditorDataStorageFeatures/Source/TedsOutliner/Public/TedsOutlinerImpl.h
+++ b/Engine/Plugins/Experimental/EditorDataStorageFeatures/Source/TedsOutliner/Public/TedsOutlinerImpl.h
@@ -131,7 +131,7 @@ struct FTedsOutlinerParams
 {
 	TEDSOUTLINER_API FTedsOutlinerParams(SSceneOutliner* InSceneOutliner);
 
-	SSceneOutliner* SceneOutliner;
+	SSceneOutliner* SceneOutliner = nullptr;
 
 	// The query description that will be used to populate rows in the TEDS Outliner
 	DataStorage::FQueryDescription QueryDescription;
@@ -140,15 +140,12 @@ struct FTedsOutlinerParams
 	FTedsOutlinerColumnDescription ColumnDescription;
 	
 	// TEDS Filters that utilize QueryDescriptions or QueryFunctions in this TEDS Outliner
-	TArray<FTedsFilterData> Filters;
-
-	// Class types to create filters for in this TEDS Outliner
-	TArray<UClass*> ClassFilters;
+	TArray<TSharedPtr<FTedsOutlinerFilter>> Filters;
 
 	// If true, this Outliner will include a column for row handle
 	bool bShowRowHandleColumn;
 
-	// If true, parent nodes will be remain visible if a child passes all filters even if the parent fails a filter
+	// If true, parent nodes will remain visible if a child passes all filters even if the parent fails a filter
 	bool bForceShowParents;
 
 	// If true, the Teds Outliner will create observers to track addition/removal of rows and update the Outliner
@@ -193,6 +190,7 @@ public:
 	UE_API virtual ~FTedsOutlinerImpl();
 
 	UE_API void Init();
+	UE_API void FullRefresh() const;
 	UE_API void OnMonitoredRowsAdded(DataStorage::FRowHandleArrayView InRows);
 	UE_API void OnMonitoredRowsRemoved(DataStorage::FRowHandleArrayView InRows);
 
@@ -412,6 +410,6 @@ protected:
 	// Whether this is a hybrid outliner, i.e it is displaying TEDS rows and some non-TEDS items
 	bool bHybridMode;
 };
-} // namsepace UE::Editor::Outliner
+} // namespace UE::Editor::Outliner
 
 #undef UE_API
diff --git a/Engine/Plugins/Experimental/EditorDataStorageFeatures/Source/TedsOutliner/Public/TedsOutlinerModule.h b/Engine/Plugins/Experimental/EditorDataStorageFeatures/Source/TedsOutliner/Public/TedsOutlinerModule.h
--- a/Engine/Plugins/Experimental/EditorDataStorageFeatures/Source/TedsOutliner/Public/TedsOutlinerModule.h
+++ b/Engine/Plugins/Experimental/EditorDataStorageFeatures/Source/TedsOutliner/Public/TedsOutlinerModule.h
@@ -7,7 +7,10 @@
 #include "Templates/SharedPointer.h"
 #include "Delegates/IDelegateInstance.h"
 #include "DataStorage/Queries/Description.h"
+#include "TedsOutlinerFilter.h"
+#include "Elements/Framework/TypedElementQueryContext.h"
 
+struct FTypedElementWorldColumn;
 class ISceneOutliner;
 struct FSceneOutlinerInitializationOptions;
 class SDockTab;
@@ -63,4 +66,23 @@ private:
 	FDelegateHandle LevelEditorTabManagerChangedHandle;
 	DataStorage::ICoreProvider* Storage = nullptr;
 };
+
+class TEDSOUTLINER_API FTedsSceneOutlinerWorldFilter : public FTedsOutlinerFilter
+{
+public:
+	FTedsSceneOutlinerWorldFilter(const FName& InFilterName, const TWeakObjectPtr<UWorld>& InSpecifiedWorldToDisplay = nullptr);
+
+	~FTedsSceneOutlinerWorldFilter();
+	
+	void BuildWorldFilter();
+	
+	void QueryFunction(TQueryContext<RowBatchInfo> Context, TResult<bool>& Result, TConstBatch<FTypedElementWorldColumn> WorldColumns);
+	
+protected:
+	/** The world which we are currently filtering for */
+	TWeakObjectPtr<UWorld> RepresentingWorld;
+
+	/** If this mode was created to display a specific world, don't allow it to be reassigned */
+	const TWeakObjectPtr<UWorld> SpecifiedWorldToDisplay;
+};
 } // namespace UE::Editor::Outliner
diff --git a/Engine/Source/Editor/SceneOutliner/Private/ActorMode.cpp b/Engine/Source/Editor/SceneOutliner/Private/ActorMode.cpp
--- a/Engine/Source/Editor/SceneOutliner/Private/ActorMode.cpp
+++ b/Engine/Source/Editor/SceneOutliner/Private/ActorMode.cpp
@@ -36,12 +36,6 @@
 
 DEFINE_LOG_CATEGORY_STATIC(LogActorMode, Log, All);
 
-static int32 GSceneOutlinerAutoRepresentingWorldNetMode = NM_Client;
-static FAutoConsoleVariableRef CVarAutoRepresentingWorldNetMode(
-	TEXT("SceneOutliner.AutoRepresentingWorldNetMode"),
-	GSceneOutlinerAutoRepresentingWorldNetMode,
-	TEXT("The preferred NetMode of the world shown in the scene outliner when the 'Auto' option is chosen: 0=Standalone, 1=DedicatedServer, 2=ListenServer, 3=Client"));
-
 #define LOCTEXT_NAMESPACE "SceneOutliner_ActorMode"
 
 using FActorFilter = TSceneOutlinerPredicateFilter<FActorTreeItem>;
@@ -196,64 +190,7 @@ void FActorMode::Rebuild()
 
 void FActorMode::ChooseRepresentingWorld()
 {
-	// Select a world to represent
-
-	RepresentingWorld = nullptr;
-
-	// If a specified world was provided, represent it
-	if (SpecifiedWorldToDisplay.IsValid())
-	{
-		RepresentingWorld = SpecifiedWorldToDisplay.Get();
-	}
-
-	// check if the user-chosen world is valid and in the editor contexts
-
-	if (!RepresentingWorld.IsValid() && UserChosenWorld.IsValid())
-	{
-		for (const FWorldContext& Context : GEngine->GetWorldContexts())
-		{
-			if (UserChosenWorld.Get() == Context.World())
-			{
-				RepresentingWorld = UserChosenWorld.Get();
-				break;
-			}
-		}
-	}
-
-	// If the user did not manually select a world, try to pick the most suitable world context
-	if (!RepresentingWorld.IsValid())
-	{
-		// Ideally we want a PIE world that is standalone or the first client, unless the preferred NetMode is overridden by CVar
-		int32 LowestPIEInstanceSeen = MAX_int32;
-		for (const FWorldContext& Context : GEngine->GetWorldContexts())
-		{
-			UWorld* World = Context.World();
-			if (World && Context.WorldType == EWorldType::PIE)
-			{
-				if (World->GetNetMode() == NM_Standalone)
-				{
-					RepresentingWorld = World;
-					break;
-				}
-				else if ((World->GetNetMode() == ENetMode(GSceneOutlinerAutoRepresentingWorldNetMode)) && (Context.PIEInstance < LowestPIEInstanceSeen))
-				{
-					RepresentingWorld = World;
-					LowestPIEInstanceSeen = Context.PIEInstance;
-				}
-			}
-		}
-	}
-
-	if (RepresentingWorld == nullptr)
-	{
-		// If there is still no world, we query the Level Editor, which prefers the PIE world over the Editor world
-		TWeakPtr<ILevelEditor> LevelEditor = FModuleManager::GetModuleChecked<FLevelEditorModule>(TEXT("LevelEditor")).GetLevelEditorInstance();
-
-		if (TSharedPtr<ILevelEditor> LevelEditorPin = LevelEditor.Pin())
-		{
-			RepresentingWorld = LevelEditorPin->GetEditorModeManager().GetWorld();
-		}
-	}
+	SceneOutliner::FSceneOutlinerHelpers::ChooseRepresentingWorld(RepresentingWorld, SpecifiedWorldToDisplay, UserChosenWorld);
 }
 
 void FActorMode::BuildWorldPickerMenu(FMenuBuilder& MenuBuilder)
diff --git a/Engine/Source/Editor/SceneOutliner/Private/SceneOutlinerHelpers.cpp b/Engine/Source/Editor/SceneOutliner/Private/SceneOutlinerHelpers.cpp
--- a/Engine/Source/Editor/SceneOutliner/Private/SceneOutlinerHelpers.cpp
+++ b/Engine/Source/Editor/SceneOutliner/Private/SceneOutlinerHelpers.cpp
@@ -7,12 +7,22 @@
 #include "ComponentTreeItem.h"
 #include "EditorActorFolders.h"
 #include "EditorClassUtils.h"
+#include "EditorModeManager.h"
 #include "Engine/Blueprint.h"
+#include "Engine/Engine.h"
+#include "ILevelEditor.h"
+#include "LevelEditor.h"
 #include "Misc/PackageName.h"
 #include "Modules/ModuleManager.h"
 #include "WorldPartition/WorldPartitionActorDescInstance.h"
 #include "UObject/Package.h"
 
+static int32 GSceneOutlinerAutoRepresentingWorldNetMode = NM_Client;
+static FAutoConsoleVariableRef CVarAutoRepresentingWorldNetMode(
+	TEXT("SceneOutliner.AutoRepresentingWorldNetMode"),
+	GSceneOutlinerAutoRepresentingWorldNetMode,
+	TEXT("The preferred NetMode of the world shown in the scene outliner when the 'Auto' option is chosen: 0=Standalone, 1=DedicatedServer, 2=ListenServer, 3=Client"));
+
 
 #define LOCTEXT_NAMESPACE "SceneOutlinerHelpers"
 
@@ -84,6 +94,70 @@ namespace SceneOutliner
 
 		return nullptr;
 	}
+	
+	void FSceneOutlinerHelpers::ChooseRepresentingWorld(TWeakObjectPtr<UWorld>& OutWorld, const TWeakObjectPtr<UWorld>& InSpecifiedWorldToDisplay, const TWeakObjectPtr<UWorld>& InUserChosenWorld)
+	{
+		// Select a world to represent
+
+		TWeakObjectPtr<UWorld> RepresentingWorld = nullptr;
+
+		// If a specified world was provided, represent it
+		if (InSpecifiedWorldToDisplay.IsValid())
+		{
+			RepresentingWorld = InSpecifiedWorldToDisplay.Get();
+		}
+
+		// check if the user-chosen world is valid and in the editor contexts
+
+		if (!RepresentingWorld.IsValid() && InUserChosenWorld.IsValid())
+		{
+			for (const FWorldContext& Context : GEngine->GetWorldContexts())
+			{
+				if (InUserChosenWorld.Get() == Context.World())
+				{
+					RepresentingWorld = InUserChosenWorld.Get();
+					break;
+				}
+			}
+		}
+
+		// If the user did not manually select a world, try to pick the most suitable world context
+		if (!RepresentingWorld.IsValid())
+		{
+			// Ideally we want a PIE world that is standalone or the first client, unless the preferred NetMode is overridden by CVar
+			int32 LowestPIEInstanceSeen = MAX_int32;
+			for (const FWorldContext& Context : GEngine->GetWorldContexts())
+			{
+				UWorld* World = Context.World();
+				if (World && Context.WorldType == EWorldType::PIE)
+				{
+					if (World->GetNetMode() == NM_Standalone)
+					{
+						RepresentingWorld = World;
+						break;
+					}
+					else if ((World->GetNetMode() == ENetMode(GSceneOutlinerAutoRepresentingWorldNetMode)) && (Context.PIEInstance < LowestPIEInstanceSeen))
+					{
+						RepresentingWorld = World;
+						LowestPIEInstanceSeen = Context.PIEInstance;
+					}
+				}
+			}
+		}
+
+		if (!RepresentingWorld.IsValid() && FModuleManager::Get().IsModuleLoaded("LevelEditor"))
+		{
+			// If there is still no world, we query the Level Editor, which prefers the PIE world over the Editor world
+			const TWeakPtr<ILevelEditor> LevelEditor = FModuleManager::GetModuleChecked<FLevelEditorModule>(TEXT("LevelEditor")).GetLevelEditorInstance();
+
+			if (const TSharedPtr<ILevelEditor> LevelEditorPin = LevelEditor.Pin())
+			{
+				RepresentingWorld = LevelEditorPin->GetEditorModeManager().GetWorld();
+			}
+		}
+
+		OutWorld = RepresentingWorld;
+	}
 
 	TSharedPtr<SWidget> FSceneOutlinerHelpers::GetClassHyperlink(UObject* InObject)
 	{
diff --git a/Engine/Source/Editor/SceneOutliner/Public/SceneOutlinerHelpers.h b/Engine/Source/Editor/SceneOutliner/Public/SceneOutlinerHelpers.h
--- a/Engine/Source/Editor/SceneOutliner/Public/SceneOutlinerHelpers.h
+++ b/Engine/Source/Editor/SceneOutliner/Public/SceneOutlinerHelpers.h
@@ -18,6 +18,10 @@ namespace SceneOutliner
 		UE_DEPRECATED(5.5, "FSceneOutlinerHelpers::GetExternalPackage has been deprecated implement/use ISceneOutlinerTreeItem::GetPackageName instead")
 		static UE_API UPackage* GetExternalPackage(const ISceneOutlinerTreeItem& TreeItem);
 		
+		static UE_API void ChooseRepresentingWorld(TWeakObjectPtr<UWorld>& OutWorld,
+			const TWeakObjectPtr<UWorld>& InSpecifiedWorldToDisplay = nullptr,
+			const TWeakObjectPtr<UWorld>& InUserChosenWorld = nullptr);
+		
 		static UE_API TSharedPtr<SWidget> GetClassHyperlink(UObject* InObject);
 		
 		static UE_API void PopulateExtraSearchStrings(const ISceneOutlinerTreeItem& TreeItem, TArray< FString >& OutSearchStrings);
