Adds API exports for clang errors to the TedsOutlinerFilter.h

#rnx

#Synced-CL 51378701

diff --git a/Engine/Plugins/Experimental/EditorDataStorageFeatures/Source/TedsOutliner/Public/TedsOutlinerFilter.h b/Engine/Plugins/Experimental/EditorDataStorageFeatures/Source/TedsOutliner/Public/TedsOutlinerFilter.h
--- a/Engine/Plugins/Experimental/EditorDataStorageFeatures/Source/TedsOutliner/Public/TedsOutlinerFilter.h
+++ b/Engine/Plugins/Experimental/EditorDataStorageFeatures/Source/TedsOutliner/Public/TedsOutlinerFilter.h
@@ -134,37 +134,37 @@ namespace UE::Editor::Outliner
 		bool IsInteractiveFilter() const { return bInteractiveFilter; }
 		
 		/** Returns the system name for this filter */
-		virtual FString GetName() const override;
+		virtual TEDSOUTLINER_API FString GetName() const override;
 
 		/** Returns the human-readable name for this filter */
-		virtual FText GetDisplayName() const override;
+		virtual TEDSOUTLINER_API FText GetDisplayName() const override;
 
 		/** Returns the tooltip for this filter, shown in the filters menu */
-		virtual FText GetToolTipText() const override;
+		virtual TEDSOUTLINER_API FText GetToolTipText() const override;
 
 		/** Returns the color this filter button will be when displayed as a button */
-		virtual FLinearColor GetColor() const override;
+		virtual TEDSOUTLINER_API FLinearColor GetColor() const override;
 
 		/** Returns the name of the icon to use in menu entries */
-		virtual FName GetIconName() const override;
+		virtual TEDSOUTLINER_API FName GetIconName() const override;
 
 		/** If true, the filter will be active in the FilterBar when it is inactive in the UI (i.e the filter pill is grayed out) */
-		virtual bool IsInverseFilter() const override;
+		virtual TEDSOUTLINER_API bool IsInverseFilter() const override;
 
 		/** Notification that the filter became active or inactive */
-		virtual void ActiveStateChanged(bool bActive) override;
+		virtual TEDSOUTLINER_API void ActiveStateChanged(bool bActive) override;
 
 		/** Called when the right-click context menu is being built for this filter */
-		virtual void ModifyContextMenu(FMenuBuilder& MenuBuilder) override;
+		virtual TEDSOUTLINER_API void ModifyContextMenu(FMenuBuilder& MenuBuilder) override;
 		
 		/** Can be overriden for custom FilterBar subclasses to save settings, currently not implemented in any generic Filter Bar */
-		virtual void SaveSettings(const FString& IniFilename, const FString& IniSection, const FString& SettingsString) const override;
+		virtual TEDSOUTLINER_API void SaveSettings(const FString& IniFilename, const FString& IniSection, const FString& SettingsString) const override;
 
 		/** Can be overridden for custom FilterBar subclasses to load settings, currently not implemented in any generic Filter Bar */
-		virtual void LoadSettings(const FString& IniFilename, const FString& IniSection, const FString& SettingsString) override;
+		virtual TEDSOUTLINER_API void LoadSettings(const FString& IniFilename, const FString& IniSection, const FString& SettingsString) override;
 		
 		/** Returns whether the specified Item passes the Filter's restrictions */
-		virtual bool PassesFilter(SceneOutliner::FilterBarType InItem ) const override;
+		virtual TEDSOUTLINER_API bool PassesFilter(SceneOutliner::FilterBarType InItem ) const override;
 
 	protected:
 		/** Sets a SceneOutlinerImpl to be used by this Filter */
