Adjusted and fixed entity & component override notifications in the outliner and details panel.

#Synced-CL 51369505

#Synced-CL 51369505

diff --git a/Engine/Plugins/Experimental/EditorDataStorageFeatures/Source/TedsOutliner/Private/Widgets/OutlinerIconWidget.cpp b/Engine/Plugins/Experimental/EditorDataStorageFeatures/Source/TedsOutliner/Private/Widgets/OutlinerIconWidget.cpp
--- a/Engine/Plugins/Experimental/EditorDataStorageFeatures/Source/TedsOutliner/Private/Widgets/OutlinerIconWidget.cpp
+++ b/Engine/Plugins/Experimental/EditorDataStorageFeatures/Source/TedsOutliner/Private/Widgets/OutlinerIconWidget.cpp
@@ -13,97 +13,95 @@
 
 #define LOCTEXT_NAMESPACE "OutlinerIconWidget"
 
-namespace UE::OutlinerIconWidget::Local
+void UOutlinerIconWidgetFactory::RegisterWidgetConstructors(UE::Editor::DataStorage::ICoreProvider& DataStorage,
+	UE::Editor::DataStorage::IUiProvider& DataStorageUi) const
 {
-	static const FSlateBrush* GetOverrideBadgeFirstLayer(const EOverriddenState& OverriddenState)
-	{
-		switch (OverriddenState)
-		{
-		case EOverriddenState::Added:
-			return FAppStyle::GetBrush("SceneOutliner.OverrideAddedBase");
+	using namespace UE::Editor::DataStorage::Queries;
 
-		case EOverriddenState::AllOverridden:
-			return FAppStyle::GetBrush("SceneOutliner.OverrideBase");
-			break;
+	DataStorageUi.RegisterWidgetFactory<FOutlinerIconWidgetConstructor>(
+		DataStorageUi.FindPurpose(IUiProvider::FPurposeInfo("SceneOutliner", "RowLabel", "Icon").GeneratePurposeID()),
+		TColumn<FTypedElementClassTypeInfoColumn>() || TColumn<FTypedElementScriptStructTypeInfoColumn>());
+}
 
-		case EOverriddenState::HasOverrides:
-			return FAppStyle::GetBrush("SceneOutliner.OverrideInsideBase");
+const FSlateBrush* FOutlinerIconWidgetLogic::GetOverrideBadgeFirstLayer(const UE::Editor::DataStorage::ICoreProvider& Storage,
+	UE::Editor::DataStorage::RowHandle Row) const
+{
+	switch (GetOverriddenState(Storage, Row))
+	{
+	case EOverriddenState::Added:
+		return FAppStyle::GetBrush("SceneOutliner.OverrideAddedBase");
 
-		case EOverriddenState::NoOverrides:
-			// No icon for no overrides
-			break;
+	case EOverriddenState::AllOverridden:
+		return FAppStyle::GetBrush("SceneOutliner.OverrideBase");
+		break;
 
-		case EOverriddenState::SubObjectsHasOverrides:
-			// Not implemented yet
-			break;
-		}
+	case EOverriddenState::HasOverrides:
+	case EOverriddenState::SubObjectsHasOverrides:
+		return FAppStyle::GetBrush("SceneOutliner.OverrideInsideBase");
 
-		return FAppStyle::GetBrush("NoBrush");
+	case EOverriddenState::NoOverrides:
+		// No icon for no overrides
+		break;
 	}
 
-	static const FSlateBrush* GetOverrideBadgeSecondLayer(const EOverriddenState& OverriddenState)
-	{
-		switch (OverriddenState)
-		{
-		case EOverriddenState::Added:
-			return FAppStyle::GetBrush("SceneOutliner.OverrideAdded");
-
-		case EOverriddenState::AllOverridden:
-			return FAppStyle::GetBrush("SceneOutliner.OverrideHere");
+	return FAppStyle::GetBrush("NoBrush");
+}
 
-		case EOverriddenState::HasOverrides:
-			return FAppStyle::GetBrush("SceneOutliner.OverrideInside");
+const FSlateBrush* FOutlinerIconWidgetLogic::GetOverrideBadgeSecondLayer(const UE::Editor::DataStorage::ICoreProvider& Storage,
+	UE::Editor::DataStorage::RowHandle Row) const
+{
+	switch (GetOverriddenState(Storage, Row))
+	{
+	case EOverriddenState::Added:
+		return FAppStyle::GetBrush("SceneOutliner.OverrideAdded");
 
-		case EOverriddenState::NoOverrides:
-			// No icon for no overrides
-			break;
+	case EOverriddenState::AllOverridden:
+		return FAppStyle::GetBrush("SceneOutliner.OverrideHere");
 
-		case EOverriddenState::SubObjectsHasOverrides:
-			// Not implemented yet
-			break;
-		}
+	case EOverriddenState::HasOverrides:
+	case EOverriddenState::SubObjectsHasOverrides:
+		return FAppStyle::GetBrush("SceneOutliner.OverrideInside");
 
-		return FAppStyle::GetBrush("NoBrush");
+	case EOverriddenState::NoOverrides:
+		// No icon for no overrides
+		break;
 	}
 
-	static FText GetOverrideTooltip(const EOverriddenState& OverriddenState)
-	{
-		switch (OverriddenState)
-		{
-		case EOverriddenState::Added:
-			return LOCTEXT("OverrideAddedTooltip", "This entity has been added.");
-
-		case EOverriddenState::AllOverridden:
-			return LOCTEXT("OverrideAllTooltip", "This entity has all its properties overridden.");
+	return FAppStyle::GetBrush("NoBrush");
+}
 
-		case EOverriddenState::HasOverrides:
-			return LOCTEXT("OverrideInsideTooltip", "At least one property or child has an override.");
+FText FOutlinerIconWidgetLogic::GetOverrideTooltip(const UE::Editor::DataStorage::ICoreProvider& Storage,
+	UE::Editor::DataStorage::RowHandle Row) const
+{
+	switch (GetOverriddenState(Storage, Row))
+	{
+	case EOverriddenState::Added:
+		return LOCTEXT("OverrideAddedTooltip", "This entity has been added.");
 
-		case EOverriddenState::NoOverrides:
-			// No icon for no overrides
-			break;
+	case EOverriddenState::AllOverridden:
+		return LOCTEXT("OverrideAllTooltip", "This entity has all its properties overridden.");
 
-		case EOverriddenState::SubObjectsHasOverrides:
-			// Not implemented yet
-			break;
-		}
+	case EOverriddenState::HasOverrides:
+	case EOverriddenState::SubObjectsHasOverrides:
+		return LOCTEXT("OverrideInsideTooltip", "At least one property or child has an override.");
 
-		return FText::GetEmpty();
+	case EOverriddenState::NoOverrides:
+		// No icon for no overrides
+		break;
 	}
+
+	return FText::GetEmpty();
 }
 
-void UOutlinerIconWidgetFactory::RegisterWidgetConstructors(UE::Editor::DataStorage::ICoreProvider& DataStorage,
-	UE::Editor::DataStorage::IUiProvider& DataStorageUi) const
+EOverriddenState FOutlinerIconWidgetLogic::GetOverriddenState(const UE::Editor::DataStorage::ICoreProvider& Storage,
+	UE::Editor::DataStorage::RowHandle Row) const
 {
-	using namespace UE::Editor::DataStorage::Queries;
-
-	DataStorageUi.RegisterWidgetFactory<FOutlinerIconWidgetConstructor>(
-		DataStorageUi.FindPurpose(IUiProvider::FPurposeInfo("SceneOutliner", "RowLabel", "Icon").GeneratePurposeID()),
-		TColumn<FTypedElementClassTypeInfoColumn>() || TColumn<FTypedElementScriptStructTypeInfoColumn>());
+	const FObjectOverrideColumn* Column = Storage.GetColumn<FObjectOverrideColumn>(Row);
+	return Column ? Column->OverriddenState : EOverriddenState::NoOverrides;
 }
 
-FOutlinerIconWidgetConstructor::FOutlinerIconWidgetConstructor()
-	: Super(FOutlinerIconWidgetConstructor::StaticStruct())
+FOutlinerIconWidgetConstructor::FOutlinerIconWidgetConstructor(const UScriptStruct* InTypeInfo)
+	: Super(InTypeInfo)
 {
 }
 
@@ -117,27 +115,32 @@ TSharedPtr<SWidget> FOutlinerIconWidgetConstructor::CreateWidget(UE::Editor::Dat
 		return SNullWidget::NullWidget;
 	}
 
-	FAttributeBinder Binder(TargetRow, DataStorage);
-
+	const FOutlinerIconWidgetLogic& Logic = GetLogic();
+	
 	TSharedRef<SLayeredImage> LayeredImageWidget = SNew(SLayeredImage)
 		.Image(TableViewerUtils::GetIconForRow(DataStorage, TargetRow))
-		.ToolTipText(Binder.BindData(&FObjectOverrideColumn::OverriddenState, [](const EOverriddenState& OverriddenState)
-			{
-				return UE::OutlinerIconWidget::Local::GetOverrideTooltip(OverriddenState);
-			}))
+		.ToolTipText_Lambda([&Logic, DataStorage, TargetRow]()
+		{
+			return Logic.GetOverrideTooltip(*DataStorage, TargetRow);
+		})
 		.ColorAndOpacity(FSlateColor::UseForeground());
 
-	LayeredImageWidget->AddLayer(Binder.BindData(&FObjectOverrideColumn::OverriddenState, [](const EOverriddenState& OverriddenState)
+	LayeredImageWidget->AddLayer(TAttribute<const FSlateBrush*>::CreateLambda([&Logic, DataStorage, TargetRow]()
 		{
-			return UE::OutlinerIconWidget::Local::GetOverrideBadgeFirstLayer(OverriddenState);
+			return Logic.GetOverrideBadgeFirstLayer(*DataStorage, TargetRow);
 		}));
-
-	LayeredImageWidget->AddLayer(Binder.BindData(&FObjectOverrideColumn::OverriddenState, [](const EOverriddenState& OverriddenState)
+	LayeredImageWidget->AddLayer(TAttribute<const FSlateBrush*>::CreateLambda([&Logic, DataStorage, TargetRow]()
 		{
-			return UE::OutlinerIconWidget::Local::GetOverrideBadgeSecondLayer(OverriddenState);
+			return Logic.GetOverrideBadgeSecondLayer(*DataStorage, TargetRow);
 		}));
 
 	return LayeredImageWidget;
 }
 
+const FOutlinerIconWidgetLogic& FOutlinerIconWidgetConstructor::GetLogic() const
+{
+	static FOutlinerIconWidgetLogic Logic;
+	return Logic;
+}
+
 #undef LOCTEXT_NAMESPACE
diff --git a/Engine/Plugins/Experimental/EditorDataStorageFeatures/Source/TedsOutliner/Public/Widgets/OutlinerIconWidget.h b/Engine/Plugins/Experimental/EditorDataStorageFeatures/Source/TedsOutliner/Public/Widgets/OutlinerIconWidget.h
--- a/Engine/Plugins/Experimental/EditorDataStorageFeatures/Source/TedsOutliner/Public/Widgets/OutlinerIconWidget.h
+++ b/Engine/Plugins/Experimental/EditorDataStorageFeatures/Source/TedsOutliner/Public/Widgets/OutlinerIconWidget.h
@@ -4,7 +4,7 @@
 
 #include "Elements/Interfaces/TypedElementDataStorageFactory.h"
 #include "Elements/Interfaces/TypedElementDataStorageUiInterface.h"
-#include "UObject/ObjectMacros.h"
+#include "UObject/OverridableManager.h"
 
 #include "OutlinerIconWidget.generated.h"
 
@@ -17,20 +17,32 @@ class UOutlinerIconWidgetFactory : public UEditorDataStorageFactory
 	GENERATED_BODY()
 
 public:
-	~UOutlinerIconWidgetFactory() override = default;
-
-	void RegisterWidgetConstructors(UE::Editor::DataStorage::ICoreProvider& DataStorage,
+	virtual void RegisterWidgetConstructors(UE::Editor::DataStorage::ICoreProvider& DataStorage,
 		UE::Editor::DataStorage::IUiProvider& DataStorageUi) const override;
 };
 
+struct FOutlinerIconWidgetLogic
+{
+	virtual ~FOutlinerIconWidgetLogic() = default;
+	
+	TEDSOUTLINER_API virtual const FSlateBrush* GetOverrideBadgeFirstLayer(const UE::Editor::DataStorage::ICoreProvider& Storage,
+		UE::Editor::DataStorage::RowHandle Row) const;
+	TEDSOUTLINER_API virtual const FSlateBrush* GetOverrideBadgeSecondLayer(const UE::Editor::DataStorage::ICoreProvider& Storage,
+		UE::Editor::DataStorage::RowHandle Row) const;
+	TEDSOUTLINER_API virtual FText GetOverrideTooltip(const UE::Editor::DataStorage::ICoreProvider& Storage,
+		UE::Editor::DataStorage::RowHandle Row) const;
+	
+	TEDSOUTLINER_API virtual EOverriddenState GetOverriddenState(const UE::Editor::DataStorage::ICoreProvider& Storage,
+		UE::Editor::DataStorage::RowHandle Row) const;
+};
+
 USTRUCT()
 struct FOutlinerIconWidgetConstructor : public FSimpleWidgetConstructor
 {
 	GENERATED_BODY()
 
 public:
-	TEDSOUTLINER_API FOutlinerIconWidgetConstructor();
-	~FOutlinerIconWidgetConstructor() override = default;
+	TEDSOUTLINER_API FOutlinerIconWidgetConstructor(const UScriptStruct* InTypeInfo = StaticStruct());
 
 	TEDSOUTLINER_API virtual TSharedPtr<SWidget> CreateWidget(
 		UE::Editor::DataStorage::ICoreProvider* DataStorage,
@@ -38,4 +50,7 @@ public:
 		UE::Editor::DataStorage::RowHandle TargetRow,
 		UE::Editor::DataStorage::RowHandle WidgetRow,
 		const UE::Editor::DataStorage::FMetaDataView& Arguments) override;
+
+protected:
+	virtual const FOutlinerIconWidgetLogic& GetLogic() const;
 };
diff --git a/Engine/Source/Runtime/CoreUObject/Private/UObject/OverridableManager.cpp b/Engine/Source/Runtime/CoreUObject/Private/UObject/OverridableManager.cpp
--- a/Engine/Source/Runtime/CoreUObject/Private/UObject/OverridableManager.cpp
+++ b/Engine/Source/Runtime/CoreUObject/Private/UObject/OverridableManager.cpp
@@ -166,8 +166,9 @@ void FOverridableManager::RestoreOverrideState(TNotNull<const UObject*> OldObjec
 }
 
 EOverriddenState FOverridableManager::GetOverriddenState(TNotNull<UObject*> Object, 
-	TFunction<bool(TNotNull<const FProperty*>)> ConsiderPropertyForOverriddenState /*= TFunction<bool(TNotNull<const FProperty*>)>*/, 
-	const FPropertyVisitorPath& Path /*= FPropertyVisitorPath()*/, bool* bOutInheritedState /*= nullptr*/)
+	const TFunction<bool(TNotNull<const FProperty*>)>& ConsiderPropertyForOverriddenState /*= TFunction<bool(TNotNull<const FProperty*>)>*/, 
+	const FPropertyVisitorPath& Path /*= FPropertyVisitorPath()*/, bool* bOutInheritedState /*= nullptr*/, bool bSkipAddedState /*= false*/,
+	const TFunction<bool(TNotNull<const UObject*>)>& ConsiderSubObjectForOverriddenState /*= {}*/)
 {
 #if WITH_EDITORONLY_DATA
 	if (const FOverriddenPropertySet* OverriddenProperties = GetOverriddenProperties(Object))
@@ -182,7 +183,7 @@ EOverriddenState FOverridableManager::GetOverriddenState(TNotNull<UObject*> Obje
 
 			// Looking if the Archetype of the object was a CDO is not working in the case of a delete and readd.
 			// So now we explicitly use the bWasAdded to remember if an object was added.
-			if (OverriddenProperties->WasAdded())
+			if (!bSkipAddedState && OverriddenProperties->WasAdded())
 			{
 				return EOverriddenState::Added;
 			}
@@ -190,7 +191,7 @@ EOverriddenState FOverridableManager::GetOverriddenState(TNotNull<UObject*> Obje
 
 		bool bAllOverridden = true;
 		bool bHasOverrides = false;
-		Object->GetClass()->Visit(Object, [this, Object, OverriddenProperties, Path, &bAllOverridden, &bHasOverrides, ConsiderPropertyForOverriddenState, &bOutInheritedState](const FPropertyVisitorContext& Context)
+		Object->GetClass()->Visit(Object, [this, Object, OverriddenProperties, Path, &bAllOverridden, &bHasOverrides, &ConsiderPropertyForOverriddenState, &bOutInheritedState, &ConsiderSubObjectForOverriddenState](const FPropertyVisitorContext& Context)
 		{
 			const EOverriddenPropertyOperation Operation = OverriddenProperties->GetOverriddenPropertyOperation(Context.Path.GetRootIterator());
 			bool bIsEqual = false;
@@ -215,6 +216,7 @@ EOverriddenState FOverridableManager::GetOverriddenState(TNotNull<UObject*> Obje
 
 				if (Operation == EOverriddenPropertyOperation::Replace)
 				{
+					bHasOverrides = true;
 					if (bOutInheritedState)
 					{
 						*bOutInheritedState = !bIsEqual;
@@ -240,6 +242,24 @@ EOverriddenState FOverridableManager::GetOverriddenState(TNotNull<UObject*> Obje
 				return EPropertyVisitorControlFlow::StepOver;
 			}
 
+			// Check if the SubObject should be considered for the overridden state evaluation 
+			UObject* SubObject = nullptr;
+			if (ObjectProperty)
+			{
+				SubObject = ObjectProperty->GetObjectPropertyValue(Context.Data.PropertyData);
+				if (SubObject)
+				{
+					if (!SubObject->IsIn(Object))
+					{
+						SubObject = nullptr;
+					}					
+					else if (ConsiderSubObjectForOverriddenState && !ConsiderSubObjectForOverriddenState(SubObject))
+					{
+						return EPropertyVisitorControlFlow::StepOver;
+					}
+				}
+			}
+				
 			// No need to dig deeper if this is a replace or an add
 			if (Operation == EOverriddenPropertyOperation::Replace || 
 				Operation == EOverriddenPropertyOperation::Add)
@@ -251,34 +271,38 @@ EOverriddenState FOverridableManager::GetOverriddenState(TNotNull<UObject*> Obje
 			// Recurse if the property contains inner ones
 			if (Context.Path.Top().bContainsInnerProperties)
 			{
-				if (Operation != EOverriddenPropertyOperation::None)
+				if (Operation == EOverriddenPropertyOperation::Modified)
 				{
-					bHasOverrides = true;
+					// We will walk through all sub-properties below, except for removed elements, so we test for those here.
+					FOverriddenPropertyNodeID NodeID(Property);
+					if (const FOverriddenPropertyNode* Node = OverriddenProperties->GetOverriddenPropertyNode(NodeID))
+					{
+						bHasOverrides |= Node->GetSubPropertyNodes().ContainsByPredicate([](const FOverriddenPropertyNode& SubNode)
+							{
+								return SubNode.GetOperation() == EOverriddenPropertyOperation::Remove;
+							});
+					}
 				}
 
 				// Only recurse on instanced subobjects
-				if (ObjectProperty)
+				if (SubObject)
 				{
-					if (UObject* InstancedSubObject = ObjectProperty->GetObjectPropertyValue(Context.Data.PropertyData))
+					switch(GetOverriddenState(SubObject, ConsiderPropertyForOverriddenState, {}, nullptr, false, ConsiderSubObjectForOverriddenState))
 					{
-						if (InstancedSubObject->IsIn(Object))
-						{
-							switch(GetOverriddenState(InstancedSubObject, ConsiderPropertyForOverriddenState))
-							{
-								case EOverriddenState::NoOverrides:
-									bAllOverridden = false;
-									break;
-								case EOverriddenState::HasOverrides:
-									bAllOverridden = false;
-									bHasOverrides = true;
-									break;
-								case EOverriddenState::AllOverridden:
-									bHasOverrides = true;
-									break;
-							}
-							return EPropertyVisitorControlFlow::StepOver;
-						}
+						case EOverriddenState::NoOverrides:
+							bAllOverridden = false;
+							break;
+						case EOverriddenState::HasOverrides:
+							bAllOverridden = false;
+							bHasOverrides = true;
+							break;
+						case EOverriddenState::AllOverridden:
+							bHasOverrides = true;
+							break;
+						default:
+							break;
 					}
+					return EPropertyVisitorControlFlow::StepOver;
 				}
 				else
 				{
@@ -296,7 +320,10 @@ EOverriddenState FOverridableManager::GetOverriddenState(TNotNull<UObject*> Obje
 			return EPropertyVisitorControlFlow::Stop;
 		});
 
-		return bAllOverridden ? EOverriddenState::AllOverridden : (bHasOverrides ? EOverriddenState::HasOverrides : EOverriddenState::NoOverrides);
+		if (bHasOverrides)
+		{
+			return bAllOverridden ? EOverriddenState::AllOverridden : EOverriddenState::HasOverrides;
+		}
 	}
 #endif // WITH_EDITORONLY_DATA
 	return EOverriddenState::NoOverrides;
diff --git a/Engine/Source/Runtime/CoreUObject/Public/UObject/OverridableManager.h b/Engine/Source/Runtime/CoreUObject/Public/UObject/OverridableManager.h
--- a/Engine/Source/Runtime/CoreUObject/Public/UObject/OverridableManager.h
+++ b/Engine/Source/Runtime/CoreUObject/Public/UObject/OverridableManager.h
@@ -159,10 +159,13 @@ public:
 	 * @param ConsiderPropertyForOverriddenState optional parameter to consider or not properties in the overridden state evaluation
 	 * @param Path to the property to know its overridden state, empty will return the object overridden state
 	 * @param bOutInheritedState optional parameter to know if the operation returned was inherited from a parent property
+	 * @param bSkipAddedState Set to True if the actual state should be determined instead of just returning "Added", if this object is added.
+	 * @param ConsiderSubObjectForOverriddenState Predicate to decide whether a given SubObject should be included in the state determination.
 	 * @return the overridden state if the object has overridable serialization enabled */
 	COREUOBJECT_API EOverriddenState GetOverriddenState(TNotNull<UObject*> Object, 
-		TFunction<bool(TNotNull<const FProperty*>)> ConsiderPropertyForOverriddenState = TFunction<bool(TNotNull<const FProperty*>)>(),
-		const FPropertyVisitorPath& Path = FPropertyVisitorPath(), bool* bOutInheritedState = nullptr);
+		const TFunction<bool(TNotNull<const FProperty*>)>& ConsiderPropertyForOverriddenState = TFunction<bool(TNotNull<const FProperty*>)>(),
+		const FPropertyVisitorPath& Path = FPropertyVisitorPath(), bool* bOutInheritedState = nullptr,
+		bool bSkipAddedState = false, const TFunction<bool(TNotNull<const UObject*>)>& ConsiderSubObjectForOverriddenState = {});
 
 	/**
 	 * Override the entire object properties and all its instanced subobjects
