Teds Outliner: Add support for displaying complex hierarchies by subclassing the hierarchy viewer's data interface.
Old hierarchy implementation in FTedsOutlinerImpl has been moved to a new FTedsOutlinerLegacyHierarchyData
Add a new FTedsOutllinerMultiHierarchyData that can combine multiple hierarchies into one view using the new hierarchy API in TEDS.

Modify the data driven outlienr tab to display multiple hierarchies

#Synced-CL 51413783

#Synced-CL 51413783

diff --git a/Engine/Plugins/ChaosVD/Source/ChaosVD/Private/ChaosVDWorldOutlinerTab.cpp b/Engine/Plugins/ChaosVD/Source/ChaosVD/Private/ChaosVDWorldOutlinerTab.cpp
--- a/Engine/Plugins/ChaosVD/Source/ChaosVD/Private/ChaosVDWorldOutlinerTab.cpp
+++ b/Engine/Plugins/ChaosVD/Source/ChaosVD/Private/ChaosVDWorldOutlinerTab.cpp
@@ -90,7 +90,7 @@ void FChaosVDWorldOutlinerTab::CreateWorldOutlinerWidget()
 			}
 		});
 		
-	Params.HierarchyData = FTedsOutlinerHierarchyData(FChaosVDTableRowParentColumn::StaticStruct(), RowHandleGetter, RowHandleSetter, ChildrenRowHandlesGetter);
+	Params.HierarchyData = MakeShared<FTedsOutlinerLegacyHierarchyInterface>(FTedsOutlinerHierarchyData(FChaosVDTableRowParentColumn::StaticStruct(), RowHandleGetter, RowHandleSetter, ChildrenRowHandlesGetter));
 	
 	if (TSharedPtr<FChaosVDScene> CvdScene = GetChaosVDScene().Pin())
 	{
diff --git a/Engine/Plugins/Experimental/EditorDataStorageFeatures/Source/TedsOutliner/Private/TedsOutlinerHierarchyInterfaces.cpp b/Engine/Plugins/Experimental/EditorDataStorageFeatures/Source/TedsOutliner/Private/TedsOutlinerHierarchyInterfaces.cpp
new file mode 100644
--- /dev/null
+++ b/Engine/Plugins/Experimental/EditorDataStorageFeatures/Source/TedsOutliner/Private/TedsOutlinerHierarchyInterfaces.cpp
@@ -0,0 +1,309 @@
+// Copyright Epic Games, Inc. All Rights Reserved.
+
+#include "TedsOutlinerHierarchyInterfaces.h"
+
+#include "TedsOutlinerItem.h"
+#include "DataStorage/Features.h"
+#include "Elements/Columns/TypedElementMiscColumns.h"
+#include "Elements/Framework/TypedElementQueryBuilder.h"
+
+namespace UE::Editor::Outliner
+{
+	namespace Private
+	{
+		// Helper method to check if the parent of an item has changed, used with the legacy interface which used sync tags to detect hierarchy changes
+		// but weren't fully accurate because sync tags are added on ANY changes to the UObject row
+		static bool HasItemParentChanged(DataStorage::IQueryContext& Context, DataStorage::RowHandle Row, DataStorage::RowHandle ParentRowHandle, ISceneOutliner& SceneOutliner)
+		{
+			const FSceneOutlinerTreeItemPtr Item = SceneOutliner.GetTreeItem(Row, true);
+
+			// If the item doesn't exist, it doesn't make sense to say its parent changed
+			if (!Item)
+			{
+				return false;
+			}
+										
+			const FSceneOutlinerTreeItemPtr ParentItem = Item->GetParent();
+
+			// If the item doesn't have a parent, but ParentRowHandle is valid: The item just got added a parent so we want to dirty it
+			if (!ParentItem)
+			{
+				return Context.IsRowAvailable(ParentRowHandle);
+			}
+										
+			const FTedsOutlinerTreeItem* TedsParentItem = ParentItem->CastTo<FTedsOutlinerTreeItem>();
+
+			if (TedsParentItem)
+			{
+				// return true if the row handle of the parent item doesn't match what we are given, i.e the parent has changed
+				return TedsParentItem->GetRowHandle() != ParentRowHandle;
+			}
+
+			return false;
+		};
+	}
+
+	// @section FTedsOutlinerMultiHierarchyInterface
+	
+	FTedsOutlinerMultiHierarchyInterface::FTedsOutlinerMultiHierarchyInterface(TSharedRef<DataStorage::FHierarchyViewerMultiData> InHierarchyData)
+		: HierarchyData(InHierarchyData)
+	{
+		
+	}
+
+	DataStorage::RowHandle FTedsOutlinerMultiHierarchyInterface::GetParent(const DataStorage::ICoreProvider& Storage,
+		DataStorage::RowHandle InRow) const
+	{
+		return HierarchyData->GetParent(Storage, InRow);
+	}
+
+	void FTedsOutlinerMultiHierarchyInterface::WalkDepthFirst(const DataStorage::ICoreProvider& Storage, DataStorage::RowHandle InRow,
+		DataStorage::ICoreProvider::FHierarchyIterationCallback VisitFn, DataStorage::ICoreProvider::ETraversalOrder TraversalOrder) const
+	{
+		HierarchyData->WalkDepthFirst(Storage, InRow, VisitFn, TraversalOrder);
+	}
+
+	void FTedsOutlinerMultiHierarchyInterface::RegisterQueries(DataStorage::ICoreProvider& Storage,
+		const DataStorage::FQueryDescription& OutlinerQueryDescription, TWeakPtr<ISceneOutliner> Outliner, bool bUsingQueryConditionsSyntax)
+	{
+		// For each hierarchy we are displaying, get the parent change column (if present) and add a processor to detect changes to rows with the change column
+		HierarchyData->ForEachHierarchy([this, &Storage](FHierarchyHandle HierarchyHandle)
+		{
+			if (const UScriptStruct* ParentChangeTag = Storage.GetParentChangedColumnType(HierarchyHandle))
+			{
+				FString ProcessorName = TEXT("FTedsOutlinerMultiHierarchyInterface: Update Parent for ");
+				ParentChangeTag->GetFName().AppendString(ProcessorName);
+				
+				FQueryDescription UpdateParentQueryDescription =
+					Select(
+					FName(ProcessorName),
+					FProcessor(EQueryTickPhase::PostPhysics, Storage.GetQueryTickGroupName(EQueryTickGroups::SyncExternalToDataStorage))
+						.SetExecutionMode(EExecutionMode::GameThread),
+					[this](IQueryContext& Context, const RowHandle* Rows)
+					{
+						const int32 RowCount = Context.GetRowCount();
+						TConstArrayView<RowHandle> RowsView = MakeArrayView(Rows, RowCount);
+							
+						for (RowHandle Row : RowsView)
+						{
+							ParentChangedEvent.Broadcast(Row);
+						}
+					})
+				.Where()
+					.All(ParentChangeTag)
+				.Compile();
+				
+				HierarchyChangeQueries.Add(Storage.RegisterQuery(MoveTemp(UpdateParentQueryDescription)));
+			}
+		});
+	}
+
+	void FTedsOutlinerMultiHierarchyInterface::UnregisterQueries(DataStorage::ICoreProvider& Storage)
+	{
+		for (DataStorage::QueryHandle Query : HierarchyChangeQueries)
+		{
+			Storage.UnregisterQuery(Query);
+		}
+		HierarchyChangeQueries.Empty();
+	}
+	
+	// @section FTedsOutlinerHierarchyData
+	
+	FTedsOutlinerHierarchyData FTedsOutlinerHierarchyData::GetDefaultHierarchyData()
+	{
+		const FGetParentRowHandle RowHandleGetter = FGetParentRowHandle::CreateLambda([](const void* InColumnData)
+			{
+				if(const FTableRowParentColumn* ParentColumn = static_cast<const FTableRowParentColumn *>(InColumnData))
+				{
+					return ParentColumn->Parent;
+				}
+    
+				return DataStorage::InvalidRowHandle;
+			});
+    
+		const FSetParentRowHandle RowHandleSetter = FSetParentRowHandle::CreateLambda([](void* InColumnData,
+			DataStorage::RowHandle InRowHandle)
+			{
+				if(FTableRowParentColumn* ParentColumn = static_cast<FTableRowParentColumn *>(InColumnData))
+				{
+					ParentColumn->Parent = InRowHandle;
+				}
+			});
+    		
+		return FTedsOutlinerHierarchyData(FTableRowParentColumn::StaticStruct(), RowHandleGetter, RowHandleSetter, FGetChildrenRowsHandles());
+	}
+	
+	// @section FTedsOutlinerLegacyHierarchyInterface
+
+	FTedsOutlinerLegacyHierarchyInterface::FTedsOutlinerLegacyHierarchyInterface(FTedsOutlinerHierarchyData InHierarchyData)
+		: HierarchyData(InHierarchyData)
+	{
+	}
+
+	DataStorage::RowHandle FTedsOutlinerLegacyHierarchyInterface::GetParent(const DataStorage::ICoreProvider& Storage,
+		DataStorage::RowHandle InRow) const
+	{
+		if (HierarchyData.GetParent.IsBound())
+		{
+			if (const void* HierarchyColumnData = Storage.GetColumnData(InRow, HierarchyData.HierarchyColumn))
+			{
+				return HierarchyData.GetParent.Execute(HierarchyColumnData);
+			}
+		}
+		
+		return DataStorage::InvalidRowHandle;
+		
+	}
+
+	void FTedsOutlinerLegacyHierarchyInterface::WalkDepthFirst(const DataStorage::ICoreProvider& Storage, DataStorage::RowHandle InRow,
+		DataStorage::ICoreProvider::FHierarchyIterationCallback VisitFn, DataStorage::ICoreProvider::ETraversalOrder TraversalOrder) const
+	{
+		TArray<DataStorage::RowHandle> Children;
+		
+		if (HierarchyData.GetChildren.IsBound())
+		{
+			if (const void* HierarchyColumnData = Storage.GetColumnData(InRow, HierarchyData.HierarchyColumn))
+			{
+				Children = HierarchyData.GetChildren.Execute(const_cast<void*>(HierarchyColumnData));
+			}
+		}
+		// We still have to support the legacy case of grabbing child rows manually since it is still in use by places that don't specify HierarchyData.GetChildren
+		// This runs a query to get all rows with FTableRowParentColumn and recurses through the tree to find rows that are children of InRow
+		// NOTE: This is extremely expensive and only exists for backward compatibility
+		else if (HierarchyData.GetParent.IsBound())
+		{
+			FRowHandleArray MatchedRowsWithParentColumn;
+
+			// Collect all rows with the parent column
+			DirectQueryCallback ChildRowCollector = CreateDirectQueryCallbackBinding(
+			[&MatchedRowsWithParentColumn] (const IDirectQueryContext& Context, const RowHandle*)
+			{
+				MatchedRowsWithParentColumn.Append(Context.GetRowHandles());
+			});
+
+			// The legacy case relies on running a query inline, but we don't want to keep supporting that in the new cases so we do a const cast here
+			const_cast<ICoreProvider&>(Storage).RunQuery(GetChildrenQuery, ChildRowCollector);
+
+			// Recursively get the children for each row
+			TFunction<void(RowHandle)> GetChildrenRecursive = 
+				[&Children, &MatchedRowsWithParentColumn, &Storage, &GetChildrenRecursive, InHierarchyData = HierarchyData]
+				(RowHandle ExpectedParentRow) -> void
+			{
+				for(RowHandle PotentialChildRow : MatchedRowsWithParentColumn.GetRows())
+				{
+					const void* ParentColumnData = Storage.GetColumnData(PotentialChildRow, InHierarchyData.HierarchyColumn);
+
+					if (ensureMsgf(ParentColumnData, TEXT("We should always have the parent column since we only grabbed rows with those ")))
+					{
+						// Get the parent row handle
+						const RowHandle ParentRowHandle = InHierarchyData.GetParent.Execute(ParentColumnData);
+			
+						// Check if this row is owned by InRow
+						if (ParentRowHandle == ExpectedParentRow)
+						{
+							Children.Add(PotentialChildRow);
+
+							// Recursively look for children of this row now
+							GetChildrenRecursive(PotentialChildRow);
+						}
+					}
+				}
+			};
+
+			GetChildrenRecursive(InRow);
+		}
+		
+		for (DataStorage::RowHandle ChildRow : Children)
+		{
+			VisitFn(Storage, InRow, ChildRow);
+		}
+	}
+
+	void FTedsOutlinerLegacyHierarchyInterface::RegisterQueries(DataStorage::ICoreProvider& Storage, 
+		const DataStorage::FQueryDescription& OutlinerQueryDescription, TWeakPtr<ISceneOutliner> Outliner, bool bUsingQueryConditionsSyntax)
+	{
+		using namespace UE::Editor::DataStorage;
+		using namespace UE::Editor::DataStorage::Queries;
+		
+		// Legacy query that uses the sync tags to detect changes to a row + a util function to determine if that change was actually a hierarchy change
+		FQueryDescription UpdateParentQueryDescription =
+			Select(
+			TEXT("FTedsOutlinerLegacyHierarchyInterface: Update item parent"),
+			FProcessor(EQueryTickPhase::PostPhysics, Storage.GetQueryTickGroupName(EQueryTickGroups::SyncExternalToDataStorage))
+				.SetExecutionMode(EExecutionMode::GameThread),
+			[this, Outliner](IQueryContext& Context, const RowHandle* Rows)
+			{
+				if (TSharedPtr<ISceneOutliner> OutlinerPin = Outliner.Pin())
+				{
+					if (const char* ParentColumn = reinterpret_cast<const char*>(Context.GetColumn(HierarchyData.HierarchyColumn)))
+					{
+						int32 ColumnSize = HierarchyData.HierarchyColumn->GetCppStructOps() 
+							? HierarchyData.HierarchyColumn->GetCppStructOps()->GetSize() 
+							: HierarchyData.HierarchyColumn->GetStructureSize();
+					
+						uint32 RowCount = Context.GetRowCount();
+					
+						for(uint32 RowIndex = 0; RowIndex < RowCount; ++RowIndex, ParentColumn += ColumnSize)
+						{
+							RowHandle ParentRowHandle = HierarchyData.GetParent.Execute(ParentColumn);
+								
+							if (Private::HasItemParentChanged(Context, Rows[RowIndex], ParentRowHandle, *OutlinerPin))
+							{
+								ParentChangedEvent.Broadcast(Rows[RowIndex]);
+							}
+						}
+					}
+				}
+				
+			})
+			.ReadOnly(HierarchyData.HierarchyColumn, EOptional::Yes)
+		.Compile();
+
+		if (bUsingQueryConditionsSyntax)
+		{
+			UpdateParentQueryDescription.Conditions.Emplace(TColumn<FTypedElementSyncBackToWorldTag>());
+		}
+		else
+		{
+			UpdateParentQueryDescription.ConditionTypes.Add(FQueryDescription::EOperatorType::SimpleAll);
+			UpdateParentQueryDescription.ConditionOperators.AddZeroed_GetRef().Type = FTypedElementSyncBackToWorldTag::StaticStruct();
+		}
+		
+		// Add the conditions from FinalQueryDescription to ensure we are only observing the rows the user requested
+		MergeQueries(UpdateParentQueryDescription, OutlinerQueryDescription);
+		
+		HierarchyChangeQuery = Storage.RegisterQuery(MoveTemp(UpdateParentQueryDescription));
+		
+		if (!HierarchyData.GetChildren.IsBound())
+		{
+			// Query to get all rows that match our conditions with a parent column (i.e all child rows)
+			FQueryDescription ChildHandleQueryDescription;
+		
+			if (bUsingQueryConditionsSyntax)
+			{
+				ChildHandleQueryDescription = Select()
+					.Where(TColumn(HierarchyData.HierarchyColumn))
+					.Compile();
+			}
+			else
+			{
+				ChildHandleQueryDescription =
+					Select()
+					.Where()
+						.All(HierarchyData.HierarchyColumn)
+					.Compile();
+			}
+		
+			MergeQueries(ChildHandleQueryDescription, OutlinerQueryDescription);
+			GetChildrenQuery = Storage.RegisterQuery(MoveTemp(ChildHandleQueryDescription));
+		}
+	}
+
+	void FTedsOutlinerLegacyHierarchyInterface::UnregisterQueries(DataStorage::ICoreProvider& Storage)
+	{
+		Storage.UnregisterQuery(HierarchyChangeQuery);
+		Storage.UnregisterQuery(GetChildrenQuery);
+		HierarchyChangeQuery = DataStorage::InvalidQueryHandle;
+		GetChildrenQuery = DataStorage::InvalidQueryHandle;
+	}
+}
diff --git a/Engine/Plugins/Experimental/EditorDataStorageFeatures/Source/TedsOutliner/Private/TedsOutlinerImpl.cpp b/Engine/Plugins/Experimental/EditorDataStorageFeatures/Source/TedsOutliner/Private/TedsOutlinerImpl.cpp
--- a/Engine/Plugins/Experimental/EditorDataStorageFeatures/Source/TedsOutliner/Private/TedsOutlinerImpl.cpp
+++ b/Engine/Plugins/Experimental/EditorDataStorageFeatures/Source/TedsOutliner/Private/TedsOutlinerImpl.cpp
@@ -79,34 +79,7 @@ namespace QueryUtils
 		return true;
 	}
 
-	static bool HasItemParentChanged(DataStorage::IQueryContext& Context, DataStorage::RowHandle Row, DataStorage::RowHandle ParentRowHandle, SSceneOutliner& SceneOutliner)
-	{
-		const FSceneOutlinerTreeItemPtr Item = SceneOutliner.GetTreeItem(Row, true);
-
-		// If the item doesn't exist, it doesn't make sense to say its parent changed
-		if (!Item)
-		{
-			return false;
-		}
-										
-		const FSceneOutlinerTreeItemPtr ParentItem = Item->GetParent();
-
-		// If the item doesn't have a parent, but ParentRowHandle is valid: The item just got added a parent so we want to dirty it
-		if (!ParentItem)
-		{
-			return Context.IsRowAvailable(ParentRowHandle);
-		}
-										
-		const FTedsOutlinerTreeItem* TedsParentItem = ParentItem->CastTo<FTedsOutlinerTreeItem>();
-
-		if (TedsParentItem)
-		{
-			// return true if the row handle of the parent item doesn't match what we are given, i.e the parent has changed
-			return TedsParentItem->GetRowHandle() != ParentRowHandle;
-		}
-
-		return false;
-	};
+	
 }
 
 namespace LabelWidgetUtils
@@ -227,7 +200,7 @@ FTedsOutlinerParams::FTedsOutlinerParams(SSceneOutliner* InSceneOutliner)
 	, bShowRowHandleColumn(true)
 	, bForceShowParents(true)
 	, bUseDefaultObservers(true)
-	, HierarchyData(FTedsOutlinerHierarchyData::GetDefaultHierarchyData())
+	, HierarchyData(MakeShared<FTedsOutlinerLegacyHierarchyInterface>(FTedsOutlinerHierarchyData::GetDefaultHierarchyData()))
 	, SelectionSet(nullptr)
 {
 	CellWidgetPurpose =
@@ -366,6 +339,17 @@ void FTedsOutlinerImpl::CreateFilterQueries()
 
 void FTedsOutlinerImpl::Init()
 {
+	if (HierarchyData)
+	{
+		HierarchyData->OnHierarchyChangedEvent().AddSPLambda(this, [this](RowHandle TargetRow)
+		{
+			FSceneOutlinerHierarchyChangedData EventData;
+			EventData.Type = FSceneOutlinerHierarchyChangedData::Moved;
+			EventData.ItemIDs.Add(TargetRow);
+			HierarchyChangedEvent.Broadcast(EventData);
+		});
+	}
+	
 	RegisterTedsOutliner();
 	CreateFilterQueries();
 	RecompileQueries();
@@ -427,6 +411,11 @@ FTedsOutlinerImpl::~FTedsOutlinerImpl()
 {
 	if (Storage)
 	{
+		if (HierarchyData)
+		{
+			HierarchyData->UnregisterQueries(*Storage);
+		}
+		
 		Storage->RemoveRow(OutlinerRowHandle);
 		Storage->OnUpdateCompleted().RemoveAll(this);
 	}
@@ -762,7 +751,7 @@ void FTedsOutlinerImpl::CreateChildren(const FSceneOutlinerTreeItemPtr& Item, TA
 	 */
 
 	// If there's no hierarchy data, there is no need to create children
-	if (!HierarchyData.IsSet())
+	if (!HierarchyData)
 	{
 		return;
 	}
@@ -784,59 +773,21 @@ void FTedsOutlinerImpl::CreateChildren(const FSceneOutlinerTreeItemPtr& Item, TA
 		return;
 	}
 
-	TArray<RowHandle> ChildItems;
-
-	if (HierarchyData->GetChildren.IsBound())
-	{
-		void* ParentColumnData = Storage->GetColumnData(ItemRowHandle, HierarchyData->HierarchyColumn);
-		ChildItems = HierarchyData->GetChildren.Execute(ParentColumnData);
-	}
-	else
+	FRowHandleArray ChildItems;
+	
+	HierarchyData->WalkDepthFirst(*Storage, ItemRowHandle, [ItemRowHandle, &ChildItems](const ICoreProvider& Context, RowHandle Owner, RowHandle Target)
 	{
-		TArray<RowHandle> MatchedRowsWithParentColumn;
-
-		// Collect all entities that are owned by our entity
-		DirectQueryCallback ChildRowCollector = CreateDirectQueryCallbackBinding(
-		[&MatchedRowsWithParentColumn] (const IDirectQueryContext& Context, const RowHandle*)
-		{
-			MatchedRowsWithParentColumn.Append(Context.GetRowHandles());
-		});
-
-		Storage->RunQuery(ChildRowHandleQuery, ChildRowCollector);
-
-		// Recursively get the children for each entity
-		TFunction<void(RowHandle)> GetChildrenRecursive = [&ChildItems, &MatchedRowsWithParentColumn, DataStorage = Storage, &GetChildrenRecursive, InHierarchyData = HierarchyData]
-		(RowHandle EntityRowHandle) -> void
+		if (Target != ItemRowHandle)
 		{
-			for(RowHandle ChildEntityRowHandle : MatchedRowsWithParentColumn)
-			{
-				const void* ParentColumnData = DataStorage->GetColumnData(ChildEntityRowHandle, InHierarchyData.GetValue().HierarchyColumn);
-
-				if (ensureMsgf(ParentColumnData, TEXT("We should always the a parent column since we only grabbed rows with those ")))
-				{
-					// Get the parent row handle
-					const RowHandle ParentRowHandle = InHierarchyData.GetValue().GetParent.Execute(ParentColumnData);
-				
-					// Check if this entity is owned by the entity we are looking children for
-					if (ParentRowHandle == EntityRowHandle)
-					{
-						ChildItems.Add(ChildEntityRowHandle);
-
-						// Recursively look for children of this item
-						GetChildrenRecursive(ChildEntityRowHandle);
-					}
-				}
-			}
-		};
-
-		GetChildrenRecursive(ItemRowHandle);
-	}
+			ChildItems.Add(Target);
+		}
+	});
 	
 	// If the row doesn't exist in our query stack results, it doesn't match the query and should not be displayed
 	FRowHandleArrayView Rows(RowHandleSortNode->GetRows());
 
 	// Actually create the items for the child entities 
-	for (RowHandle ChildItemRowHandle : ChildItems)
+	for (RowHandle ChildItemRowHandle : ChildItems.GetRows())
 	{
 		if (!Rows.Contains(ChildItemRowHandle))
 		{
@@ -853,21 +804,12 @@ void FTedsOutlinerImpl::CreateChildren(const FSceneOutlinerTreeItemPtr& Item, TA
 DataStorage::RowHandle FTedsOutlinerImpl::GetParentRow(DataStorage::RowHandle InRowHandle)
 {
 	// No parent if there is no hierarchy data specified
-	if (!HierarchyData.IsSet())
+	if (!HierarchyData)
 	{
 		return DataStorage::InvalidRowHandle;
 	}
 	
-	// If this entity does not have a parent entity, return InvalidRowHandle
-	const void* ParentColumnData = Storage->GetColumnData(InRowHandle, HierarchyData.GetValue().HierarchyColumn);
-	
-	if (!ParentColumnData)
-	{
-		return DataStorage::InvalidRowHandle;
-	}
-
-	// If the parent is invalid for some reason, return InvalidRowHandle
-	const DataStorage::RowHandle ParentRowHandle = HierarchyData.GetValue().GetParent.Execute(ParentColumnData);
+	const DataStorage::RowHandle ParentRowHandle = HierarchyData->GetParent(*Storage, InRowHandle);
 	
 	if (!Storage->IsRowAvailable(ParentRowHandle))
 	{
@@ -989,77 +931,10 @@ void FTedsOutlinerImpl::RecompileQueries()
 	AppendExternalQueries(FinalQueryDescription);
 
 	// Queries to track parent info, only required if we have hierarchy data
-	if (HierarchyData.IsSet())
+	if (HierarchyData)
 	{
-		const UScriptStruct* ParentColumnType = HierarchyData.GetValue().HierarchyColumn;
-
-		// Query to get all rows that match our conditions with a parent column (i.e all child rows)
-		FQueryDescription ChildHandleQueryDescription;
-		
-		if (bUsingQueryConditionsSyntax)
-		{
-			ChildHandleQueryDescription = Select()
-				.Where(TColumn(ParentColumnType))
-				.Compile();
-		}
-		else
-		{
-			ChildHandleQueryDescription =
-				Select()
-				.Where()
-					.All(ParentColumnType)
-				.Compile();
-		}
-		
-		// Add the conditions from FinalQueryDescription to ensure we are tracking removal of the rows the user requested
-		MergeQueries(ChildHandleQueryDescription, FinalQueryDescription);
-
-		FQueryDescription UpdateParentQueryDescription =
-			Select(
-			TEXT("Update item parent"),
-			FProcessor(EQueryTickPhase::PostPhysics, Storage->GetQueryTickGroupName(EQueryTickGroups::SyncExternalToDataStorage))
-				.SetExecutionMode(EExecutionMode::GameThread),
-			[this, ParentColumnType](IQueryContext& Context, const RowHandle* Rows)
-			{
-				if (const char* ParentColumn = reinterpret_cast<const char*>(Context.GetColumn(ParentColumnType)))
-				{
-					int32 ColumnSize = ParentColumnType->GetCppStructOps() ? ParentColumnType->GetCppStructOps()->GetSize() : ParentColumnType->GetStructureSize();
-					const FTedsOutlinerColumn* TedsOutlinerColumns = Context.GetColumn<FTedsOutlinerColumn>();
-					uint32 RowCount = Context.GetRowCount();
-					
-					for(uint32 RowIndex = 0; RowIndex < RowCount; ++RowIndex, ParentColumn += ColumnSize)
-					{
-						RowHandle ParentRowHandle = HierarchyData.GetValue().GetParent.Execute(ParentColumn);
-								
-						if (QueryUtils::HasItemParentChanged(Context, Rows[RowIndex], ParentRowHandle, *SceneOutliner))
-						{
-							FSceneOutlinerHierarchyChangedData EventData;
-							EventData.Type = FSceneOutlinerHierarchyChangedData::Moved;
-							EventData.ItemIDs.Add(Rows[RowIndex]);
-							HierarchyChangedEvent.Broadcast(EventData);
-						}
-					}
-				}
-			})
-			.ReadOnly(ParentColumnType, EOptional::Yes)
-			.ReadOnly<FTedsOutlinerColumn>(EOptional::Yes)
-		.Compile();
-
-		if (bUsingQueryConditionsSyntax)
-		{
-			UpdateParentQueryDescription.Conditions.Emplace(TColumn<FTypedElementSyncBackToWorldTag>());
-		}
-		else
-		{
-			UpdateParentQueryDescription.ConditionTypes.Add(FQueryDescription::EOperatorType::SimpleAll);
-			UpdateParentQueryDescription.ConditionOperators.AddZeroed_GetRef().Type = FTypedElementSyncBackToWorldTag::StaticStruct();
-		}
-		
-		// Add the conditions from FinalQueryDescription to ensure we are the rows the user requested
-		MergeQueries(UpdateParentQueryDescription, FinalQueryDescription);
-		
-		ChildRowHandleQuery = Storage->RegisterQuery(MoveTemp(ChildHandleQueryDescription));
-		UpdateParentQuery = Storage->RegisterQuery(MoveTemp(UpdateParentQueryDescription));
+		HierarchyData->UnregisterQueries(*Storage);
+		HierarchyData->RegisterQueries(*Storage, FinalQueryDescription, StaticCastSharedRef<ISceneOutliner>(SceneOutliner->AsShared()), bUsingQueryConditionsSyntax);
 	}
 
 	if (SelectionSet)
@@ -1114,15 +989,11 @@ void FTedsOutlinerImpl::UnregisterQueries()
 	}
 	if (Storage)
 	{
-		Storage->UnregisterQuery(ChildRowHandleQuery);
-		Storage->UnregisterQuery(UpdateParentQuery);
 		Storage->UnregisterQuery(SelectedRowsQuery);
 		Storage->UnregisterQuery(SelectionAddedQuery);
 		Storage->UnregisterQuery(SelectionRemovedQuery);
 		Storage->UnregisterQuery(LabelUpdateQuery);
 		
-		ChildRowHandleQuery = DataStorage::InvalidQueryHandle;
-		UpdateParentQuery = DataStorage::InvalidQueryHandle;
 		SelectedRowsQuery = DataStorage::InvalidQueryHandle;
 		SelectionAddedQuery = DataStorage::InvalidQueryHandle;
 		SelectionRemovedQuery = DataStorage::InvalidQueryHandle;
@@ -1263,10 +1134,6 @@ ISceneOutlinerHierarchy::FHierarchyChangedEvent& FTedsOutlinerImpl::OnHierarchyC
 	return HierarchyChangedEvent;
 }
 
-const TOptional<FTedsOutlinerHierarchyData>& FTedsOutlinerImpl::GetHierarchyData()
-{
-	return HierarchyData;
-}
 } // namespace UE::Editor::Outliner
 
 #undef LOCTEXT_NAMESPACE
diff --git a/Engine/Plugins/Experimental/EditorDataStorageFeatures/Source/TedsOutliner/Private/TedsOutlinerModule.cpp b/Engine/Plugins/Experimental/EditorDataStorageFeatures/Source/TedsOutliner/Private/TedsOutlinerModule.cpp
--- a/Engine/Plugins/Experimental/EditorDataStorageFeatures/Source/TedsOutliner/Private/TedsOutlinerModule.cpp
+++ b/Engine/Plugins/Experimental/EditorDataStorageFeatures/Source/TedsOutliner/Private/TedsOutlinerModule.cpp
@@ -27,7 +27,6 @@
 #include "Modules/ModuleManager.h"
 #include "Compatibility/SceneOutlinerTedsBridge.h"
 #include "Compatibility/SceneOutlinerRowHandleColumn.h"
-#include "Columns/ActorFolderColumns.h"
 #include "Widgets/Docking/SDockTab.h"
 
 #define LOCTEXT_NAMESPACE "TedsOutlinerModule"
@@ -174,8 +173,7 @@ TSharedRef<SWidget> FTedsOutlinerModule::CreateLevelEditorTedsOutliner()
 	const FQueryDescription OutlinerQueryDescription =
 		Select()
 		.Where()
-			.All<FTypedElementUObjectColumn>()
-			.Any<FFolderCompatibilityColumn, FTypedElementWorldColumn>()
+			.All<FTypedElementUObjectColumn, FTypedElementWorldColumn>()
 			.None<FHideRowFromUITag>()
 		.Compile();
 
@@ -189,6 +187,13 @@ TSharedRef<SWidget> FTedsOutlinerModule::CreateLevelEditorTedsOutliner()
 	Params.QueryDescription = OutlinerQueryDescription;
 	Params.ColumnDescription = GetLevelEditorTedsOutlinerColumnDescription();
 	
+	const TArray<FHierarchyHandle> Hierarchies{
+		Storage->FindHierarchyByName(TEXT("ActorEditorHierarchy")),
+		Storage->FindHierarchyByName(TEXT("ObjectComponent"))
+		};
+	
+	Params.HierarchyData = MakeShared<FTedsOutlinerMultiHierarchyInterface>(MakeShared<FHierarchyViewerMultiData>(*Storage, Hierarchies));
+	
 	Params.Filters.Add(MakeShared<FTedsOutlinerFilter>(AActor::StaticClass()));
 	Params.Filters.Add(MakeShared<FTedsSceneOutlinerWorldFilter>(TEXT("WorldFilter")));
 	
diff --git a/Engine/Plugins/Experimental/EditorDataStorageFeatures/Source/TedsOutliner/Public/TedsOutlinerHierarchyInterfaces.h b/Engine/Plugins/Experimental/EditorDataStorageFeatures/Source/TedsOutliner/Public/TedsOutlinerHierarchyInterfaces.h
new file mode 100644
--- /dev/null
+++ b/Engine/Plugins/Experimental/EditorDataStorageFeatures/Source/TedsOutliner/Public/TedsOutlinerHierarchyInterfaces.h
@@ -0,0 +1,123 @@
+// Copyright Epic Games, Inc. All Rights Reserved.
+
+#pragma once
+
+#include "HierarchyViewerIntefaces.h"
+#include "ISceneOutliner.h"
+#include "DataStorage/Handles.h"
+#include "Elements/Columns/TypedElementHiearchyColumns.h"
+
+#define UE_API TEDSOUTLINER_API
+
+namespace UE::Editor::Outliner
+{
+	/* Base class for all hierarchy data interfaces used by the Teds Outliner to display hierarchies in the UI.
+	 * Teds Outliner hierarchies have the additional step of registering queries to detect changes to hierarchies when compared to the hierarchy viewer
+	 */
+	class ITedsOutlinerHierarchyDataInterface : public DataStorage::IHierarchyViewerDataInterface
+	{
+	public:
+		virtual ~ITedsOutlinerHierarchyDataInterface() override = default;
+		
+		/** 
+		 * Register any queries required by this data interface
+		 * @param Storage Reference to the editor data storage interface
+		 * @param OutlinerQueryDescription The query description used to gather rows in the outliner, can be merged with any local queries 
+		 *								   to filter down on the rows observed to only those in the outliner
+		 * @param Outliner Ptr to the Scene Outliner owning this interface
+		 * @param bUsingQueryConditionsSyntax Whether the owning outliner query description is using the new query condition syntax or not
+		 *									  Useful for query merge which requires both input queries to be using the same syntax
+		 */
+		virtual void RegisterQueries(DataStorage::ICoreProvider& Storage, const DataStorage::FQueryDescription& OutlinerQueryDescription, 
+			TWeakPtr<ISceneOutliner> Outliner, bool bUsingQueryConditionsSyntax) = 0;
+		
+		// Unregister any queries required by this data interface
+		virtual void UnregisterQueries(DataStorage::ICoreProvider& Storage) = 0;
+	};
+	
+	// Hierarchy Data interface that can be used to combine multiple hierarchies into one view
+	class FTedsOutlinerMultiHierarchyInterface : public ITedsOutlinerHierarchyDataInterface
+	{
+	public:
+		// @see FHierarchyViewerMultiData, which is used to specify the ordered list of hierarchies
+		UE_API explicit FTedsOutlinerMultiHierarchyInterface(TSharedRef<DataStorage::FHierarchyViewerMultiData> InHierarchyData);
+		
+		UE_API virtual DataStorage::RowHandle GetParent(const DataStorage::ICoreProvider& Storage, DataStorage::RowHandle InRow) const override;
+		
+		UE_API virtual void WalkDepthFirst(const DataStorage::ICoreProvider& Storage, DataStorage::RowHandle InRow, 
+			DataStorage::ICoreProvider::FHierarchyIterationCallback VisitFn,
+			DataStorage::ICoreProvider::ETraversalOrder TraversalOrder = DataStorage::ICoreProvider::ETraversalOrder::PreOrder) const override;
+		
+		UE_API virtual void RegisterQueries(DataStorage::ICoreProvider& Storage, const DataStorage::FQueryDescription& OutlinerQueryDescription,
+			TWeakPtr<ISceneOutliner> Outliner, bool bUsingQueryConditionsSyntax) override;
+		
+		UE_API virtual void UnregisterQueries(DataStorage::ICoreProvider& Storage) override;
+		
+	protected:
+		TSharedRef<DataStorage::FHierarchyViewerMultiData> HierarchyData;
+		TArray<DataStorage::QueryHandle> HierarchyChangeQueries;
+	};
+	
+	/**
+	  * @section Legacy Hierarchy Interfaces
+	  * These classes exist for backward compatibility with existing uses of the Teds Outliner which involve custom hierarchies instead of the new
+	  * TEDS hierarchies feature. It is NOT recommended to use these anymore, and instead to use the hierarchy API on ICoreProvider to set your hierarchy
+	  * alongside the new data interfaces above. 
+	  */
+	
+	// Legacy struct storing information on how hierarchies are handled in the TEDS Outliner
+    struct FTedsOutlinerHierarchyData
+    {
+    	/** A delegate used to get the parent row handle for a given row */
+    	DECLARE_DELEGATE_RetVal_OneParam(DataStorage::RowHandle, FGetParentRowHandle, const void* /* InColumnData */);
+    
+    	/** A delegate used to get the children row handles for a given row */
+    	DECLARE_DELEGATE_RetVal_OneParam(TArrayView<DataStorage::RowHandle>, FGetChildrenRowsHandles, void* /* InColumnData */);
+    	
+    	/** A delegate used to set the parent row handle for a given row */
+    	DECLARE_DELEGATE_TwoParams(FSetParentRowHandle, void* /* InColumnData */, DataStorage::RowHandle /* InParentRowHandle */);
+    
+    	FTedsOutlinerHierarchyData(const UScriptStruct* InHierarchyColumn, const FGetParentRowHandle& InGetParent, const FSetParentRowHandle& InSetParent, const FGetChildrenRowsHandles& InGetChildren)
+    		: HierarchyColumn(InHierarchyColumn)
+    		, GetParent(InGetParent)
+    		, GetChildren(InGetChildren)
+    		, SetParent(InSetParent)
+    	{
+    	
+    	}
+    
+    	// The column that contains the parent row handle for rows
+    	const UScriptStruct* HierarchyColumn;
+    
+    	// Function to get parent row handle
+    	FGetParentRowHandle GetParent;
+    
+    	// Function to get the children row handle
+    	FGetChildrenRowsHandles GetChildren;
+    
+    	// Function to set the parent row handle
+    	FSetParentRowHandle SetParent;
+    	
+    	// Get the default hierarchy data for the TEDS Outliner that uses FTableRowParentColumn to get the parent
+    	static FTedsOutlinerHierarchyData GetDefaultHierarchyData();
+    };
+	
+	// Legacy data interface for displaying custom hierarchies in the Teds Outliner
+	class FTedsOutlinerLegacyHierarchyInterface : public ITedsOutlinerHierarchyDataInterface
+	{
+	public:
+		UE_API explicit FTedsOutlinerLegacyHierarchyInterface(FTedsOutlinerHierarchyData InHierarchyData);
+		UE_API virtual DataStorage::RowHandle GetParent(const DataStorage::ICoreProvider& Storage, DataStorage::RowHandle InRow) const override;
+		UE_API virtual void WalkDepthFirst(const DataStorage::ICoreProvider& Storage, DataStorage::RowHandle InRow, DataStorage::ICoreProvider::FHierarchyIterationCallback VisitFn,
+											DataStorage::ICoreProvider::ETraversalOrder TraversalOrder = DataStorage::ICoreProvider::ETraversalOrder::PreOrder) const override;
+		UE_API virtual void RegisterQueries(DataStorage::ICoreProvider& Storage, const DataStorage::FQueryDescription& OutlinerQueryDescription, TWeakPtr<ISceneOutliner> Outliner, bool bUsingQueryConditionsSyntax) override;
+		UE_API virtual void UnregisterQueries(DataStorage::ICoreProvider& Storage) override;
+		
+	protected:
+		FTedsOutlinerHierarchyData HierarchyData;
+		DataStorage::QueryHandle HierarchyChangeQuery = DataStorage::InvalidQueryHandle;
+		DataStorage::QueryHandle GetChildrenQuery = DataStorage::InvalidQueryHandle;
+	};
+}
+
+#undef UE_API
\ No newline at end of file
diff --git a/Engine/Plugins/Experimental/EditorDataStorageFeatures/Source/TedsOutliner/Public/TedsOutlinerImpl.h b/Engine/Plugins/Experimental/EditorDataStorageFeatures/Source/TedsOutliner/Public/TedsOutlinerImpl.h
--- a/Engine/Plugins/Experimental/EditorDataStorageFeatures/Source/TedsOutliner/Public/TedsOutlinerImpl.h
+++ b/Engine/Plugins/Experimental/EditorDataStorageFeatures/Source/TedsOutliner/Public/TedsOutlinerImpl.h
@@ -12,6 +12,7 @@
 #include "Elements/Framework/TypedElementQueryBuilder.h"
 #include "Elements/Framework/TypedElementSelectionSet.h"
 #include "TedsOutlinerFilter.h"
+#include "TedsOutlinerHierarchyInterfaces.h"
 
 #define UE_API TEDSOUTLINER_API
 
@@ -45,64 +46,6 @@ class SWidget;
 
 namespace UE::Editor::Outliner
 {
-// Struct storing information on how hierarchies are handled in the TEDS Outliner
-struct FTedsOutlinerHierarchyData
-{
-	/** A delegate used to get the parent row handle for a given row */
-	DECLARE_DELEGATE_RetVal_OneParam(DataStorage::RowHandle, FGetParentRowHandle, const void* /* InColumnData */);
-
-	/** A delegate used to get the children row handles for a given row */
-	DECLARE_DELEGATE_RetVal_OneParam(TArrayView<DataStorage::RowHandle>, FGetChildrenRowsHandles, void* /* InColumnData */);
-	
-	/** A delegate used to set the parent row handle for a given row */
-	DECLARE_DELEGATE_TwoParams(FSetParentRowHandle, void* /* InColumnData */, DataStorage::RowHandle /* InParentRowHandle */);
-
-	FTedsOutlinerHierarchyData(const UScriptStruct* InHierarchyColumn, const FGetParentRowHandle& InGetParent, const FSetParentRowHandle& InSetParent, const FGetChildrenRowsHandles& InGetChildren)
-		: HierarchyColumn(InHierarchyColumn)
-		, GetParent(InGetParent)
-		, GetChildren(InGetChildren)
-		, SetParent(InSetParent)
-	{
-	
-	}
-
-	// The column that contains the parent row handle for rows
-	const UScriptStruct* HierarchyColumn;
-
-	// Function to get parent row handle
-	FGetParentRowHandle GetParent;
-
-	// Function to get the children row handle
-	FGetChildrenRowsHandles GetChildren;
-
-	// Function to set the parent row handle
-	FSetParentRowHandle SetParent;
-	
-	// Get the default hierarchy data for the TEDS Outliner that uses FTableRowParentColumn to get the parent
-	static FTedsOutlinerHierarchyData GetDefaultHierarchyData()
-	{
-		const FGetParentRowHandle RowHandleGetter = FGetParentRowHandle::CreateLambda([](const void* InColumnData)
-			{
-				if(const FTableRowParentColumn* ParentColumn = static_cast<const FTableRowParentColumn *>(InColumnData))
-				{
-					return ParentColumn->Parent;
-				}
-
-				return DataStorage::InvalidRowHandle;
-			});
-
-		const FSetParentRowHandle RowHandleSetter = FSetParentRowHandle::CreateLambda([](void* InColumnData,
-			DataStorage::RowHandle InRowHandle)
-			{
-				if(FTableRowParentColumn* ParentColumn = static_cast<FTableRowParentColumn *>(InColumnData))
-				{
-					ParentColumn->Parent = InRowHandle;
-				}
-			});
-		
-		return FTedsOutlinerHierarchyData(FTableRowParentColumn::StaticStruct(), RowHandleGetter, RowHandleSetter, FGetChildrenRowsHandles());
-	}
-};
 	
 struct FTedsOutlinerColumnDescription
 {
@@ -153,7 +96,7 @@ struct FTedsOutlinerParams
 
 	// If specified, this is how the TEDS Outliner will handle hierarchies. If not specified - there will be no hierarchies shown as a
 	// parent-child relation in the tree view
-	TOptional<FTedsOutlinerHierarchyData> HierarchyData;
+	TSharedPtr<ITedsOutlinerHierarchyDataInterface> HierarchyData;
 
 	// The selection set to use for this Outliner, nullptr = don't propagate tree selection through TypedElementSelectionInterface
 	UTypedElementSelectionSet* SelectionSet;
@@ -218,9 +161,6 @@ public:
 	// Helper function to create a label widget for a given row
 	UE_API TSharedRef<SWidget> CreateLabelWidgetForItem(DataStorage::RowHandle TargetRow, ISceneOutlinerTreeItem& TreeItem, const STableRow<FSceneOutlinerTreeItemPtr>& RowItem, bool bIsInteractable = true) const;
 
-	// Get the hierarchy data associated with this table viewer
-	UE_API const TOptional<FTedsOutlinerHierarchyData>& GetHierarchyData();
-	
 	// Add an external query description to the Outliner
 	UE_API void AddExternalQuery(const FName& QueryName, const QueryHandle& InQueryHandle);
 	UE_API void RemoveExternalQuery(const FName& QueryName);
@@ -319,7 +259,7 @@ protected:
 	TArray<TSharedPtr<DataStorage::QueryStack::FRowFilterNode>> ClassFilterNodes;
 
 	// Optional Hierarchy Data
-	TOptional<FTedsOutlinerHierarchyData> HierarchyData;
+	TSharedPtr<ITedsOutlinerHierarchyDataInterface> HierarchyData;
 	// The query stack node responsible for collecting all rows that match the initial query on FullRefresh()
 	TSharedPtr<DataStorage::QueryStack::FRowQueryResultsNode> InitialRowCollectorNode;
 	// The query stack node responsible for collecting all rows that match the composite query on FullRefresh()
@@ -350,12 +290,6 @@ protected:
 	// The executor used to update the query stack stack.
 	TUniquePtr<DataStorage::QueryStack::FExplicitUpdateExecutor> QueryStackExecutor;
 
-	// Query to get all child rows
-	DataStorage::QueryHandle ChildRowHandleQuery = DataStorage::InvalidQueryHandle;
-
-	// Query to track when a row's parent gets changed
-	DataStorage::QueryHandle UpdateParentQuery = DataStorage::InvalidQueryHandle;
-
 	// Query to get all selected rows, track selection added, track selection removed
 	DataStorage::QueryHandle SelectedRowsQuery = DataStorage::InvalidQueryHandle;
 	DataStorage::QueryHandle SelectionAddedQuery = DataStorage::InvalidQueryHandle;
diff --git a/Engine/Plugins/Experimental/EditorDataStorageFeatures/Source/TedsOutliner/TedsOutliner.Build.cs b/Engine/Plugins/Experimental/EditorDataStorageFeatures/Source/TedsOutliner/TedsOutliner.Build.cs
--- a/Engine/Plugins/Experimental/EditorDataStorageFeatures/Source/TedsOutliner/TedsOutliner.Build.cs
+++ b/Engine/Plugins/Experimental/EditorDataStorageFeatures/Source/TedsOutliner/TedsOutliner.Build.cs
@@ -26,6 +26,7 @@ public class TedsOutliner : ModuleRules
 					"SceneOutliner",
 					"Slate",
 					"SlateCore",
+					"TedsTableViewer",
 					"TypedElementFramework",
 					"TypedElementRuntime"
 				});
diff --git a/Engine/Plugins/Experimental/EditorDataStorageFeatures/Source/TedsTableViewer/Private/HierarchyViewerIntefaces.cpp b/Engine/Plugins/Experimental/EditorDataStorageFeatures/Source/TedsTableViewer/Private/HierarchyViewerIntefaces.cpp
--- a/Engine/Plugins/Experimental/EditorDataStorageFeatures/Source/TedsTableViewer/Private/HierarchyViewerIntefaces.cpp
+++ b/Engine/Plugins/Experimental/EditorDataStorageFeatures/Source/TedsTableViewer/Private/HierarchyViewerIntefaces.cpp
@@ -32,6 +32,11 @@ namespace UE::Editor::DataStorage
 		Storage.WalkDepthFirst(HierarchyHandle, InRow, VisitFn, TraversalOrder);
 	}
 
+	FHierarchyHandle FHierarchyViewerData::GetHierarchy() const
+	{
+		return HierarchyHandle;
+	}
+
 	void FHierarchyViewerData::RegisterQueries(ICoreProvider& Storage)
 	{
 		using namespace UE::Editor::DataStorage::Queries;
@@ -114,4 +119,12 @@ namespace UE::Editor::DataStorage
 			}, TraversalOrder);
 		}
 	}
+
+	void FHierarchyViewerMultiData::ForEachHierarchy(TFunctionRef<void(FHierarchyHandle)> VisitFn) const
+	{
+		for (const TSharedPtr<FHierarchyViewerData>& Data : Hierarchies)
+		{
+			VisitFn(Data->GetHierarchy());
+		}
+	}
 }
diff --git a/Engine/Plugins/Experimental/EditorDataStorageFeatures/Source/TedsTableViewer/Public/HierarchyViewerIntefaces.h b/Engine/Plugins/Experimental/EditorDataStorageFeatures/Source/TedsTableViewer/Public/HierarchyViewerIntefaces.h
--- a/Engine/Plugins/Experimental/EditorDataStorageFeatures/Source/TedsTableViewer/Public/HierarchyViewerIntefaces.h
+++ b/Engine/Plugins/Experimental/EditorDataStorageFeatures/Source/TedsTableViewer/Public/HierarchyViewerIntefaces.h
@@ -42,6 +42,7 @@ namespace UE::Editor::DataStorage
 		UE_API virtual RowHandle GetParent(const ICoreProvider& Storage, RowHandle InRow) const override;
 		UE_API virtual void WalkDepthFirst(const ICoreProvider& Storage, RowHandle InRow, ICoreProvider::FHierarchyIterationCallback VisitFn,
 			ICoreProvider::ETraversalOrder TraversalOrder = ICoreProvider::ETraversalOrder::PreOrder) const override;
+		UE_API FHierarchyHandle GetHierarchy() const;
 
 	protected:
 		void RegisterQueries(ICoreProvider& Storage);
@@ -64,6 +65,8 @@ namespace UE::Editor::DataStorage
 		UE_API virtual void WalkDepthFirst(const ICoreProvider& Storage, RowHandle InRow, ICoreProvider::FHierarchyIterationCallback VisitFn,
 			ICoreProvider::ETraversalOrder TraversalOrder = ICoreProvider::ETraversalOrder::PreOrder) const override;
 		
+		UE_API void ForEachHierarchy(TFunctionRef<void(FHierarchyHandle)> VisitFn) const;
+		
 	protected:
 		TArray<TSharedPtr<FHierarchyViewerData>> Hierarchies;
 	};
