TEDS Outliner Hierarchies:

1. Inside GetParent(), find the first parent that matches the initial query by iterating over all parents. Without this change when you are combining multiple hierarchies the outliner would stop at the first parent that has a valid row handle, even if that parent is not displayable by the outliner. This would lead to later parents that could be valid getting ignored.
2. Fix various instances of the UI row -> Outliner widget row hierarchy being applied to the data row incorrectly

#Synced-CL 51413783

diff --git a/Engine/Plugins/Experimental/EditorDataStorageFeatures/Source/TedsOutliner/Private/TedsOutlinerHierarchyInterfaces.cpp b/Engine/Plugins/Experimental/EditorDataStorageFeatures/Source/TedsOutliner/Private/TedsOutlinerHierarchyInterfaces.cpp
--- a/Engine/Plugins/Experimental/EditorDataStorageFeatures/Source/TedsOutliner/Private/TedsOutlinerHierarchyInterfaces.cpp
+++ b/Engine/Plugins/Experimental/EditorDataStorageFeatures/Source/TedsOutliner/Private/TedsOutlinerHierarchyInterfaces.cpp
@@ -44,7 +44,21 @@ namespace UE::Editor::Outliner
 	}
 
 	// @section FTedsOutlinerMultiHierarchyInterface
-	
+
+	void ITedsOutlinerHierarchyDataInterface::ForEachImmediateParent(const DataStorage::ICoreProvider& Storage, DataStorage::RowHandle InRow,
+		FParentIterationCallback Callback)
+	{
+		if (Callback.IsSet())
+		{
+			RowHandle ParentRow = GetParent(Storage, InRow);
+		
+			if (Storage.IsRowAvailable(ParentRow))
+			{
+				Callback(Storage, ParentRow);
+			}
+		}
+	}
+
 	FTedsOutlinerMultiHierarchyInterface::FTedsOutlinerMultiHierarchyInterface(TSharedRef<DataStorage::FHierarchyViewerMultiData> InHierarchyData)
 		: HierarchyData(InHierarchyData)
 	{
@@ -67,9 +81,9 @@ namespace UE::Editor::Outliner
 		const DataStorage::FQueryDescription& OutlinerQueryDescription, TWeakPtr<ISceneOutliner> Outliner, bool bUsingQueryConditionsSyntax)
 	{
 		// For each hierarchy we are displaying, get the parent change column (if present) and add a processor to detect changes to rows with the change column
-		HierarchyData->ForEachHierarchy([this, &Storage](FHierarchyHandle HierarchyHandle)
+		HierarchyData->ForEachHierarchyData([this, &Storage](const TSharedRef<FHierarchyViewerData>& HierarchyViewerData)
 		{
-			if (const UScriptStruct* ParentChangeTag = Storage.GetParentChangedColumnType(HierarchyHandle))
+			if (const UScriptStruct* ParentChangeTag = Storage.GetParentChangedColumnType(HierarchyViewerData->GetHierarchy()))
 			{
 				FString ProcessorName = TEXT("FTedsOutlinerMultiHierarchyInterface: Update Parent for ");
 				ParentChangeTag->GetFName().AppendString(ProcessorName);
@@ -95,6 +109,8 @@ namespace UE::Editor::Outliner
 				
 				HierarchyChangeQueries.Add(Storage.RegisterQuery(MoveTemp(UpdateParentQueryDescription)));
 			}
+
+			return true;
 		});
 	}
 
@@ -106,7 +122,26 @@ namespace UE::Editor::Outliner
 		}
 		HierarchyChangeQueries.Empty();
 	}
-	
+
+	void FTedsOutlinerMultiHierarchyInterface::ForEachImmediateParent(const DataStorage::ICoreProvider& Storage, DataStorage::RowHandle InRow,
+		FParentIterationCallback Callback)
+	{
+		if (Callback.IsSet())
+		{
+			HierarchyData->ForEachHierarchyData([InRow, Callback, &Storage](const TSharedRef<FHierarchyViewerData>& HierarchyViewerData) -> bool
+			{
+				RowHandle ParentRow = HierarchyViewerData->GetParent(Storage, InRow);
+				
+				if (Storage.IsRowAvailable(ParentRow))
+				{
+					return Callback(Storage, ParentRow);
+				}
+
+				return true;
+			});
+		}
+	}
+
 	// @section FTedsOutlinerHierarchyData
 	
 	FTedsOutlinerHierarchyData FTedsOutlinerHierarchyData::GetDefaultHierarchyData()
diff --git a/Engine/Plugins/Experimental/EditorDataStorageFeatures/Source/TedsOutliner/Private/TedsOutlinerImpl.cpp b/Engine/Plugins/Experimental/EditorDataStorageFeatures/Source/TedsOutliner/Private/TedsOutlinerImpl.cpp
--- a/Engine/Plugins/Experimental/EditorDataStorageFeatures/Source/TedsOutliner/Private/TedsOutlinerImpl.cpp
+++ b/Engine/Plugins/Experimental/EditorDataStorageFeatures/Source/TedsOutliner/Private/TedsOutlinerImpl.cpp
@@ -114,6 +114,7 @@ namespace LabelWidgetUtils
 		DataStorage::ICoreProvider& Storage,
 		DataStorage::IUiProvider& StorageUi,
 		const DataStorage::RowHandle TargetRow,
+		const DataStorage::RowHandle OutlinerWidgetRow,
 		ISceneOutlinerTreeItem& TreeItem,
 		const STableRow<FSceneOutlinerTreeItemPtr>& RowItem,
 		const bool bInteractable)
@@ -128,7 +129,7 @@ namespace LabelWidgetUtils
 		checkf(WidgetRow != InvalidRowHandle, TEXT("Expected valid row handle!"));
 
 		FHierarchyHandle Hierarchy = Storage.FindHierarchyByName(TableViewerUtils::GetWidgetHierarchyName());
-		Storage.SetParentRow(Hierarchy, WidgetRow, TargetRow);
+		Storage.SetParentRow(Hierarchy, WidgetRow, OutlinerWidgetRow);
 
 		if (FTypedElementRowReferenceColumn* RowReference = Storage.GetColumn<FTypedElementRowReferenceColumn>(WidgetRow))
 		{
@@ -465,6 +466,7 @@ TSharedRef<SWidget> FTedsOutlinerImpl::CreateLabelWidget(
 	DataStorage::IUiProvider& StorageUi,
 	const DataStorage::IUiProvider::FPurposeID& Purpose,
 	DataStorage::RowHandle TargetRow,
+	const DataStorage::RowHandle OutlinerWidgetRow,
 	ISceneOutlinerTreeItem& TreeItem,
 	const STableRow<FSceneOutlinerTreeItemPtr>& RowItem,
 	bool bInteractable)
@@ -475,7 +477,7 @@ TSharedRef<SWidget> FTedsOutlinerImpl::CreateLabelWidget(
 		return SNullWidget::NullWidget;
 	}
 
-	TSharedPtr<SWidget> Widget = LabelWidgetUtils::CreateLabelWidgetInternal(*WidgetConstructor, Storage, StorageUi, TargetRow, TreeItem, RowItem, bInteractable);
+	TSharedPtr<SWidget> Widget = LabelWidgetUtils::CreateLabelWidgetInternal(*WidgetConstructor, Storage, StorageUi, TargetRow, OutlinerWidgetRow, TreeItem, RowItem, bInteractable);
 	if (!Widget)
 	{
 		return SNullWidget::NullWidget;
@@ -491,7 +493,7 @@ TSharedRef<SWidget> FTedsOutlinerImpl::CreateLabelWidget(
 
 TSharedRef<SWidget>FTedsOutlinerImpl::CreateLabelWidgetForItem(DataStorage::RowHandle TargetRow, ISceneOutlinerTreeItem& TreeItem, const STableRow<FSceneOutlinerTreeItemPtr>& RowItem, bool bIsInteractable) const
 {
-	return CreateLabelWidget(*Storage, *StorageUi, LabelWidgetPurpose, TargetRow, TreeItem, RowItem, bIsInteractable);
+	return CreateLabelWidget(*Storage, *StorageUi, LabelWidgetPurpose, TargetRow, GetOutlinerRowHandle(), TreeItem, RowItem, bIsInteractable);
 }
 
 void FTedsOutlinerImpl::AddExternalQuery(const FName& QueryName, const QueryHandle& InQueryHandle)
@@ -809,17 +811,40 @@ DataStorage::RowHandle FTedsOutlinerImpl::GetParentRow(DataStorage::RowHandle In
 		return DataStorage::InvalidRowHandle;
 	}
 	
-	const DataStorage::RowHandle ParentRowHandle = HierarchyData->GetParent(*Storage, InRowHandle);
+	RowHandle ParentRowHandle = DataStorage::InvalidRowHandle;
+	DataStorage::FRowHandleArrayView Rows;
+
+	if (ShouldForceShowParentRows())
+	{
+		// We want to only include parents that match the initial query but not any additional filters. This essentially acts as 
+		// a way to force show parents even if they don't match the current filter but their child does
 	
-	if (!Storage->IsRowAvailable(ParentRowHandle))
+		// Update this on demand, if the results are already sorted no work is done
+		InitialQueryResultsSortedNodeExecutor->Update();
+	
+		Rows = InitialQueryResultsSortedNode->GetRows();
+	}
+	else
 	{
-		return DataStorage::InvalidRowHandle;
+		// If we aren't force showing parents, make sure the parent also matches the final query with all additional filters included
+		Rows = RowHandleSortNode->GetRows();
 	}
-
-	// If the row doesn't exist in our query stack results, it doesn't match the query and should not be displayed
-	DataStorage::FRowHandleArrayView Rows(RowHandleSortNode->GetRows());
-
-	if (!Rows.Contains(ParentRowHandle))
+	
+	// Iterate through each parent in case there are multiple hierarchies, so that we don't stop if the first parent doesn't match our condition and
+	// instead keep looking
+	HierarchyData->ForEachImmediateParent(*Storage, InRowHandle, [&ParentRowHandle, &Rows](const ICoreProvider& Storage, RowHandle ParentRow)
+		{
+			// The results are sorted by row handle so a binary search is done here
+			if (Rows.Contains(ParentRow))
+			{
+				ParentRowHandle = ParentRow;
+				return false; // We have found a valid parent, no need to continue iteration
+			}
+			
+			return true;
+		});
+	
+	if (!Storage->IsRowAvailable(ParentRowHandle))
 	{
 		return DataStorage::InvalidRowHandle;
 	}
@@ -865,6 +890,17 @@ void FTedsOutlinerImpl::RecompileQueries()
 
 	// Node to collect all rows matching the initial query (stored so OR filters can filter on the original query)
 	InitialRowCollectorNode = MakeShared<FRowQueryResultsNode>(*Storage, InitialQueryNode, SyncActions);
+	
+	// We store a copy of the collected rows into an array sorted by row handles for faster lookup
+	TSharedRef<FRowCopyNode> InitialCopyNode = MakeShared<FRowCopyNode>(InitialRowCollectorNode);
+	InitialQueryResultsSortedNode = MakeShared<FRowHandleSortNode>(InitialCopyNode);
+	
+	static int InstanceCounter = 0;
+	
+	// InstanceCounter is incremented later on
+	InitialQueryResultsSortedNodeExecutor = 
+		MakeUnique<QueryStack::FExplicitUpdateExecutor>(FName("TedsOutlinerInitialResults", InstanceCounter++), InitialQueryResultsSortedNode);
+
 
 	// Node to collect all rows matching the composite query
 	CombinedRowCollectorNode = MakeShared<FRowQueryResultsNode>(*Storage, CombinedQueryNode, SyncActions);
@@ -921,7 +957,6 @@ void FTedsOutlinerImpl::RecompileQueries()
 
 	// A custom node to store row handle -> sort index mapping.
 	RowSortMapNode = MakeShared<FRowMapNode>(RowPrimaryInversionNode);
-	static int InstanceCounter = 0;
 	QueryStackExecutor = MakeUnique<QueryStack::FExplicitUpdateExecutor>(FName("TedsOutliner", InstanceCounter++), RowSortMapNode);
 
 	// Our final query to collect rows to populate the Outliner - currently the same as the initial query the user provided
diff --git a/Engine/Plugins/Experimental/EditorDataStorageFeatures/Source/TedsOutliner/Private/TedsOutlinerItem.cpp b/Engine/Plugins/Experimental/EditorDataStorageFeatures/Source/TedsOutliner/Private/TedsOutlinerItem.cpp
--- a/Engine/Plugins/Experimental/EditorDataStorageFeatures/Source/TedsOutliner/Private/TedsOutlinerItem.cpp
+++ b/Engine/Plugins/Experimental/EditorDataStorageFeatures/Source/TedsOutliner/Private/TedsOutlinerItem.cpp
@@ -73,10 +73,6 @@ TSharedRef<SWidget> FTedsOutlinerTreeItem::GenerateLabelWidget(ISceneOutliner& O
 {
 	if (TSharedPtr TedsOutlinerImplPin = TedsOutlinerImpl.Pin())
 	{
-		ICoreProvider* Storage = TedsOutlinerImplPin->GetStorage();
-		FHierarchyHandle Hierarchy = Storage->FindHierarchyByName(TableViewerUtils::GetWidgetHierarchyName());
-		Storage->SetParentRow(Hierarchy, RowHandle, TedsOutlinerImplPin->GetOutlinerRowHandle());
-
 		return TedsOutlinerImplPin->CreateLabelWidgetForItem(RowHandle, *this, InRow, CanInteract());
 	}
 	return SNullWidget::NullWidget;
diff --git a/Engine/Plugins/Experimental/EditorDataStorageFeatures/Source/TedsOutliner/Public/TedsOutlinerHierarchyInterfaces.h b/Engine/Plugins/Experimental/EditorDataStorageFeatures/Source/TedsOutliner/Public/TedsOutlinerHierarchyInterfaces.h
--- a/Engine/Plugins/Experimental/EditorDataStorageFeatures/Source/TedsOutliner/Public/TedsOutlinerHierarchyInterfaces.h
+++ b/Engine/Plugins/Experimental/EditorDataStorageFeatures/Source/TedsOutliner/Public/TedsOutlinerHierarchyInterfaces.h
@@ -33,6 +33,12 @@ namespace UE::Editor::Outliner
 		
 		// Unregister any queries required by this data interface
 		virtual void UnregisterQueries(DataStorage::ICoreProvider& Storage) = 0;
+		
+		using FParentIterationCallback = TFunction<bool(const DataStorage::ICoreProvider& Context, DataStorage::RowHandle Parent)>;
+		
+		// Iterate over the immediate parents of InRow. In the simple case this is just one parent, but if this is a complex hierarchy combining multiple
+		// individual ones there could be multiple immediate parents
+		virtual void ForEachImmediateParent(const DataStorage::ICoreProvider& Storage, DataStorage::RowHandle InRow, FParentIterationCallback Callback);
 	};
 	
 	// Hierarchy Data interface that can be used to combine multiple hierarchies into one view
@@ -53,6 +59,8 @@ namespace UE::Editor::Outliner
 		
 		UE_API virtual void UnregisterQueries(DataStorage::ICoreProvider& Storage) override;
 		
+		UE_API virtual void ForEachImmediateParent(const DataStorage::ICoreProvider& Storage, DataStorage::RowHandle InRow, FParentIterationCallback Callback) override;
+		
 	protected:
 		TSharedRef<DataStorage::FHierarchyViewerMultiData> HierarchyData;
 		TArray<DataStorage::QueryHandle> HierarchyChangeQueries;
diff --git a/Engine/Plugins/Experimental/EditorDataStorageFeatures/Source/TedsOutliner/Public/TedsOutlinerImpl.h b/Engine/Plugins/Experimental/EditorDataStorageFeatures/Source/TedsOutliner/Public/TedsOutlinerImpl.h
--- a/Engine/Plugins/Experimental/EditorDataStorageFeatures/Source/TedsOutliner/Public/TedsOutlinerImpl.h
+++ b/Engine/Plugins/Experimental/EditorDataStorageFeatures/Source/TedsOutliner/Public/TedsOutlinerImpl.h
@@ -125,6 +125,7 @@ public:
 		DataStorage::IUiProvider& StorageUi,
 		const DataStorage::IUiProvider::FPurposeID& Purpose,
 		DataStorage::RowHandle TargetRow,
+		const DataStorage::RowHandle OutlinerWidgetRow,
 		ISceneOutlinerTreeItem& TreeItem,
 		const STableRow<FSceneOutlinerTreeItemPtr>& RowItem,
 		bool bInteractable = true);
@@ -289,6 +290,12 @@ protected:
 
 	// The executor used to update the query stack stack.
 	TUniquePtr<DataStorage::QueryStack::FExplicitUpdateExecutor> QueryStackExecutor;
+	
+	// A query stack node that keeps a sorted copy of the initial query results (i.e without any filters or search) for fast lookup
+	TSharedPtr<DataStorage::QueryStack::FRowHandleSortNode> InitialQueryResultsSortedNode;
+	
+	// The executor used to update the node storing a copy of the initial query results. This is updated on demand and not every tick
+	TUniquePtr<DataStorage::QueryStack::FExplicitUpdateExecutor> InitialQueryResultsSortedNodeExecutor;
 
 	// Query to get all selected rows, track selection added, track selection removed
 	DataStorage::QueryHandle SelectedRowsQuery = DataStorage::InvalidQueryHandle;
diff --git a/Engine/Plugins/Experimental/EditorDataStorageFeatures/Source/TedsTableViewer/Private/HierarchyViewerIntefaces.cpp b/Engine/Plugins/Experimental/EditorDataStorageFeatures/Source/TedsTableViewer/Private/HierarchyViewerIntefaces.cpp
--- a/Engine/Plugins/Experimental/EditorDataStorageFeatures/Source/TedsTableViewer/Private/HierarchyViewerIntefaces.cpp
+++ b/Engine/Plugins/Experimental/EditorDataStorageFeatures/Source/TedsTableViewer/Private/HierarchyViewerIntefaces.cpp
@@ -120,11 +120,17 @@ namespace UE::Editor::DataStorage
 		}
 	}
 
-	void FHierarchyViewerMultiData::ForEachHierarchy(TFunctionRef<void(FHierarchyHandle)> VisitFn) const
+	void FHierarchyViewerMultiData::ForEachHierarchyData(TFunctionRef<bool(const TSharedRef<FHierarchyViewerData>&)> VisitFn) const
 	{
 		for (const TSharedPtr<FHierarchyViewerData>& Data : Hierarchies)
 		{
-			VisitFn(Data->GetHierarchy());
+			if (Data)
+			{
+				if (!VisitFn(Data.ToSharedRef()))
+				{
+					return;
+				}
+			}
 		}
 	}
 }
diff --git a/Engine/Plugins/Experimental/EditorDataStorageFeatures/Source/TedsTableViewer/Public/HierarchyViewerIntefaces.h b/Engine/Plugins/Experimental/EditorDataStorageFeatures/Source/TedsTableViewer/Public/HierarchyViewerIntefaces.h
--- a/Engine/Plugins/Experimental/EditorDataStorageFeatures/Source/TedsTableViewer/Public/HierarchyViewerIntefaces.h
+++ b/Engine/Plugins/Experimental/EditorDataStorageFeatures/Source/TedsTableViewer/Public/HierarchyViewerIntefaces.h
@@ -65,7 +65,8 @@ namespace UE::Editor::DataStorage
 		UE_API virtual void WalkDepthFirst(const ICoreProvider& Storage, RowHandle InRow, ICoreProvider::FHierarchyIterationCallback VisitFn,
 			ICoreProvider::ETraversalOrder TraversalOrder = ICoreProvider::ETraversalOrder::PreOrder) const override;
 		
-		UE_API void ForEachHierarchy(TFunctionRef<void(FHierarchyHandle)> VisitFn) const;
+		// Iterate over each contained hierarchy data. The callback can return false to terminate iteration early
+		UE_API void ForEachHierarchyData(TFunctionRef<bool(const TSharedRef<FHierarchyViewerData>&)> VisitFn) const;
 		
 	protected:
 		TArray<TSharedPtr<FHierarchyViewerData>> Hierarchies;
