TEdsOutlinerHierarchyInterface: Fix missing export

#Synced-CL 51413783

diff --git a/Engine/Plugins/Experimental/EditorDataStorageFeatures/Source/TedsOutliner/Public/TedsOutlinerHierarchyInterfaces.h b/Engine/Plugins/Experimental/EditorDataStorageFeatures/Source/TedsOutliner/Public/TedsOutlinerHierarchyInterfaces.h
--- a/Engine/Plugins/Experimental/EditorDataStorageFeatures/Source/TedsOutliner/Public/TedsOutlinerHierarchyInterfaces.h
+++ b/Engine/Plugins/Experimental/EditorDataStorageFeatures/Source/TedsOutliner/Public/TedsOutlinerHierarchyInterfaces.h
@@ -28,17 +28,17 @@ namespace UE::Editor::Outliner
 		 * @param bUsingQueryConditionsSyntax Whether the owning outliner query description is using the new query condition syntax or not
 		 *									  Useful for query merge which requires both input queries to be using the same syntax
 		 */
-		virtual void RegisterQueries(DataStorage::ICoreProvider& Storage, const DataStorage::FQueryDescription& OutlinerQueryDescription, 
+		UE_API virtual void RegisterQueries(DataStorage::ICoreProvider& Storage, const DataStorage::FQueryDescription& OutlinerQueryDescription, 
 			TWeakPtr<ISceneOutliner> Outliner, bool bUsingQueryConditionsSyntax) = 0;
 		
 		// Unregister any queries required by this data interface
-		virtual void UnregisterQueries(DataStorage::ICoreProvider& Storage) = 0;
+		UE_API virtual void UnregisterQueries(DataStorage::ICoreProvider& Storage) = 0;
 		
 		using FParentIterationCallback = TFunction<bool(const DataStorage::ICoreProvider& Context, DataStorage::RowHandle Parent)>;
 		
 		// Iterate over the immediate parents of InRow. In the simple case this is just one parent, but if this is a complex hierarchy combining multiple
 		// individual ones there could be multiple immediate parents
-		virtual void ForEachImmediateParent(const DataStorage::ICoreProvider& Storage, DataStorage::RowHandle InRow, FParentIterationCallback Callback);
+		UE_API virtual void ForEachImmediateParent(const DataStorage::ICoreProvider& Storage, DataStorage::RowHandle InRow, FParentIterationCallback Callback);
 	};
 	
 	// Hierarchy Data interface that can be used to combine multiple hierarchies into one view
