TedsSettings : perform incremental registration of auto-discovered UDeveloperSettings with the ISettingsModule.

The legacy registration of UDeveloperSettings in the ISettingsEditorModule is disabled by using the SetShouldRegisterSettingCallback when TedsSettings is enabled.

This incremental registration registers or unregisters only the CDOs for the reflected UDeveloperSettings classes in the specific module that has been loaded or unloaded in the OnModulesChanged handler.  This provides a more efficient algorithm with performance benefits over the legacy registration in the ISettingsEditorModule which unregisters all previously registered settings and then re-registers all reflected settings from all modules every time a new module is loaded or unloaded.

Another benefit to this eager registration of settings objects is that it makes them available in the ISettingsModule and editor data storage as soon as their module has loaded.  In comparison, the legacy registration is a lazy mechanism triggered by certain user actions such as opening the Editor Preferences tab.  It is more useful to tooling built on editor data storage to have this data eagerly available.

#Synced-CL 51439317

diff --git a/Engine/Plugins/Experimental/EditorDataStorageFeatures/Source/TedsSettings/Private/TedsSettingsManager.cpp b/Engine/Plugins/Experimental/EditorDataStorageFeatures/Source/TedsSettings/Private/TedsSettingsManager.cpp
--- a/Engine/Plugins/Experimental/EditorDataStorageFeatures/Source/TedsSettings/Private/TedsSettingsManager.cpp
+++ b/Engine/Plugins/Experimental/EditorDataStorageFeatures/Source/TedsSettings/Private/TedsSettingsManager.cpp
@@ -10,15 +10,19 @@
 #include "Elements/Columns/TypedElementTypeInfoColumns.h"
 #include "Elements/Framework/TypedElementQueryBuilder.h"
 #include "Elements/Interfaces/TypedElementDataStorageInterface.h"
+#include "Engine/DeveloperSettings.h"
 #include "ISettingsCategory.h"
 #include "ISettingsContainer.h"
+#include "ISettingsEditorModule.h"
 #include "ISettingsModule.h"
 #include "ISettingsSection.h"
 #include "Logging/LogMacros.h"
+#include "Misc/PackageName.h"
 #include "Modules/ModuleManager.h"
 #include "ProfilingDebugging/CpuProfilerTrace.h"
 #include "TedsSettingsColumns.h"
 #include "TedsSettingsLog.h"
+#include "UObject/Package.h"
 #include "UObject/UnrealType.h"
 
 namespace UE::Editor::Settings::Private
@@ -61,6 +65,7 @@ FTedsSettingsManager::FTedsSettingsManager()
 	, SettingsCategoryTable{ UE::Editor::DataStorage::InvalidTableHandle }
 	, SettingsPropertyTable{ UE::Editor::DataStorage::InvalidTableHandle }
 	, SettingsHierarchy{ }
+	, SelectAllSettingsInModuleQueryMap{ }
 {
 }
 
@@ -77,7 +82,33 @@ void FTedsSettingsManager::Initialize()
 				RegisterTables(*DataStorage);
 				RegisterHierarchies(*DataStorage);
 				RegisterQueries(*DataStorage);
-				RegisterActiveSettings();
+				AddSettings();
+
+				////////////////////////////////////////////////////////////////////////////////////////////////////
+				// Take ownership of auto-discovered settings (UDeveloperSettings) registration from the SettingsEditorModule.
+				// First, force the SettingsEditorModule to register any pending settings from modules that may have already loaded.
+				// Then, disable the SettingsEditorModule's ability to register settings so that we can take over with incremental registration.
+				ISettingsModule& SettingsModule = FModuleManager::LoadModuleChecked<ISettingsModule>("Settings");
+				ISettingsEditorModule& SettingsEditorModule = FModuleManager::LoadModuleChecked<ISettingsEditorModule>("SettingsEditor");
+
+				// During the call to SettingsEditorModule.UpdateSettings we want to register any pending settings
+				// but we don't want the SettingsEditorModule to remember those (if it did then it might try to unregister them while we're in charge)
+				SettingsEditorModule.SetShouldRegisterSettingCallback(FShouldRegisterSettingsDelegate::CreateLambda([this, &SettingsModule](UDeveloperSettings* Settings)
+					{
+						RegisterSettings(SettingsModule, Settings);
+						return false;
+					}));
+
+				constexpr bool bForce = true;
+				SettingsEditorModule.UpdateSettings(bForce);
+
+				SettingsEditorModule.SetShouldRegisterSettingCallback(FShouldRegisterSettingsDelegate::CreateLambda([](UDeveloperSettings* Settings)
+					{
+						return false;
+					}));
+
+				FModuleManager::Get().OnModulesChanged().AddSP(this, &FTedsSettingsManager::OnModulesChanged);
+				////////////////////////////////////////////////////////////////////////////////////////////////////
 			};
 
 		if (AreEditorDataStorageFeaturesEnabled())
@@ -102,12 +133,19 @@ void FTedsSettingsManager::Shutdown()
 
 		if (AreEditorDataStorageFeaturesEnabled())
 		{
+			FModuleManager::Get().OnModulesChanged().RemoveAll(this);
+
+			if (ISettingsEditorModule* SettingsEditorModule = FModuleManager::GetModulePtr<ISettingsEditorModule>("SettingsEditor"))
+			{
+				SettingsEditorModule->SetShouldRegisterSettingCallback(nullptr);
+			}
+
 			ICoreProvider* DataStorage = GetMutableDataStorageFeature<ICoreProvider>(StorageFeatureName);
 			check(DataStorage);
 
-			UnregisterActiveSettings();
+			RemoveSettings();
 			UnregisterQueries(*DataStorage);
-		}
+		}	
 
 		bIsInitialized = false;
 	}
@@ -278,44 +316,50 @@ void FTedsSettingsManager::RegisterQueries(UE::Editor::DataStorage::ICoreProvide
 
 void FTedsSettingsManager::UnregisterQueries(UE::Editor::DataStorage::ICoreProvider& DataStorage)
 {
+	using namespace UE::Editor::DataStorage;
+
 	DataStorage.UnregisterQuery(SelectAllSettingsQuery);
-	SelectAllSettingsQuery = UE::Editor::DataStorage::InvalidQueryHandle;
+	SelectAllSettingsQuery = InvalidQueryHandle;
+
+	for (const TPair<FName, QueryHandle>& Pairs : SelectAllSettingsInModuleQueryMap)
+	{
+		DataStorage.UnregisterQuery(Pairs.Value);
+	}
+	SelectAllSettingsInModuleQueryMap.Reset();
 }
 
-void FTedsSettingsManager::RegisterActiveSettings()
+void FTedsSettingsManager::AddSettings()
 {
-	TRACE_CPUPROFILER_EVENT_SCOPE(TedsSettingsManager.RegisterActiveSettings);
+	TRACE_CPUPROFILER_EVENT_SCOPE(TedsSettingsManager.AddSettings);
 
-	ISettingsModule* SettingsModule = FModuleManager::GetModulePtr<ISettingsModule>("Settings");
-	check(SettingsModule);
+	ISettingsModule& SettingsModule = FModuleManager::GetModuleChecked<ISettingsModule>("Settings");
 
 	TArray<FName> ContainerNames;
-	SettingsModule->GetContainerNames(ContainerNames);
+	SettingsModule.GetContainerNames(ContainerNames);
 
 	for (FName ContainerName : ContainerNames)
 	{
-		RegisterSettingsContainer(ContainerName);
+		AddSettingsContainer(ContainerName);
 	}
 
-	SettingsModule->OnContainerAdded().AddSP(this, &FTedsSettingsManager::RegisterSettingsContainer);
+	SettingsModule.OnContainerAdded().AddSP(this, &FTedsSettingsManager::AddSettingsContainer);
 }
 
-void FTedsSettingsManager::RegisterSettingsContainer(const FName& ContainerName)
+void FTedsSettingsManager::AddSettingsContainer(const FName& ContainerName)
 {
 	using namespace UE::Editor::DataStorage;
 	using namespace UE::Editor::Settings::Private;
 
-	TRACE_CPUPROFILER_EVENT_SCOPE(TedsSettingsManager.RegisterSettingsContainer);
+	TRACE_CPUPROFILER_EVENT_SCOPE(TedsSettingsManager.AddSettingsContainer);
 
-	ISettingsModule* SettingsModule = FModuleManager::GetModulePtr<ISettingsModule>("Settings");
-	check(SettingsModule);
+	ISettingsModule& SettingsModule = FModuleManager::GetModuleChecked<ISettingsModule>("Settings");
 
 	ICoreProvider* DataStorage = GetMutableDataStorageFeature<ICoreProvider>(StorageFeatureName);
 	check(DataStorage);
 
-	UE_LOG(LogTedsSettings, Verbose, TEXT("Register Settings Container : '%s'"), *ContainerName.ToString());
+	UE_LOG(LogTedsSettings, Verbose, TEXT("Add Settings Container : '%s'"), *ContainerName.ToString());
 
-	ISettingsContainerPtr ContainerPtr = SettingsModule->GetContainer(ContainerName);
+	ISettingsContainerPtr ContainerPtr = SettingsModule.GetContainer(ContainerName);
 
 	FMapKey ContainerIndexKey = GenerateContainerIndexKey(ContainerName);
 	RowHandle ContainerRow = DataStorage->LookupMappedRow(Settings::MappingDomain, ContainerIndexKey);
@@ -350,15 +394,14 @@ void FTedsSettingsManager::RegisterSettingsContainer(const FName& ContainerName)
 		});
 }
 
-void FTedsSettingsManager::UnregisterActiveSettings()
+void FTedsSettingsManager::RemoveSettings()
 {
 	using namespace UE::Editor::DataStorage;
 	using namespace UE::Editor::Settings::Private;
 
-	TRACE_CPUPROFILER_EVENT_SCOPE(TedsSettingsManager.UnregisterActiveSettings);
+	TRACE_CPUPROFILER_EVENT_SCOPE(TedsSettingsManager.RemoveSettings);
 
-	ISettingsModule* SettingsModule = FModuleManager::GetModulePtr<ISettingsModule>("Settings");
-	check(SettingsModule);
+	ISettingsModule& SettingsModule = FModuleManager::GetModuleChecked<ISettingsModule>("Settings");
 
 	ICoreProvider* DataStorage = GetMutableDataStorageFeature<ICoreProvider>(StorageFeatureName);
 	check(DataStorage);
@@ -366,16 +409,16 @@ void FTedsSettingsManager::UnregisterActiveSettings()
 	ICompatibilityProvider* DataStorageCompatibility = GetMutableDataStorageFeature<ICompatibilityProvider>(CompatibilityFeatureName);
 	check(DataStorageCompatibility);
 
-	SettingsModule->OnContainerAdded().RemoveAll(this);
+	SettingsModule.OnContainerAdded().RemoveAll(this);
 
 	TArray<FName> ContainerNames;
-	SettingsModule->GetContainerNames(ContainerNames);
+	SettingsModule.GetContainerNames(ContainerNames);
 
 	for (FName ContainerName : ContainerNames)
 	{
-		UE_LOG(LogTedsSettings, Verbose, TEXT("Unregister Settings Container : '%s'"), *ContainerName.ToString());
+		UE_LOG(LogTedsSettings, Verbose, TEXT("Remove Settings Container : '%s'"), *ContainerName.ToString());
 
-		ISettingsContainerPtr ContainerPtr = SettingsModule->GetContainer(ContainerName);
+		ISettingsContainerPtr ContainerPtr = SettingsModule.GetContainer(ContainerName);
 
 		ContainerPtr->OnCategoryModified().RemoveAll(this);
 
@@ -503,13 +546,18 @@ void FTedsSettingsManager::UpdateSettingsCategory(TSharedPtr<ISettingsCategory>
 
 			RowHandle NewSectionRow = DataStorageCompatibility->AddCompatibleObject(SettingsObjectPtr);
 
-			DataStorage->AddColumn<FSettingsSectionTag>(NewSectionRow);
 			DataStorage->AddColumn<FSettingsContainerReferenceColumn>(NewSectionRow, { .ContainerName = ContainerName, .ContainerRow = ContainerRow });
 			DataStorage->AddColumn<FSettingsCategoryReferenceColumn>(NewSectionRow, { .CategoryName = CategoryName, .CategoryRow = CategoryRow });
 			DataStorage->AddColumn<FSettingsNameColumn>(NewSectionRow, { .Name = SectionName });
 			DataStorage->AddColumn<FDisplayNameColumn>(NewSectionRow, { .DisplayName = SectionPtr->GetDisplayName() });
 			DataStorage->AddColumn<FDescriptionColumn>(NewSectionRow, { .Description = SectionPtr->GetDescription() });
 			
+			if (UPackage* ClassPackage = SettingsObjectPtr->GetClass()->GetOuterUPackage())
+			{
+				FName ModuleName = FPackageName::GetShortFName(ClassPackage->GetFName());
+				DataStorage->AddColumn<FSettingsModuleTag>(NewSectionRow, ModuleName);
+			}
+
 			if (OldSectionRow != InvalidRowHandle && OldSectionRow != NewSectionRow)
 			{
 				// Remove the old section row and any child property rows.
@@ -533,13 +581,8 @@ void FTedsSettingsManager::UpdateSettingsCategory(TSharedPtr<ISettingsCategory>
 				{
 					DataStorage->RemoveRow(OldSectionRow);
 				}
-
-				UE_LOG(LogTedsSettings, Verbose, TEXT("Removed Settings Section : '%s'"), *SectionName.ToString());
 			}
 			
-			UE_CLOG(OldSectionRow != NewSectionRow, LogTedsSettings, Verbose, TEXT("Added Settings Section : '%s'"), *SectionName.ToString());
-			UE_CLOG(OldSectionRow == NewSectionRow, LogTedsSettings, Verbose, TEXT("Updated Settings Section : '%s'"), *SectionName.ToString());
-
 			if (OldSectionRow != NewSectionRow)
 			{
 				DataStorage->SetParentRow(SettingsHierarchy, NewSectionRow, CategoryRow);
@@ -563,6 +606,12 @@ void FTedsSettingsManager::UpdateSettingsCategory(TSharedPtr<ISettingsCategory>
 					DataStorage->SetParentRow(SettingsHierarchy, NewPropertyRow, NewSectionRow);
 				}
 			}
+
+			// Add the FSettingsSectionTag last so that any FObserver::OnAdd<FSettingsSectionTag>() can read the other columns.
+			DataStorage->AddColumn<FSettingsSectionTag>(NewSectionRow);
+
+			UE_CLOG(OldSectionRow == InvalidRowHandle, LogTedsSettings, Verbose, TEXT("Added Settings Section : '%s'"), *SectionName.ToString());
+			UE_CLOG(OldSectionRow != InvalidRowHandle, LogTedsSettings, Verbose, TEXT("Updated Settings Section : '%s'"), *SectionName.ToString());
 		}
 	}
 
@@ -604,3 +653,147 @@ void FTedsSettingsManager::UpdateSettingsCategory(TSharedPtr<ISettingsCategory>
 		UE_LOG(LogTedsSettings, Verbose, TEXT("Removed Settings Section : '%s'"), *OldSectionName.ToString());
 	}
 }
+
+void FTedsSettingsManager::OnModulesChanged(FName ModuleName, EModuleChangeReason ChangeReason)
+{
+	if (ChangeReason == EModuleChangeReason::ModuleLoaded)
+	{
+		RegisterSettingsForModule(ModuleName);
+	}
+	else if (ChangeReason == EModuleChangeReason::ModuleUnloaded)
+	{
+		UnregisterSettingsForModule(ModuleName);
+	}
+}
+
+void FTedsSettingsManager::RegisterSettings(ISettingsModule& SettingsModule, UDeveloperSettings* Settings)
+{
+	TSharedPtr<SWidget> CustomWidget = Settings->GetCustomSettingsWidget();
+
+	const FName ContainerName = Settings->GetContainerName();
+	const FName CategoryName = Settings->GetCategoryName();
+	const FName SectionName = Settings->GetSectionName();
+
+	UE_LOG(LogTedsSettings, Verbose, TEXT("Register Settings '%s->%s->%s' with the ISettingsModule"),
+		*ContainerName.ToString(), *CategoryName.ToString(), *SectionName.ToString());
+
+	if (CustomWidget.IsValid())
+	{
+		SettingsModule.RegisterSettings(
+			ContainerName,
+			CategoryName,
+			SectionName,
+			Settings->GetSectionText(),
+			Settings->GetSectionDescription(),
+			CustomWidget.ToSharedRef());
+	}
+	else
+	{
+		SettingsModule.RegisterSettings(
+			ContainerName,
+			CategoryName,
+			SectionName,
+			Settings->GetSectionText(),
+			Settings->GetSectionDescription(),
+			Settings);
+	}
+}
+
+void FTedsSettingsManager::RegisterSettingsForModule(FName ModuleName)
+{
+	TRACE_CPUPROFILER_EVENT_SCOPE(TedsSettingsManager.RegisterSettingsForModule);
+
+	// Find the class package for this module
+	const UPackage* const ClassPackage = FindPackage(nullptr, *(FString("/Script/") + ModuleName.ToString()));
+	if (!ClassPackage)
+	{
+		return;
+	}
+
+	TArray<UObject*> PackageObjects;
+	GetObjectsWithPackage(ClassPackage, PackageObjects, false);
+
+	TArray<UDeveloperSettings*> SettingsObjects;
+	for (UObject* Object : PackageObjects)
+	{
+		const UClass* CurrentClass = Cast<UClass>(Object);
+
+		if (CurrentClass &&
+			!CurrentClass->HasAnyClassFlags(CLASS_Deprecated | CLASS_Abstract) &&
+			CurrentClass->IsChildOf(UDeveloperSettings::StaticClass()))
+		{
+			UDeveloperSettings* SettingsObject = Cast<UDeveloperSettings>(CurrentClass->GetDefaultObject());
+			if (SettingsObject && SettingsObject->SupportsAutoRegistration())
+			{
+				SettingsObjects.Add(SettingsObject);
+			}
+		}
+	}
+
+	if (!SettingsObjects.IsEmpty())
+	{
+		UE_LOG(LogTedsSettings, Verbose, TEXT("Register Settings from module : '%s'"), *ModuleName.ToString());
+
+		ISettingsModule& SettingsModule = FModuleManager::GetModuleChecked<ISettingsModule>("Settings");
+
+		for (UDeveloperSettings* Settings : SettingsObjects)
+		{
+			RegisterSettings(SettingsModule, Settings);
+		}
+	}
+}
+
+void FTedsSettingsManager::UnregisterSettingsForModule(FName ModuleName)
+{
+	using namespace UE::Editor::DataStorage::Queries;
+
+	TRACE_CPUPROFILER_EVENT_SCOPE(TedsSettingsManager.UnregisterSettingsForModule);
+
+	ICoreProvider* DataStorage = GetMutableDataStorageFeature<ICoreProvider>(StorageFeatureName);
+	checkf(DataStorage, TEXT("DataStorage is not available"));
+
+	QueryHandle* SelectAllSettingsInModuleQuery = SelectAllSettingsInModuleQueryMap.Find(ModuleName);
+
+	if (!SelectAllSettingsInModuleQuery)
+	{
+		SelectAllSettingsInModuleQuery = &SelectAllSettingsInModuleQueryMap.Emplace(ModuleName, DataStorage->RegisterQuery(
+			Select()
+				.ReadOnly<FSettingsContainerReferenceColumn, FSettingsCategoryReferenceColumn, FSettingsNameColumn>()
+			.Where()
+				.All<FSettingsSectionTag>()
+				.All<FSettingsModuleTag>(ModuleName)
+			.Compile()));
+	}
+
+	TArray<FName> ContainerNames;
+	TArray<FName> CategoryNames;
+	TArray<FName> SectionNames;
+
+	DataStorage->RunQuery(*SelectAllSettingsInModuleQuery, CreateDirectQueryCallbackBinding(
+		[&ContainerNames, &CategoryNames, &SectionNames](IDirectQueryContext& Context,
+			const FSettingsContainerReferenceColumn* ContainerColumns,
+			const FSettingsCategoryReferenceColumn* CategoryColumns,
+			const FSettingsNameColumn* SectionNameColumns)
+		{
+			const uint32 RowCount = Context.GetRowCount();
+
+			for (uint32 RowIndex = 0; RowIndex < RowCount; ++RowIndex)
+			{
+				ContainerNames.Emplace(ContainerColumns[RowIndex].ContainerName);
+				CategoryNames.Emplace(CategoryColumns[RowIndex].CategoryName);
+				SectionNames.Emplace(SectionNameColumns[RowIndex].Name);
+			}
+		}));
+
+	if (!ContainerNames.IsEmpty())
+	{
+		UE_LOG(LogTedsSettings, Verbose, TEXT("Unregister Settings from module : '%s'"), *ModuleName.ToString());
+
+		ISettingsModule& SettingsModule = FModuleManager::GetModuleChecked<ISettingsModule>("Settings");
+
+		for (int32 i = 0; i < ContainerNames.Num(); ++i)
+		{
+			SettingsModule.UnregisterSettings(ContainerNames[i], CategoryNames[i], SectionNames[i]);
+		}
+	}
+}
diff --git a/Engine/Plugins/Experimental/EditorDataStorageFeatures/Source/TedsSettings/Private/TedsSettingsManager.h b/Engine/Plugins/Experimental/EditorDataStorageFeatures/Source/TedsSettings/Private/TedsSettingsManager.h
--- a/Engine/Plugins/Experimental/EditorDataStorageFeatures/Source/TedsSettings/Private/TedsSettingsManager.h
+++ b/Engine/Plugins/Experimental/EditorDataStorageFeatures/Source/TedsSettings/Private/TedsSettingsManager.h
@@ -6,6 +6,8 @@
 #include "Templates/SharedPointer.h"
 
 class ISettingsCategory;
+class ISettingsModule;
+class UDeveloperSettings;
 
 class FTedsSettingsManager final : public TSharedFromThis<FTedsSettingsManager>
 {
@@ -39,17 +41,24 @@ private:
 	void RegisterQueries(UE::Editor::DataStorage::ICoreProvider& DataStorage);
 	void UnregisterQueries(UE::Editor::DataStorage::ICoreProvider& DataStorage);
 
-	void RegisterActiveSettings();
-	void UnregisterActiveSettings();
+	void AddSettings();
+	void RemoveSettings();
 
-	void RegisterSettingsContainer(const FName& ContainerName);
+	void AddSettingsContainer(const FName& ContainerName);
 
 	void UpdateSettingsCategory(TSharedPtr<ISettingsCategory> SettingsCategory, UE::Editor::DataStorage::RowHandle ContainerRow, const bool bQueryExistingRows = true);
 
+	void OnModulesChanged(FName ModuleName, EModuleChangeReason ChangeReason);
+
+	void RegisterSettings(ISettingsModule& SettingsModule, UDeveloperSettings* Settings);
+	void RegisterSettingsForModule(FName ModuleName);
+	void UnregisterSettingsForModule(FName ModuleName);
+
 	bool bIsInitialized;
 	UE::Editor::DataStorage::QueryHandle SelectAllSettingsQuery;
 	UE::Editor::DataStorage::TableHandle SettingsContainerTable;
 	UE::Editor::DataStorage::TableHandle SettingsCategoryTable;
 	UE::Editor::DataStorage::TableHandle SettingsPropertyTable;
 	UE::Editor::DataStorage::FHierarchyHandle SettingsHierarchy;
+	TMap<FName, UE::Editor::DataStorage::QueryHandle> SelectAllSettingsInModuleQueryMap;
 };
diff --git a/Engine/Plugins/Experimental/EditorDataStorageFeatures/Source/TedsSettings/Public/TedsSettingsColumns.h b/Engine/Plugins/Experimental/EditorDataStorageFeatures/Source/TedsSettings/Public/TedsSettingsColumns.h
--- a/Engine/Plugins/Experimental/EditorDataStorageFeatures/Source/TedsSettings/Public/TedsSettingsColumns.h
+++ b/Engine/Plugins/Experimental/EditorDataStorageFeatures/Source/TedsSettings/Public/TedsSettingsColumns.h
@@ -19,7 +19,7 @@ struct FSettingsContainerReferenceColumn final : public FEditorDataStorageColumn
 	GENERATED_BODY()
 
 	// Store the name as it is frequently accessed to avoid indirection to the row.
-	UPROPERTY(meta = (Searchable))
+	UPROPERTY(meta = (Searchable, Sortable))
 	FName ContainerName;
 
 	UE::Editor::DataStorage::RowHandle ContainerRow;
@@ -31,7 +31,7 @@ struct FSettingsCategoryReferenceColumn final : public FEditorDataStorageColumn
 	GENERATED_BODY()
 
 	// Store the name as it is frequently accessed to avoid indirection to the row.
-	UPROPERTY(meta = (Searchable))
+	UPROPERTY(meta = (Searchable, Sortable))
 	FName CategoryName;
 
 	UE::Editor::DataStorage::RowHandle CategoryRow;
@@ -44,10 +44,17 @@ struct FSettingsNameColumn final : public FEditorDataStorageColumn
 {
 	GENERATED_BODY()
 
-	UPROPERTY(meta = (Searchable))
+	UPROPERTY(meta = (Searchable, Sortable))
 	FName Name;
 };
 
+// A column whose dynamic identifier specifies the module that contains the class of this row's settings section object.
+USTRUCT(meta = (DisplayName = "Settings Module", EditorDataStorage_DynamicColumnTemplate))
+struct FSettingsModuleTag final : public FEditorDataStorageTag
+{
+	GENERATED_BODY()
+};
+
 USTRUCT(meta = (DisplayName = "Settings Container"))
 struct FSettingsContainerTag final : public FEditorDataStorageTag
 {
diff --git a/Engine/Plugins/Experimental/EditorDataStorageFeatures/Source/TedsSettings/TedsSettings.Build.cs b/Engine/Plugins/Experimental/EditorDataStorageFeatures/Source/TedsSettings/TedsSettings.Build.cs
--- a/Engine/Plugins/Experimental/EditorDataStorageFeatures/Source/TedsSettings/TedsSettings.Build.cs
+++ b/Engine/Plugins/Experimental/EditorDataStorageFeatures/Source/TedsSettings/TedsSettings.Build.cs
@@ -21,7 +21,10 @@ public class TedsSettings : ModuleRules
 				{
 					"Core",
 					"CoreUObject",
+					"DeveloperSettings",
 					"Engine",
+					"Settings",
+					"SettingsEditor",
 					"Slate",
 					"SlateCore",
 					"UnrealEd",
