[TEDS Outliner] Add Column customization parameters to allow specifiying order and visibility

#okforithub public
#Synced-CL 51460409

diff --git a/Engine/Plugins/Experimental/EditorDataStorageFeatures/Source/TedsOutliner/Private/TedsOutlinerHelpers.cpp b/Engine/Plugins/Experimental/EditorDataStorageFeatures/Source/TedsOutliner/Private/TedsOutlinerHelpers.cpp
--- a/Engine/Plugins/Experimental/EditorDataStorageFeatures/Source/TedsOutliner/Private/TedsOutlinerHelpers.cpp
+++ b/Engine/Plugins/Experimental/EditorDataStorageFeatures/Source/TedsOutliner/Private/TedsOutlinerHelpers.cpp
@@ -9,6 +9,7 @@
 #include "Modules/ModuleManager.h"
 #include "ISceneOutliner.h"
 #include "LevelEditor.h"
+#include "TedsOutlinerImpl.h"
 #include "Widgets/Docking/SDockTab.h"
 
 namespace UE::Editor::Outliner::Helpers
@@ -75,6 +76,155 @@ namespace UE::Editor::Outliner::Helpers
 		}
 	}
 	
+	void OrderColumnGroup(const TArray<TWeakObjectPtr<const UScriptStruct>>& InColumnGroup,
+		const FTedsOutlinerColumnDescription& InColumnDescription, const uint8 InDefaultStartingPriorityIndex,
+		TMap<TWeakObjectPtr<const UScriptStruct>, uint8>& OutOrderedColumnMap)
+	{
+		// Custom implementation of the Kahn Topological Sort to order the columns.
+		// Early return if we didn't receive any columns to order
+		if (InColumnGroup.IsEmpty())
+		{
+			return;
+		}
+		using FColumnPriorityRelation = FTedsOutlinerColumnParams::FColumnPriorityRelation;
+		using EColumnPriorityRelation = FColumnPriorityRelation::EColumnPriorityRelation;
+
+		// Default number of Columns to allocate, in the rare case that the column number is
+		// above this amount, it will fall back to heap allocation
+		static constexpr uint32 InlineColumnCount = 8;
+		const int ColumnCount = InColumnGroup.Num();
+		
+		TMap<TWeakObjectPtr<const UScriptStruct>, TArray<TWeakObjectPtr<const UScriptStruct>, TInlineAllocator<InlineColumnCount>>,
+			TInlineSetAllocator<InlineColumnCount>> ColumnOrderGraph;
+		TMap<TWeakObjectPtr<const UScriptStruct>, int, TInlineSetAllocator<InlineColumnCount>> InDegree;
+		ColumnOrderGraph.Reserve(ColumnCount);
+		InDegree.Reserve(ColumnCount);
+
+		// Initialize Graph
+		for (const TWeakObjectPtr<const UScriptStruct>& Column : InColumnGroup)
+		{
+			ColumnOrderGraph.Add(Column, {});
+			InDegree.Add(Column, 0);
+		}
+
+		// Add Relations to the Graph
+		for (const TWeakObjectPtr<const UScriptStruct>& Column : InColumnGroup)
+		{
+			if (Column.IsValid())
+			{
+				if (const FTedsOutlinerColumnParams* Params = InColumnDescription.FindColumnParams(Column))
+				{
+					const FColumnPriorityRelation& PriorityRelation = Params->PriorityRelation;
+					if (PriorityRelation.RelatedColumn.IsValid() && InColumnGroup.Contains(PriorityRelation.RelatedColumn))
+					{
+						if (PriorityRelation.PriorityRelation == EColumnPriorityRelation::Before)
+						{
+							ColumnOrderGraph[Column].Add(PriorityRelation.RelatedColumn);
+							InDegree[PriorityRelation.RelatedColumn]++;
+						}
+						else if (PriorityRelation.PriorityRelation == EColumnPriorityRelation::After)
+						{
+							ColumnOrderGraph[PriorityRelation.RelatedColumn].Add(Column);
+							InDegree[Column]++;
+						}
+					}
+				}
+			}
+		}
+		
+		TQueue<TWeakObjectPtr<const UScriptStruct>> OrderQueue;
+		TArray<TWeakObjectPtr<const UScriptStruct>> FinalGroupOrder;
+        FinalGroupOrder.Reserve(ColumnCount);
+
+		for (const TWeakObjectPtr<const UScriptStruct>& Column : InColumnGroup)
+		{
+			if (InDegree[Column] == 0)
+			{
+				OrderQueue.Enqueue(Column);
+			}
+		}
+
+		while (!OrderQueue.IsEmpty())
+		{
+			TWeakObjectPtr<const UScriptStruct> CurrentColumn;
+			OrderQueue.Dequeue(CurrentColumn);
+			FinalGroupOrder.Add(CurrentColumn);
+		
+			for(const TWeakObjectPtr<const UScriptStruct>& FolRightingColumn : ColumnOrderGraph[CurrentColumn])
+			{
+				InDegree[FolRightingColumn]--;
+				if (InDegree[FolRightingColumn] == 0)
+				{
+					OrderQueue.Enqueue(FolRightingColumn);
+				}
+			}
+		}
+
+		// If the column sizes are different, there was a circular dependency that could not be resolved,
+		// let the user know that they should reorder their columns.
+		ensureAlwaysMsgf(FinalGroupOrder.Num() == ColumnCount,
+				TEXT("The TEDS Outliner Column order contains a circular priority relation, "
+				"the order will be incorrect until fixed."));
+
+		
+		// Assign Priority Indexes
+		int32 PriorityIndexOffset = 0;
+		for (const TWeakObjectPtr<const UScriptStruct>& Column : FinalGroupOrder)
+		{
+			if (Column.IsValid())
+			{
+				OutOrderedColumnMap.Add(Column, 
+					static_cast<uint8>(FMath::Clamp(InDefaultStartingPriorityIndex + static_cast<uint8>(
+						FMath::Clamp(PriorityIndexOffset, 0, 255)), 0, 255)));
+				++PriorityIndexOffset;
+			}
+		}
+	}
+		
+	void OrderColumns(const TArray<TWeakObjectPtr<const UScriptStruct>>& InColumns,
+		const FTedsOutlinerColumnDescription& InColumnDescription, TMap<TWeakObjectPtr<const UScriptStruct>, uint8>& OutOrderedColumnMap)
+	{
+		using EColumnPriorityGroup = FTedsOutlinerColumnParams::EColumnPriorityGroup;
+		
+		const int ColumnTypeNum = InColumns.Num();
+		TArray<TWeakObjectPtr<const UScriptStruct>> LeftPriorityColumns, MiddlePriorityColumns, RightPriorityColumns;
+		LeftPriorityColumns.Reserve(ColumnTypeNum);
+		MiddlePriorityColumns.Reserve(ColumnTypeNum);
+		RightPriorityColumns.Reserve(ColumnTypeNum);
+
+		// Sort into the different priority groups
+		for (const TWeakObjectPtr<const UScriptStruct>& ColumnType : InColumns)
+		{
+			if (const FTedsOutlinerColumnParams* FoundColumnParams = InColumnDescription.FindColumnParams(ColumnType))
+			{
+				switch (FoundColumnParams->PriorityGroup)
+				{
+				case EColumnPriorityGroup::Left:
+					LeftPriorityColumns.Add(ColumnType);
+					break;
+				case EColumnPriorityGroup::Right:
+					RightPriorityColumns.Add(ColumnType);
+					break;
+				case EColumnPriorityGroup::Middle:
+				default:
+					MiddlePriorityColumns.Add(ColumnType);
+				}
+			}
+			else
+			{
+				MiddlePriorityColumns.Add(ColumnType);
+			}
+		}
+
+		constexpr uint8 DefaultLeftPriorityIndex = 0;
+		constexpr uint8 DefaultMiddlePriorityIndex = 85;
+		constexpr uint8 DefaultRightPriorityIndex = 170;
+		
+		OrderColumnGroup(LeftPriorityColumns, InColumnDescription, DefaultLeftPriorityIndex, OutOrderedColumnMap);
+		OrderColumnGroup(MiddlePriorityColumns, InColumnDescription, DefaultMiddlePriorityIndex, OutOrderedColumnMap);
+		OrderColumnGroup(RightPriorityColumns, InColumnDescription, DefaultRightPriorityIndex, OutOrderedColumnMap);
+	}
+	
 	FName FindOutlinerColumnFromTedsColumns(const DataStorage::ICoreProvider* Storage, TConstArrayView<TWeakObjectPtr<const UScriptStruct>> TEDSColumns)
 	{
 		if(const USceneOutlinerTedsBridgeFactory* Factory = Storage->FindFactory<USceneOutlinerTedsBridgeFactory>())
diff --git a/Engine/Plugins/Experimental/EditorDataStorageFeatures/Source/TedsOutliner/Private/TedsOutlinerImpl.cpp b/Engine/Plugins/Experimental/EditorDataStorageFeatures/Source/TedsOutliner/Private/TedsOutlinerImpl.cpp
--- a/Engine/Plugins/Experimental/EditorDataStorageFeatures/Source/TedsOutliner/Private/TedsOutlinerImpl.cpp
+++ b/Engine/Plugins/Experimental/EditorDataStorageFeatures/Source/TedsOutliner/Private/TedsOutlinerImpl.cpp
@@ -37,6 +37,7 @@
 #include "Elements/Columns/TypedElementCompatibilityColumns.h"
 #include "Elements/Columns/TypedElementHandleColumn.h"
 #include "Elements/Columns/TypedElementOverrideColumns.h"
+#include "Elements/Columns/TypedElementPackageColumns.h"
 #include "Elements/Columns/TypedElementSelectionColumns.h"
 #include "Elements/Columns/TypedElementSlateWidgetColumns.h"
 #include "Elements/Columns/TypedElementUIColumns.h"
@@ -165,9 +166,26 @@ namespace LabelWidgetUtils
 	}
 }
 
+FTedsOutlinerColumnParams::FColumnPriorityRelation::FColumnPriorityRelation(const EColumnPriorityRelation InPriorityRelation,
+	const TWeakObjectPtr<const UScriptStruct>& InRelatedColumn)
+	: PriorityRelation(InPriorityRelation), RelatedColumn(InRelatedColumn)
+{}
+
+FTedsOutlinerColumnParams::FTedsOutlinerColumnParams(const ESceneOutlinerColumnVisibility InInitialVisibility,
+	const EColumnPriorityGroup InPriorityGroup, const FColumnPriorityRelation& InPriorityRelation)
+	: InitialVisibility(InInitialVisibility), PriorityGroup(InPriorityGroup), PriorityRelation(InPriorityRelation)
+{}
+
+FTedsOutlinerColumnParams::FTedsOutlinerColumnParams(const EColumnPriorityGroup InPriorityGroup,
+	const FColumnPriorityRelation& InPriorityRelation)
+		: FTedsOutlinerColumnParams(ESceneOutlinerColumnVisibility::Visible, InPriorityGroup, InPriorityRelation)
+{}
+
 FTedsOutlinerColumnDescription::FTedsOutlinerColumnDescription(const TArray<TWeakObjectPtr<const UScriptStruct>>& InColumns,
+	const TMap<TWeakObjectPtr<const UScriptStruct>, FTedsOutlinerColumnParams>& InColumnParams,
 	const TArray<TSharedPtr<FColumnMetaData>>& InColumnMetaData, const FMetaData& InGenericMetaData)
 	: Columns(InColumns)
+	, ColumnParams(InColumnParams)
 	, GenericMetaData(InGenericMetaData)
 {
 	ColumnMetaData.Reserve(Columns.Num());
@@ -176,7 +194,7 @@ FTedsOutlinerColumnDescription::FTedsOutlinerColumnDescription(const TArray<TWea
 	// can override this by passing through the constructor.
 	if (InColumnMetaData.IsEmpty())
 	{
-		for (const TWeakObjectPtr<const UScriptStruct> Column : Columns)
+		for (const TWeakObjectPtr<const UScriptStruct>& Column : InColumns)
 		{
 			if (Column.IsValid())
 			{
@@ -191,10 +209,29 @@ FTedsOutlinerColumnDescription::FTedsOutlinerColumnDescription(const TArray<TWea
 }
 	
 FTedsOutlinerColumnDescription::FTedsOutlinerColumnDescription(const TArray<TWeakObjectPtr<const UScriptStruct>>& InColumns, const FMetaData& InGenericMetaData)
-	: Columns(InColumns)
-	, GenericMetaData(InGenericMetaData)
+	: FTedsOutlinerColumnDescription(InColumns, {}, {}, InGenericMetaData)
 {}
 
+const TArray<TWeakObjectPtr<const UScriptStruct>>& FTedsOutlinerColumnDescription::GetColumns() const
+{
+	return Columns;
+}
+
+const FTedsOutlinerColumnParams* FTedsOutlinerColumnDescription::FindColumnParams(const TWeakObjectPtr<const UScriptStruct>& InColumn) const
+{
+	return ColumnParams.Find(InColumn);
+}
+
+const TArray<TSharedPtr<FColumnMetaData>>& FTedsOutlinerColumnDescription::GetColumnMetaData() const
+{
+	return ColumnMetaData;
+}
+
+FMetaData FTedsOutlinerColumnDescription::GetGenericMetaData() const
+{
+	return GenericMetaData;
+}
+
 FTedsOutlinerParams::FTedsOutlinerParams(SSceneOutliner* InSceneOutliner)
 	: SceneOutliner(InSceneOutliner)
 	, QueryDescription()
@@ -630,7 +667,6 @@ void FTedsOutlinerImpl::RegenerateColumns(const Outliner::FTedsOutlinerColumnDes
 	}
 
 	using MatchApproach = DataStorage::IUiProvider::EMatchApproach;
-	constexpr uint8 DefaultPriorityIndex = 100;
 	
 	DataStorage::FColumnsMetaDataView MetaDataView(ColumnDescription.GetColumnMetaData());
 	
@@ -638,37 +674,62 @@ void FTedsOutlinerImpl::RegenerateColumns(const Outliner::FTedsOutlinerColumnDes
 	
 	TArray<TWeakObjectPtr<const UScriptStruct>> ColumnTypes =
 		TableViewerUtils::CreateVerifiedColumnTypeArray(ColumnDescription.GetColumns());
+
+	// Sort by before/after before passing here
+	TMap<TWeakObjectPtr<const UScriptStruct>, uint8> OrderedColumnMap;
+	OrderedColumnMap.Reserve(ColumnTypes.Num());
+	Helpers::OrderColumns(ColumnTypes, ColumnDescription, OrderedColumnMap);
 	
 	const RowHandle CellPurposeRow = StorageUi->FindPurpose(CellWidgetPurpose);
 	const RowHandle HeaderPurposeRow = StorageUi->FindPurpose(HeaderWidgetPurpose);
 
 	int32 IndexOffset = 0;
-	auto ColumnConstructor = [this, ColumnDescription, TedsOutlinerColumn, MetaDataView, &IndexOffset, HeaderPurposeRow, Dealiaser](
+	auto ColumnConstructor = [this, ColumnDescription, TedsOutlinerColumn, MetaDataView, &IndexOffset, &OrderedColumnMap, HeaderPurposeRow, Dealiaser](
 		IUiProvider::FWidgetConstructorPtr CellConstructor, TConstArrayView<TWeakObjectPtr<const UScriptStruct>> ColumnTypes)
 		{
-			/* If we have a fallback column for this query, remove it, take over it's priority and 
-			 * replace it with the TEDS column. But also allow the TEDS-Outliner column to fallback to it for
+			/* If we have a fallback column for this query, remove it, take over its priority and 
+			 * replace it with the TEDS column. But also allow the TEDS-Outliner column to fall back to it for
 			 * data not in TEDS yet.
 			 */
-			FName FallbackColumnName = Helpers::FindOutlinerColumnFromTedsColumns(Storage, ColumnTypes);
-			const FSceneOutlinerColumnInfo* FallbackColumnInfo = SceneOutliner->GetSharedData().ColumnMap.Find(FallbackColumnName);
-			uint8 ColumnPriority = FallbackColumnInfo
-				? FallbackColumnInfo->PriorityIndex
-				: static_cast<uint8>(
-					FMath::Clamp(DefaultPriorityIndex + static_cast<uint8>(
-						FMath::Clamp(IndexOffset, 0, 255)), 0, 255));
+			const FName FallbackColumnName = Helpers::FindOutlinerColumnFromTedsColumns(Storage, ColumnTypes);
+			// A priority should always be found, but use the IndexOffset as a default
+			uint8 ColumnPriority = static_cast<uint8>(FMath::Clamp(IndexOffset, 0, 255));
+			ESceneOutlinerColumnVisibility ColumnVisibility = ESceneOutlinerColumnVisibility::Visible;
+			if (const FSceneOutlinerColumnInfo* FallbackColumnInfo = SceneOutliner->GetSharedData().ColumnMap.Find(FallbackColumnName))
+			{
+				ColumnVisibility = FallbackColumnInfo->Visibility;
+				ColumnPriority = FallbackColumnInfo->PriorityIndex;
+			}
+			else
+			{
+				for (const TWeakObjectPtr<const UScriptStruct>& ColumnType : ColumnTypes)
+				{
+					if (const FTedsOutlinerColumnParams* FoundColumnParams = ColumnDescription.FindColumnParams(ColumnType))
+					{
+						ColumnVisibility = FoundColumnParams->InitialVisibility;
+					}
+					if (const uint8* FoundColumnPriority = OrderedColumnMap.Find(ColumnType))
+					{
+						// We want to use the same column if possible so break as soon as we find the Priority
+						// (meaning Params would have been defined if desired for this column above, if not, none
+						// were given for the matching column we found for the priority)
+						ColumnPriority = *FoundColumnPriority;
+						break;
+					}
+				}
+			}
 
 			// Grab a reference to the fallback column and remove it from the actual outliner
 			const TMap<FName, TSharedPtr<ISceneOutlinerColumn>>& Columns = SceneOutliner->GetColumns();
 			const TSharedPtr<ISceneOutlinerColumn> FallbackColumnPtr = Columns.FindRef(FallbackColumnName);
 			SceneOutliner->RemoveColumn(FallbackColumnName);
 
-			FName NameId = UE::Editor::DataStorage::TableViewerUtils::FindLongestMatchingName(ColumnTypes, IndexOffset);
-			FText DisplayName = CellConstructor.Get()->CreateWidgetDisplayNameText(Storage);
+			FName NameId = TableViewerUtils::FindLongestMatchingName(ColumnTypes, IndexOffset);
+			const FText DisplayName = CellConstructor.Get()->CreateWidgetDisplayNameText(Storage);
 			AddedColumns.Add(NameId);
 			SceneOutliner->AddColumn(NameId,
 				FSceneOutlinerColumnInfo(
-					ESceneOutlinerColumnVisibility::Visible, 
+					ColumnVisibility, 
 					ColumnPriority,
 					FCreateSceneOutlinerColumn::CreateLambda(
 						[this, ColumnDescription, TedsOutlinerColumn, MetaDataView, NameId, &ColumnTypes, CellConstructor, FallbackColumnPtr, HeaderPurposeRow, Dealiaser](ISceneOutliner&)
diff --git a/Engine/Plugins/Experimental/EditorDataStorageFeatures/Source/TedsOutliner/Private/TedsOutlinerModule.cpp b/Engine/Plugins/Experimental/EditorDataStorageFeatures/Source/TedsOutliner/Private/TedsOutlinerModule.cpp
--- a/Engine/Plugins/Experimental/EditorDataStorageFeatures/Source/TedsOutliner/Private/TedsOutlinerModule.cpp
+++ b/Engine/Plugins/Experimental/EditorDataStorageFeatures/Source/TedsOutliner/Private/TedsOutlinerModule.cpp
@@ -98,18 +98,30 @@ TSharedRef<ISceneOutliner> FTedsOutlinerModule::CreateTedsOutliner(const FSceneO
 		return new FTedsOutlinerMode(InitTedsOptions);
 	});
 
+	// Explicitly set the priority Index to greater than high-priority icon columns since the Label Column
+	// is manually added to the ColumnMap right now and cannot use dynamic priority grouping
+	constexpr int LabelColumnPriorityIndex = 10;
+	// Explicitly set the priority Index to one less than the Label Priority since we cannot use priority groups (not
+	// passed in as a TEDS Column), and we want it to appear after high-priority icon columns but before the Label Column.
+	constexpr int RowHandleColumnPriorityIndex = LabelColumnPriorityIndex - 1;
+	
 	// Add the custom column that displays row handles
 	if (InInitTedsOptions.bShowRowHandleColumn)
 	{
 		InitOptions.ColumnMap.Add(FSceneOutlinerRowHandleColumn::GetID(),
-			FSceneOutlinerColumnInfo(ESceneOutlinerColumnVisibility::Visible, 2,
+			FSceneOutlinerColumnInfo(ESceneOutlinerColumnVisibility::Invisible, RowHandleColumnPriorityIndex,
 				FCreateSceneOutlinerColumn::CreateLambda([](ISceneOutliner& InSceneOutliner)
 				{
 					return MakeShareable(new FSceneOutlinerRowHandleColumn(InSceneOutliner));
 				})));
 	}
 	
-	InitOptions.ColumnMap.Add(FSceneOutlinerBuiltInColumnTypes::Label(), FSceneOutlinerColumnInfo(ESceneOutlinerColumnVisibility::Visible, 10));
+	InitOptions.ColumnMap.Add(FSceneOutlinerBuiltInColumnTypes::Label(),
+		FSceneOutlinerColumnInfo(
+			ESceneOutlinerColumnVisibility::Visible,
+			LabelColumnPriorityIndex,
+			FCreateSceneOutlinerColumn(),
+			false));
 	
 	TSharedRef<ISceneOutliner> TedsOutlinerShared = SNew(SSceneOutliner, InitOptions);
 	
@@ -133,6 +145,7 @@ void FTedsOutlinerModule::ShutdownModule()
 FTedsOutlinerColumnDescription FTedsOutlinerModule::GetLevelEditorTedsOutlinerColumnDescription()
 {
 	using namespace UE::Editor::DataStorage::Columns;
+	using EColumnPriorityGroup = FTedsOutlinerColumnParams::EColumnPriorityGroup;
 		
 	TArray<TWeakObjectPtr<const UScriptStruct>> OutColumns({
 		FVisibleInEditorColumn::StaticStruct(),
@@ -154,7 +167,13 @@ FTedsOutlinerColumnDescription FTedsOutlinerModule::GetLevelEditorTedsOutlinerCo
 		OutColumns.Add(FTypedElementPackageReference::StaticStruct());
 	}
 
-	return FTedsOutlinerColumnDescription(OutColumns);
+	TMap<TWeakObjectPtr<const UScriptStruct>, FTedsOutlinerColumnParams> OutColumnParams
+	{
+		{FVisibleInEditorColumn::StaticStruct(), FTedsOutlinerColumnParams(EColumnPriorityGroup::Left)},
+		{FTedsPrimaryPackageObjectTag::StaticStruct(), FTedsOutlinerColumnParams(EColumnPriorityGroup::Left)},
+	};
+	
+	return FTedsOutlinerColumnDescription(OutColumns, OutColumnParams);
 }
 
 TSharedRef<SWidget> FTedsOutlinerModule::CreateLevelEditorTedsOutliner()
diff --git a/Engine/Plugins/Experimental/EditorDataStorageFeatures/Source/TedsOutliner/Public/TedsOutlinerHelpers.h b/Engine/Plugins/Experimental/EditorDataStorageFeatures/Source/TedsOutliner/Public/TedsOutlinerHelpers.h
--- a/Engine/Plugins/Experimental/EditorDataStorageFeatures/Source/TedsOutliner/Public/TedsOutlinerHelpers.h
+++ b/Engine/Plugins/Experimental/EditorDataStorageFeatures/Source/TedsOutliner/Public/TedsOutlinerHelpers.h
@@ -7,6 +7,8 @@
 #include "SceneOutlinerFwd.h"
 #include "Columns/TedsOutlinerColumns.h"
 
+class FTedsOutlinerColumnDescription;
+
 // Helper functions for Teds Outliner related functionality
 namespace UE::Editor::Outliner::Helpers
 {
@@ -22,6 +24,17 @@ namespace UE::Editor::Outliner::Helpers
 	
 	// Refresh all the Outliners currently open in the level editor
 	TEDSOUTLINER_API void RefreshLevelEditorOutliners();
+
+	// Orders a given array of Columns within the same priority group and relations defined by the
+	// ColumnDescription and returns evaluated priorities to the OutOrderedColumnMap
+	void TEDSOUTLINER_API OrderColumnGroup(const TArray<TWeakObjectPtr<const UScriptStruct>>& InColumnGroup,
+    	const FTedsOutlinerColumnDescription& InColumnDescription, const uint8 InDefaultStartingPriorityIndex,
+    	TMap<TWeakObjectPtr<const UScriptStruct>, uint8>& OutOrderedColumnMap);
+
+	// Orders a given array of Columns using priority groups and relations defined by the
+	// ColumnDescription and returns evaluated priorities to the OutOrderedColumnMap
+	void TEDSOUTLINER_API OrderColumns(const TArray<TWeakObjectPtr<const UScriptStruct>>& InColumns,
+		const FTedsOutlinerColumnDescription& InColumnDescription, TMap<TWeakObjectPtr<const UScriptStruct>, uint8>& OutOrderedColumnMap);
 	
 	// Get the name of the Outliner column corresponding to the given TEDS column (if any)
 	FName FindOutlinerColumnFromTedsColumns(const DataStorage::ICoreProvider* Storage, TConstArrayView<TWeakObjectPtr<const UScriptStruct>> TEDSColumns);
diff --git a/Engine/Plugins/Experimental/EditorDataStorageFeatures/Source/TedsOutliner/Public/TedsOutlinerImpl.h b/Engine/Plugins/Experimental/EditorDataStorageFeatures/Source/TedsOutliner/Public/TedsOutlinerImpl.h
--- a/Engine/Plugins/Experimental/EditorDataStorageFeatures/Source/TedsOutliner/Public/TedsOutlinerImpl.h
+++ b/Engine/Plugins/Experimental/EditorDataStorageFeatures/Source/TedsOutliner/Public/TedsOutlinerImpl.h
@@ -11,6 +11,7 @@
 #include "Elements/Interfaces/TypedElementDataStorageUiInterface.h"
 #include "Elements/Framework/TypedElementQueryBuilder.h"
 #include "Elements/Framework/TypedElementSelectionSet.h"
+#include "SceneOutlinerPublicTypes.h"
 #include "TedsOutlinerFilter.h"
 #include "TedsOutlinerHierarchyInterfaces.h"
 
@@ -47,23 +48,74 @@ class SWidget;
 namespace UE::Editor::Outliner
 {
 	
+struct FTedsOutlinerColumnParams
+{
+	// Defines the priority groups, where Left priority is ordered first and Right priority is ordered last
+	// ||  Left  |  Middle  |  Right  ||
+	enum class EColumnPriorityGroup
+	{
+		Middle = 0, // Initial Priority value of 80
+		Left = 1 << 0, // Initial Priority value of 0
+		Right = 1 << 1 // Initial Priority value of 160
+	};
+
+	// Defines possible relations between columns so they can specify the desired ordering
+	struct FColumnPriorityRelation
+	{
+		enum class EColumnPriorityRelation
+		{
+			None = 0,
+			Before = 1 << 0, // This column comes before the specified RelatedColumn
+			After = 1 << 1 // This column comes after the specified RelatedColumn
+		};
+		
+		explicit FColumnPriorityRelation(const EColumnPriorityRelation InPriorityRelation = EColumnPriorityRelation::None,
+			const TWeakObjectPtr<const UScriptStruct>& InRelatedColumn = nullptr);
+
+		// When determining priority, all columns with Before relations are sorted before columns with After Columns
+		EColumnPriorityRelation PriorityRelation;
+		TWeakObjectPtr<const UScriptStruct> RelatedColumn;
+	};
+
+	explicit TEDSOUTLINER_API FTedsOutlinerColumnParams(
+		const ESceneOutlinerColumnVisibility InInitialVisibility = ESceneOutlinerColumnVisibility::Visible,
+		const EColumnPriorityGroup InPriorityGroup = EColumnPriorityGroup::Middle,
+		const FColumnPriorityRelation& InPriorityRelation = FColumnPriorityRelation());
+
+	explicit TEDSOUTLINER_API FTedsOutlinerColumnParams(
+		const EColumnPriorityGroup InPriorityGroup,
+		const FColumnPriorityRelation& InPriorityRelation = FColumnPriorityRelation());
+
+	// Visibility State of this column when the Outliner is Initialized (can change the visibility via the Header right-click menu)
+	ESceneOutlinerColumnVisibility InitialVisibility;
+	// Priority Group this column belongs to, it will sort by relations in this group and have an adjusted priority value based on its group
+	EColumnPriorityGroup PriorityGroup;
+	// Specifies if this column should appear before or after another column (must be in the same priority group). Without this, order
+	// is defined by what Column is passed the Column Array first
+	FColumnPriorityRelation PriorityRelation;
+};
+	
 struct FTedsOutlinerColumnDescription
 {
 	TEDSOUTLINER_API FTedsOutlinerColumnDescription(
 		const TArray<TWeakObjectPtr<const UScriptStruct>>& InColumns = {},
+		const TMap<TWeakObjectPtr<const UScriptStruct>, FTedsOutlinerColumnParams>& InColumnParams = {},
 		const TArray<TSharedPtr<FColumnMetaData>>& InColumnMetaData = {},
 		const FMetaData& InGenericMetaData = FMetaData());
 
 	TEDSOUTLINER_API FTedsOutlinerColumnDescription(
 		const TArray<TWeakObjectPtr<const UScriptStruct>>& InColumns, const FMetaData& InGenericMetaData);
 
-	TArray<TWeakObjectPtr<const UScriptStruct>> GetColumns() const { return Columns; }
-	TArray<TSharedPtr<FColumnMetaData>> GetColumnMetaData() const { return ColumnMetaData; }
-	FMetaData GetGenericMetaData() const { return GenericMetaData; }
+	const TArray<TWeakObjectPtr<const UScriptStruct>>& GetColumns() const;
+	const FTedsOutlinerColumnParams* FindColumnParams(const TWeakObjectPtr<const UScriptStruct>& InColumn) const;
+	const TArray<TSharedPtr<FColumnMetaData>>& GetColumnMetaData() const;
+	FMetaData GetGenericMetaData() const;
 
 protected:
 	// Columns to display in the Outliner
 	TArray<TWeakObjectPtr<const UScriptStruct>> Columns;
+	// Customization Parameters for a given column
+	TMap<TWeakObjectPtr<const UScriptStruct>, FTedsOutlinerColumnParams> ColumnParams;
 	// Metadata for Columns (ReadOnly/ReadWrite), by default the constructor will populate this with ReadOnly Metadata for all columns
 	TArray<TSharedPtr<FColumnMetaData>> ColumnMetaData;
 	// Generic meta data for the TEDS (i.e. GenericMetaData.AddOrSetMutableData("bAllowSorting", false))
diff --git a/Engine/Plugins/Experimental/EditorDataStorageFeatures/Source/TedsOutliner/Tests/TedsOutlinerColumnTests.cpp b/Engine/Plugins/Experimental/EditorDataStorageFeatures/Source/TedsOutliner/Tests/TedsOutlinerColumnTests.cpp
new file mode 100644
--- /dev/null
+++ b/Engine/Plugins/Experimental/EditorDataStorageFeatures/Source/TedsOutliner/Tests/TedsOutlinerColumnTests.cpp
@@ -0,0 +1,352 @@
+// Copyright Epic Games, Inc. All Rights Reserved.
+
+#if WITH_TESTS
+
+#include "Elements/Framework/TypedElementTestColumns.h"
+#include "Misc/AutomationTest.h"
+#include "TedsOutlinerHelpers.h"
+#include "TedsOutlinerImpl.h"
+#include "Tests/Assertions.h"
+#include "Tests/TestHarnessAdapter.h"
+
+static uint8 GetPriorityFromMap(const TWeakObjectPtr<const UScriptStruct>& ColumnType, const TMap<TWeakObjectPtr<const UScriptStruct>, uint8>& PriorityMap)
+{
+	if(const uint8* Priority = PriorityMap.Find(ColumnType))
+	{
+		return *Priority;
+	}
+	return UINT8_MAX;
+}
+
+TEST_CASE_NAMED(TEDS_Outliner_Column_Tests, "Editor::DataStorage::Outliner::ColumnOrder", "[ApplicationContextMask][EngineFilter]")
+{
+	using namespace UE::Editor::DataStorage;
+	using namespace UE::Editor::Outliner;
+	using FTedsOutlinerColumnDescription = UE::Editor::Outliner::FTedsOutlinerColumnDescription;
+	using FTedsOutlinerColumnParams = UE::Editor::Outliner::FTedsOutlinerColumnParams;
+	using EColumnPriorityGroup = FTedsOutlinerColumnParams::EColumnPriorityGroup;
+	using FColumnPriorityRelation = FTedsOutlinerColumnParams::FColumnPriorityRelation;
+	using EColumnPriorityRelation = FColumnPriorityRelation::EColumnPriorityRelation;
+
+	const TArray<TWeakObjectPtr<const UScriptStruct>> TestColumns({
+		FTestColumnA::StaticStruct(),
+		FTestColumnB::StaticStruct(),
+		FTestColumnC::StaticStruct(),
+		FTestColumnD::StaticStruct(),
+		FTestColumnE::StaticStruct(),
+		FTestColumnF::StaticStruct(),
+		FTestColumnG::StaticStruct()});
+	TMap<TWeakObjectPtr<const UScriptStruct>, uint8> OrderedColumnMap;
+
+	SECTION("No order specified")
+	{
+		const FTedsOutlinerColumnDescription TestColumnDescription = FTedsOutlinerColumnDescription(TestColumns);
+		OrderedColumnMap.Reset();
+		Helpers::OrderColumns(TestColumns, TestColumnDescription, OrderedColumnMap);
+
+		CHECK_EQUALS(TEXT("Priority Map Size"), OrderedColumnMap.Num(), 7);
+		CHECK_EQUALS(TEXT("Priority[Column A]"), GetPriorityFromMap(FTestColumnA::StaticStruct(), OrderedColumnMap), 85);
+		CHECK_EQUALS(TEXT("Priority[Column B]"), GetPriorityFromMap(FTestColumnB::StaticStruct(), OrderedColumnMap), 86);
+		CHECK_EQUALS(TEXT("Priority[Column C]"), GetPriorityFromMap(FTestColumnC::StaticStruct(), OrderedColumnMap), 87);
+		CHECK_EQUALS(TEXT("Priority[Column D]"), GetPriorityFromMap(FTestColumnD::StaticStruct(), OrderedColumnMap), 88);
+		CHECK_EQUALS(TEXT("Priority[Column E]"), GetPriorityFromMap(FTestColumnE::StaticStruct(), OrderedColumnMap), 89);
+		CHECK_EQUALS(TEXT("Priority[Column F]"), GetPriorityFromMap(FTestColumnF::StaticStruct(), OrderedColumnMap), 90);
+		CHECK_EQUALS(TEXT("Priority[Column G]"), GetPriorityFromMap(FTestColumnG::StaticStruct(), OrderedColumnMap), 91);
+	}
+
+	SECTION("Order by group")
+	{
+		const TMap<TWeakObjectPtr<const UScriptStruct>, FTedsOutlinerColumnParams> TestColumnParams
+		{
+			{FTestColumnA::StaticStruct(), FTedsOutlinerColumnParams(EColumnPriorityGroup::Right)},
+			{FTestColumnB::StaticStruct(), FTedsOutlinerColumnParams(EColumnPriorityGroup::Left)},
+			{FTestColumnC::StaticStruct(), FTedsOutlinerColumnParams(EColumnPriorityGroup::Middle)},
+			{FTestColumnD::StaticStruct(), FTedsOutlinerColumnParams(EColumnPriorityGroup::Left)},
+			{FTestColumnE::StaticStruct(), FTedsOutlinerColumnParams(EColumnPriorityGroup::Right)},
+			{FTestColumnF::StaticStruct(), FTedsOutlinerColumnParams(EColumnPriorityGroup::Middle)},
+			{FTestColumnG::StaticStruct(), FTedsOutlinerColumnParams(EColumnPriorityGroup::Left)},
+		};
+		const FTedsOutlinerColumnDescription TestColumnDescription = FTedsOutlinerColumnDescription(TestColumns, TestColumnParams);
+		OrderedColumnMap.Reset();
+		Helpers::OrderColumns(TestColumns, TestColumnDescription, OrderedColumnMap);
+
+		CHECK_EQUALS(TEXT("Priority Map Size"), OrderedColumnMap.Num(), 7);
+		CHECK_EQUALS(TEXT("Priority[Column B]"), GetPriorityFromMap(FTestColumnB::StaticStruct(), OrderedColumnMap), 0);
+		CHECK_EQUALS(TEXT("Priority[Column D]"), GetPriorityFromMap(FTestColumnD::StaticStruct(), OrderedColumnMap), 1);
+		CHECK_EQUALS(TEXT("Priority[Column G]"), GetPriorityFromMap(FTestColumnG::StaticStruct(), OrderedColumnMap), 2);
+		CHECK_EQUALS(TEXT("Priority[Column C]"), GetPriorityFromMap(FTestColumnC::StaticStruct(), OrderedColumnMap), 85);
+		CHECK_EQUALS(TEXT("Priority[Column F]"), GetPriorityFromMap(FTestColumnF::StaticStruct(), OrderedColumnMap), 86);
+		CHECK_EQUALS(TEXT("Priority[Column A]"), GetPriorityFromMap(FTestColumnA::StaticStruct(), OrderedColumnMap), 170);
+		CHECK_EQUALS(TEXT("Priority[Column E]"), GetPriorityFromMap(FTestColumnE::StaticStruct(), OrderedColumnMap), 171);
+	}
+
+	SECTION("Order by group with before relations")
+	{
+		const TMap<TWeakObjectPtr<const UScriptStruct>, FTedsOutlinerColumnParams> TestColumnParams
+		{
+			{FTestColumnA::StaticStruct(), FTedsOutlinerColumnParams(EColumnPriorityGroup::Left)},
+			{FTestColumnB::StaticStruct(), FTedsOutlinerColumnParams(EColumnPriorityGroup::Left,
+				FColumnPriorityRelation(EColumnPriorityRelation::Before, FTestColumnA::StaticStruct()))},
+			{FTestColumnC::StaticStruct(), FTedsOutlinerColumnParams(EColumnPriorityGroup::Middle)},
+			{FTestColumnD::StaticStruct(), FTedsOutlinerColumnParams(EColumnPriorityGroup::Middle,
+				FColumnPriorityRelation(EColumnPriorityRelation::Before, FTestColumnC::StaticStruct()))},
+			{FTestColumnE::StaticStruct(), FTedsOutlinerColumnParams(EColumnPriorityGroup::Left,
+				FColumnPriorityRelation(EColumnPriorityRelation::Before, FTestColumnB::StaticStruct()))},
+			{FTestColumnF::StaticStruct(), FTedsOutlinerColumnParams(EColumnPriorityGroup::Middle,
+				FColumnPriorityRelation(EColumnPriorityRelation::Before, FTestColumnC::StaticStruct()))},
+			{FTestColumnG::StaticStruct(), FTedsOutlinerColumnParams(EColumnPriorityGroup::Left, 
+				FColumnPriorityRelation(EColumnPriorityRelation::Before, FTestColumnE::StaticStruct()))},
+		};
+		const FTedsOutlinerColumnDescription TestColumnDescription = FTedsOutlinerColumnDescription(TestColumns, TestColumnParams);
+		OrderedColumnMap.Reset();
+		Helpers::OrderColumns(TestColumns, TestColumnDescription, OrderedColumnMap);
+
+		CHECK_EQUALS(TEXT("Priority Map Size"), OrderedColumnMap.Num(), 7);
+		CHECK_EQUALS(TEXT("Priority[Column G]"), GetPriorityFromMap(FTestColumnG::StaticStruct(), OrderedColumnMap), 0);
+		CHECK_EQUALS(TEXT("Priority[Column E]"), GetPriorityFromMap(FTestColumnE::StaticStruct(), OrderedColumnMap), 1);
+		CHECK_EQUALS(TEXT("Priority[Column B]"), GetPriorityFromMap(FTestColumnB::StaticStruct(), OrderedColumnMap), 2);
+		CHECK_EQUALS(TEXT("Priority[Column A]"), GetPriorityFromMap(FTestColumnA::StaticStruct(), OrderedColumnMap), 3);
+		CHECK_EQUALS(TEXT("Priority[Column D]"), GetPriorityFromMap(FTestColumnD::StaticStruct(), OrderedColumnMap), 85);
+		CHECK_EQUALS(TEXT("Priority[Column F]"), GetPriorityFromMap(FTestColumnF::StaticStruct(), OrderedColumnMap), 86);
+		CHECK_EQUALS(TEXT("Priority[Column C]"), GetPriorityFromMap(FTestColumnC::StaticStruct(), OrderedColumnMap), 87);	
+	}
+
+	SECTION("Order by group with after relations")
+	{
+		const TMap<TWeakObjectPtr<const UScriptStruct>, FTedsOutlinerColumnParams> TestColumnParams
+		{
+			{FTestColumnA::StaticStruct(), FTedsOutlinerColumnParams(EColumnPriorityGroup::Left,
+				FColumnPriorityRelation(EColumnPriorityRelation::After, FTestColumnB::StaticStruct()))},
+			{FTestColumnB::StaticStruct(), FTedsOutlinerColumnParams(EColumnPriorityGroup::Left)},
+			{FTestColumnC::StaticStruct(), FTedsOutlinerColumnParams(EColumnPriorityGroup::Middle,
+				FColumnPriorityRelation(EColumnPriorityRelation::After, FTestColumnE::StaticStruct()))},
+			{FTestColumnD::StaticStruct(), FTedsOutlinerColumnParams(EColumnPriorityGroup::Middle,
+				FColumnPriorityRelation(EColumnPriorityRelation::After, FTestColumnC::StaticStruct()))},
+			{FTestColumnE::StaticStruct(), FTedsOutlinerColumnParams(EColumnPriorityGroup::Middle,
+				FColumnPriorityRelation(EColumnPriorityRelation::After, FTestColumnF::StaticStruct()))},
+			{FTestColumnF::StaticStruct(), FTedsOutlinerColumnParams(EColumnPriorityGroup::Middle)},
+			{FTestColumnG::StaticStruct(), FTedsOutlinerColumnParams(EColumnPriorityGroup::Middle, 
+				FColumnPriorityRelation(EColumnPriorityRelation::After, FTestColumnC::StaticStruct()))},
+		};
+		const FTedsOutlinerColumnDescription TestColumnDescription = FTedsOutlinerColumnDescription(TestColumns, TestColumnParams);
+		OrderedColumnMap.Reset();
+		Helpers::OrderColumns(TestColumns, TestColumnDescription, OrderedColumnMap);
+
+		CHECK_EQUALS(TEXT("Priority Map Size"), OrderedColumnMap.Num(), 7);
+		CHECK_EQUALS(TEXT("Priority[Column B]"), GetPriorityFromMap(FTestColumnB::StaticStruct(), OrderedColumnMap), 0);
+		CHECK_EQUALS(TEXT("Priority[Column A]"), GetPriorityFromMap(FTestColumnA::StaticStruct(), OrderedColumnMap), 1);
+		CHECK_EQUALS(TEXT("Priority[Column F]"), GetPriorityFromMap(FTestColumnF::StaticStruct(), OrderedColumnMap), 85);
+		CHECK_EQUALS(TEXT("Priority[Column E]"), GetPriorityFromMap(FTestColumnE::StaticStruct(), OrderedColumnMap), 86);
+		CHECK_EQUALS(TEXT("Priority[Column C]"), GetPriorityFromMap(FTestColumnC::StaticStruct(), OrderedColumnMap), 87);
+		CHECK_EQUALS(TEXT("Priority[Column D]"), GetPriorityFromMap(FTestColumnD::StaticStruct(), OrderedColumnMap), 88);
+		CHECK_EQUALS(TEXT("Priority[Column G]"), GetPriorityFromMap(FTestColumnG::StaticStruct(), OrderedColumnMap), 89);
+	}
+
+	SECTION("Order by group with before and after relations")
+	{
+		const TMap<TWeakObjectPtr<const UScriptStruct>, FTedsOutlinerColumnParams> TestColumnParams
+		{
+			{FTestColumnA::StaticStruct(), FTedsOutlinerColumnParams(EColumnPriorityGroup::Right,
+				FColumnPriorityRelation(EColumnPriorityRelation::After, FTestColumnF::StaticStruct()))},
+			{FTestColumnB::StaticStruct(), FTedsOutlinerColumnParams(EColumnPriorityGroup::Middle,
+				FColumnPriorityRelation(EColumnPriorityRelation::After, FTestColumnD::StaticStruct()))},
+			{FTestColumnC::StaticStruct(), FTedsOutlinerColumnParams(EColumnPriorityGroup::Left)},
+			{FTestColumnD::StaticStruct(), FTedsOutlinerColumnParams(EColumnPriorityGroup::Middle)},
+			{FTestColumnE::StaticStruct(), FTedsOutlinerColumnParams(EColumnPriorityGroup::Left,
+				FColumnPriorityRelation(EColumnPriorityRelation::Before, FTestColumnC::StaticStruct()))},
+			{FTestColumnF::StaticStruct(), FTedsOutlinerColumnParams(EColumnPriorityGroup::Right,
+				FColumnPriorityRelation(EColumnPriorityRelation::Before, FTestColumnG::StaticStruct()))},
+			{FTestColumnG::StaticStruct(), FTedsOutlinerColumnParams(EColumnPriorityGroup::Right)},
+		};
+		const FTedsOutlinerColumnDescription TestColumnDescription = FTedsOutlinerColumnDescription(TestColumns, TestColumnParams);
+		OrderedColumnMap.Reset();
+		Helpers::OrderColumns(TestColumns, TestColumnDescription, OrderedColumnMap);
+
+		CHECK_EQUALS(TEXT("Priority Map Size"), OrderedColumnMap.Num(), 7);
+		CHECK_EQUALS(TEXT("Priority[Column E]"), GetPriorityFromMap(FTestColumnE::StaticStruct(), OrderedColumnMap), 0);
+		CHECK_EQUALS(TEXT("Priority[Column C]"), GetPriorityFromMap(FTestColumnC::StaticStruct(), OrderedColumnMap), 1);
+		CHECK_EQUALS(TEXT("Priority[Column D]"), GetPriorityFromMap(FTestColumnD::StaticStruct(), OrderedColumnMap), 85);
+		CHECK_EQUALS(TEXT("Priority[Column B]"), GetPriorityFromMap(FTestColumnB::StaticStruct(), OrderedColumnMap), 86);
+		CHECK_EQUALS(TEXT("Priority[Column F]"), GetPriorityFromMap(FTestColumnF::StaticStruct(), OrderedColumnMap), 170);
+		CHECK_EQUALS(TEXT("Priority[Column A]"), GetPriorityFromMap(FTestColumnA::StaticStruct(), OrderedColumnMap), 171);	
+		CHECK_EQUALS(TEXT("Priority[Column G]"), GetPriorityFromMap(FTestColumnG::StaticStruct(), OrderedColumnMap), 172);
+	}
+
+	SECTION("Order by group with before and after relations on the same column")
+	{
+		const TMap<TWeakObjectPtr<const UScriptStruct>, FTedsOutlinerColumnParams> TestColumnParams
+		{
+			{FTestColumnA::StaticStruct(), FTedsOutlinerColumnParams(EColumnPriorityGroup::Middle)},
+			{FTestColumnB::StaticStruct(), FTedsOutlinerColumnParams(EColumnPriorityGroup::Middle,
+				FColumnPriorityRelation(EColumnPriorityRelation::After, FTestColumnA::StaticStruct()))},
+			{FTestColumnC::StaticStruct(), FTedsOutlinerColumnParams(EColumnPriorityGroup::Middle,
+				FColumnPriorityRelation(EColumnPriorityRelation::After, FTestColumnA::StaticStruct()))},
+			{FTestColumnD::StaticStruct(), FTedsOutlinerColumnParams(EColumnPriorityGroup::Middle,
+				FColumnPriorityRelation(EColumnPriorityRelation::Before, FTestColumnA::StaticStruct()))},
+			{FTestColumnE::StaticStruct(), FTedsOutlinerColumnParams(EColumnPriorityGroup::Middle,
+				FColumnPriorityRelation(EColumnPriorityRelation::Before, FTestColumnA::StaticStruct()))},
+			{FTestColumnF::StaticStruct(), FTedsOutlinerColumnParams(EColumnPriorityGroup::Middle,
+				FColumnPriorityRelation(EColumnPriorityRelation::Before, FTestColumnA::StaticStruct()))},
+			{FTestColumnG::StaticStruct(), FTedsOutlinerColumnParams(EColumnPriorityGroup::Middle,
+				FColumnPriorityRelation(EColumnPriorityRelation::After, FTestColumnA::StaticStruct()))},
+		};
+		const FTedsOutlinerColumnDescription TestColumnDescription = FTedsOutlinerColumnDescription(TestColumns, TestColumnParams);
+		OrderedColumnMap.Reset();
+		Helpers::OrderColumns(TestColumns, TestColumnDescription, OrderedColumnMap);
+
+		CHECK_EQUALS(TEXT("Priority Map Size"), OrderedColumnMap.Num(), 7);
+		CHECK_EQUALS(TEXT("Priority[Column D]"), GetPriorityFromMap(FTestColumnD::StaticStruct(), OrderedColumnMap), 85);
+		CHECK_EQUALS(TEXT("Priority[Column E]"), GetPriorityFromMap(FTestColumnE::StaticStruct(), OrderedColumnMap), 86);
+		CHECK_EQUALS(TEXT("Priority[Column F]"), GetPriorityFromMap(FTestColumnF::StaticStruct(), OrderedColumnMap), 87);
+		CHECK_EQUALS(TEXT("Priority[Column A]"), GetPriorityFromMap(FTestColumnA::StaticStruct(), OrderedColumnMap), 88);
+		CHECK_EQUALS(TEXT("Priority[Column B]"), GetPriorityFromMap(FTestColumnB::StaticStruct(), OrderedColumnMap), 89);
+		CHECK_EQUALS(TEXT("Priority[Column C]"), GetPriorityFromMap(FTestColumnC::StaticStruct(), OrderedColumnMap), 90);
+		CHECK_EQUALS(TEXT("Priority[Column G]"), GetPriorityFromMap(FTestColumnG::StaticStruct(), OrderedColumnMap), 91);
+	}
+	
+	SECTION("Order by group with duplicate column defined parameters")
+    {
+    	const TMap<TWeakObjectPtr<const UScriptStruct>, FTedsOutlinerColumnParams> TestColumnParams
+		{
+			{FTestColumnA::StaticStruct(), FTedsOutlinerColumnParams(EColumnPriorityGroup::Right,
+				FColumnPriorityRelation(EColumnPriorityRelation::After, FTestColumnF::StaticStruct()))},
+			{FTestColumnB::StaticStruct(), FTedsOutlinerColumnParams(EColumnPriorityGroup::Middle,
+				FColumnPriorityRelation(EColumnPriorityRelation::After, FTestColumnD::StaticStruct()))},
+			{FTestColumnC::StaticStruct(), FTedsOutlinerColumnParams(EColumnPriorityGroup::Left)},
+			{FTestColumnD::StaticStruct(), FTedsOutlinerColumnParams(EColumnPriorityGroup::Middle)},
+			{FTestColumnE::StaticStruct(), FTedsOutlinerColumnParams(EColumnPriorityGroup::Left,
+				FColumnPriorityRelation(EColumnPriorityRelation::Before, FTestColumnC::StaticStruct()))},
+			{FTestColumnF::StaticStruct(), FTedsOutlinerColumnParams(EColumnPriorityGroup::Right,
+				FColumnPriorityRelation(EColumnPriorityRelation::Before, FTestColumnG::StaticStruct()))},
+			{FTestColumnG::StaticStruct(), FTedsOutlinerColumnParams(EColumnPriorityGroup::Right)},
+			// Add another definition of TestColumnA's parameters, it should overwrite the previous
+			{FTestColumnA::StaticStruct(), FTedsOutlinerColumnParams(EColumnPriorityGroup::Left,
+				FColumnPriorityRelation(EColumnPriorityRelation::After, FTestColumnE::StaticStruct()))}
+		};
+		const FTedsOutlinerColumnDescription TestColumnDescription = FTedsOutlinerColumnDescription(TestColumns, TestColumnParams);
+		OrderedColumnMap.Reset();
+		Helpers::OrderColumns(TestColumns, TestColumnDescription, OrderedColumnMap);
+
+		CHECK_EQUALS(TEXT("Priority Map Size"), OrderedColumnMap.Num(), 7);
+		CHECK_EQUALS(TEXT("Priority[Column E]"), GetPriorityFromMap(FTestColumnE::StaticStruct(), OrderedColumnMap), 0);
+		CHECK_EQUALS(TEXT("Priority[Column A]"), GetPriorityFromMap(FTestColumnA::StaticStruct(), OrderedColumnMap), 1);	
+		CHECK_EQUALS(TEXT("Priority[Column C]"), GetPriorityFromMap(FTestColumnC::StaticStruct(), OrderedColumnMap), 2);
+		CHECK_EQUALS(TEXT("Priority[Column D]"), GetPriorityFromMap(FTestColumnD::StaticStruct(), OrderedColumnMap), 85);
+		CHECK_EQUALS(TEXT("Priority[Column B]"), GetPriorityFromMap(FTestColumnB::StaticStruct(), OrderedColumnMap), 86);
+		CHECK_EQUALS(TEXT("Priority[Column F]"), GetPriorityFromMap(FTestColumnF::StaticStruct(), OrderedColumnMap), 170);
+		CHECK_EQUALS(TEXT("Priority[Column G]"), GetPriorityFromMap(FTestColumnG::StaticStruct(), OrderedColumnMap), 171);
+    }
+
+	SECTION("Order by group with before relation to column in a different group")
+	{
+		const TMap<TWeakObjectPtr<const UScriptStruct>, FTedsOutlinerColumnParams> TestColumnParams
+		{
+			{FTestColumnA::StaticStruct(), FTedsOutlinerColumnParams(EColumnPriorityGroup::Right)},
+			{FTestColumnB::StaticStruct(), FTedsOutlinerColumnParams(EColumnPriorityGroup::Middle)},
+			{FTestColumnC::StaticStruct(), FTedsOutlinerColumnParams(EColumnPriorityGroup::Left)},
+			{FTestColumnD::StaticStruct(), FTedsOutlinerColumnParams(EColumnPriorityGroup::Middle)},
+			{FTestColumnE::StaticStruct(), FTedsOutlinerColumnParams(EColumnPriorityGroup::Left,
+				FColumnPriorityRelation(EColumnPriorityRelation::Before, FTestColumnD::StaticStruct()))},
+			{FTestColumnF::StaticStruct(), FTedsOutlinerColumnParams(EColumnPriorityGroup::Right)},
+			{FTestColumnG::StaticStruct(), FTedsOutlinerColumnParams(EColumnPriorityGroup::Right)},
+		};
+		const FTedsOutlinerColumnDescription TestColumnDescription = FTedsOutlinerColumnDescription(TestColumns, TestColumnParams);
+		OrderedColumnMap.Reset();
+		Helpers::OrderColumns(TestColumns, TestColumnDescription, OrderedColumnMap);
+
+		CHECK_EQUALS(TEXT("Priority Map Size"), OrderedColumnMap.Num(), 7);
+		CHECK_EQUALS(TEXT("Priority[Column C]"), GetPriorityFromMap(FTestColumnC::StaticStruct(), OrderedColumnMap), 0);
+		CHECK_EQUALS(TEXT("Priority[Column E]"), GetPriorityFromMap(FTestColumnE::StaticStruct(), OrderedColumnMap), 1);
+		CHECK_EQUALS(TEXT("Priority[Column B]"), GetPriorityFromMap(FTestColumnB::StaticStruct(), OrderedColumnMap), 85);
+		CHECK_EQUALS(TEXT("Priority[Column D]"), GetPriorityFromMap(FTestColumnD::StaticStruct(), OrderedColumnMap), 86);
+		CHECK_EQUALS(TEXT("Priority[Column A]"), GetPriorityFromMap(FTestColumnA::StaticStruct(), OrderedColumnMap), 170);
+		CHECK_EQUALS(TEXT("Priority[Column F]"), GetPriorityFromMap(FTestColumnF::StaticStruct(), OrderedColumnMap), 171);
+		CHECK_EQUALS(TEXT("Priority[Column G]"), GetPriorityFromMap(FTestColumnG::StaticStruct(), OrderedColumnMap), 172);
+	}
+
+	// The check is needed for CHECK_ENSURE as it fails on Clang, Mac, and Linux
+#if PLATFORM_WINDOWS && !PLATFORM_COMPILER_CLANG
+	SECTION("Order by group with circular relations")
+	{
+#if !UE_BUILD_SHIPPING
+		//ignore the debugger as all the ensure makes for an annoying debugging experience
+		TGuardValue<bool> IgnoreDebugger(GIgnoreDebugger, true);
+#endif
+		
+		const TMap<TWeakObjectPtr<const UScriptStruct>, FTedsOutlinerColumnParams> TestColumnParams
+		{
+			{FTestColumnA::StaticStruct(), FTedsOutlinerColumnParams(EColumnPriorityGroup::Middle,
+				FColumnPriorityRelation(EColumnPriorityRelation::After, FTestColumnB::StaticStruct()))},
+			{FTestColumnB::StaticStruct(), FTedsOutlinerColumnParams(EColumnPriorityGroup::Middle,
+				FColumnPriorityRelation(EColumnPriorityRelation::After, FTestColumnA::StaticStruct()))},
+			{FTestColumnC::StaticStruct(), FTedsOutlinerColumnParams(EColumnPriorityGroup::Middle)},
+			{FTestColumnD::StaticStruct(), FTedsOutlinerColumnParams(EColumnPriorityGroup::Middle)},
+			{FTestColumnE::StaticStruct(), FTedsOutlinerColumnParams(EColumnPriorityGroup::Right,
+				FColumnPriorityRelation(EColumnPriorityRelation::Before, FTestColumnG::StaticStruct()))},
+			{FTestColumnF::StaticStruct(), FTedsOutlinerColumnParams(EColumnPriorityGroup::Right,
+				FColumnPriorityRelation(EColumnPriorityRelation::Before, FTestColumnG::StaticStruct()))},
+			{FTestColumnG::StaticStruct(), FTedsOutlinerColumnParams(EColumnPriorityGroup::Right)},
+		};
+		const FTedsOutlinerColumnDescription TestColumnDescription = FTedsOutlinerColumnDescription(TestColumns, TestColumnParams);
+		OrderedColumnMap.Reset();
+		
+		CHECK_ENSURE(Helpers::OrderColumns(TestColumns, TestColumnDescription, OrderedColumnMap));
+
+		// Since a circular dependency is present between A & B, they will not be present in the order
+		CHECK_EQUALS(TEXT("Priority Map Size"), OrderedColumnMap.Num(), 5);
+		CHECK_EQUALS(TEXT("Priority[Column A]"), GetPriorityFromMap(FTestColumnA::StaticStruct(), OrderedColumnMap), UINT8_MAX);
+		CHECK_EQUALS(TEXT("Priority[Column B]"), GetPriorityFromMap(FTestColumnB::StaticStruct(), OrderedColumnMap), UINT8_MAX);
+		CHECK_EQUALS(TEXT("Priority[Column C]"), GetPriorityFromMap(FTestColumnC::StaticStruct(), OrderedColumnMap), 85);
+		CHECK_EQUALS(TEXT("Priority[Column D]"), GetPriorityFromMap(FTestColumnD::StaticStruct(), OrderedColumnMap), 86);
+		CHECK_EQUALS(TEXT("Priority[Column E]"), GetPriorityFromMap(FTestColumnE::StaticStruct(), OrderedColumnMap), 170);
+		CHECK_EQUALS(TEXT("Priority[Column F]"), GetPriorityFromMap(FTestColumnF::StaticStruct(), OrderedColumnMap), 171);
+		CHECK_EQUALS(TEXT("Priority[Column G]"), GetPriorityFromMap(FTestColumnG::StaticStruct(), OrderedColumnMap), 172);
+	}
+#endif
+
+	SECTION("Order by group with relations to invalid columns")
+	{
+		const TArray<TWeakObjectPtr<const UScriptStruct>> TestColumnsWithInvalidColumn({
+			FTestColumnA::StaticStruct(),
+			FTestColumnB::StaticStruct(),
+			FTestColumnC::StaticStruct(),
+			FTestColumnD::StaticStruct(),
+			FTestColumnE::StaticStruct(),
+			FTestColumnF::StaticStruct(),
+			FTestColumnG::StaticStruct(),
+			"/Script/TedsStyling.SlateStyleTag"_TypeOptional});
+		
+		const TMap<TWeakObjectPtr<const UScriptStruct>, FTedsOutlinerColumnParams> TestColumnParams
+		{
+			{FTestColumnA::StaticStruct(), FTedsOutlinerColumnParams(EColumnPriorityGroup::Right,
+				FColumnPriorityRelation(EColumnPriorityRelation::After, "/Script/TedsStyling.SlateStyleTag"_TypeOptional))},
+			{FTestColumnB::StaticStruct(), FTedsOutlinerColumnParams(EColumnPriorityGroup::Middle,
+				FColumnPriorityRelation(EColumnPriorityRelation::After, FTestColumnD::StaticStruct()))},
+			{FTestColumnC::StaticStruct(), FTedsOutlinerColumnParams(EColumnPriorityGroup::Left)},
+			{FTestColumnD::StaticStruct(), FTedsOutlinerColumnParams(EColumnPriorityGroup::Middle)},
+			{FTestColumnE::StaticStruct(), FTedsOutlinerColumnParams(EColumnPriorityGroup::Left,
+				FColumnPriorityRelation(EColumnPriorityRelation::Before, FTestColumnC::StaticStruct()))},
+			{FTestColumnF::StaticStruct(), FTedsOutlinerColumnParams(EColumnPriorityGroup::Right,
+				FColumnPriorityRelation(EColumnPriorityRelation::Before, "/Script/TedsStyling.SlateStyleTag"_TypeOptional))},
+			{FTestColumnG::StaticStruct(), FTedsOutlinerColumnParams(EColumnPriorityGroup::Right)},
+			{"/Script/TedsStyling.SlateStyleTag"_TypeOptional, FTedsOutlinerColumnParams(EColumnPriorityGroup::Right,
+				FColumnPriorityRelation(EColumnPriorityRelation::After, FTestColumnG::StaticStruct()))}
+		};
+
+		const FTedsOutlinerColumnDescription TestColumnDescription = FTedsOutlinerColumnDescription(TestColumns, TestColumnParams);
+		OrderedColumnMap.Reset();
+		Helpers::OrderColumns(TestColumns, TestColumnDescription, OrderedColumnMap);
+
+		// The invalid column should be ignored and any relations pointing to it will have it ignored
+		CHECK_EQUALS(TEXT("Priority Map Size"), OrderedColumnMap.Num(), 7);
+		CHECK_EQUALS(TEXT("Priority[Column E]"), GetPriorityFromMap(FTestColumnE::StaticStruct(), OrderedColumnMap), 0);
+		CHECK_EQUALS(TEXT("Priority[Column C]"), GetPriorityFromMap(FTestColumnC::StaticStruct(), OrderedColumnMap), 1);
+		CHECK_EQUALS(TEXT("Priority[Column D]"), GetPriorityFromMap(FTestColumnD::StaticStruct(), OrderedColumnMap), 85);
+		CHECK_EQUALS(TEXT("Priority[Column B]"), GetPriorityFromMap(FTestColumnB::StaticStruct(), OrderedColumnMap), 86);
+		CHECK_EQUALS(TEXT("Priority[Column A]"), GetPriorityFromMap(FTestColumnA::StaticStruct(), OrderedColumnMap), 170);
+		CHECK_EQUALS(TEXT("Priority[Column F]"), GetPriorityFromMap(FTestColumnF::StaticStruct(), OrderedColumnMap), 171);
+		CHECK_EQUALS(TEXT("Priority[Column G]"), GetPriorityFromMap(FTestColumnG::StaticStruct(), OrderedColumnMap), 172);
+	}
+}
+
+#endif // #if WITH_TESTS
