[Static Analysis Warning] Tests CHECK macro expanded and caused an error, can be fixed by having the funciton return a bool

#Synced-CL 51497549

diff --git a/Engine/Plugins/Experimental/EditorDataStorageFeatures/Source/TedsOutliner/Private/TedsOutlinerHelpers.cpp b/Engine/Plugins/Experimental/EditorDataStorageFeatures/Source/TedsOutliner/Private/TedsOutlinerHelpers.cpp
--- a/Engine/Plugins/Experimental/EditorDataStorageFeatures/Source/TedsOutliner/Private/TedsOutlinerHelpers.cpp
+++ b/Engine/Plugins/Experimental/EditorDataStorageFeatures/Source/TedsOutliner/Private/TedsOutlinerHelpers.cpp
@@ -76,7 +76,7 @@ namespace UE::Editor::Outliner::Helpers
 		}
 	}
 	
-	void OrderColumnGroup(const TArray<TWeakObjectPtr<const UScriptStruct>>& InColumnGroup,
+	bool OrderColumnGroup(const TArray<TWeakObjectPtr<const UScriptStruct>>& InColumnGroup,
 		const FTedsOutlinerColumnDescription& InColumnDescription, const uint8 InDefaultStartingPriorityIndex,
 		TMap<TWeakObjectPtr<const UScriptStruct>, uint8>& OutOrderedColumnMap)
 	{
@@ -84,7 +84,7 @@ namespace UE::Editor::Outliner::Helpers
 		// Early return if we didn't receive any columns to order
 		if (InColumnGroup.IsEmpty())
 		{
-			return;
+			return false;
 		}
 		using FColumnPriorityRelation = FTedsOutlinerColumnParams::FColumnPriorityRelation;
 		using EColumnPriorityRelation = FColumnPriorityRelation::EColumnPriorityRelation;
@@ -159,13 +159,6 @@ namespace UE::Editor::Outliner::Helpers
 				}
 			}
 		}
-
-		// If the column sizes are different, there was a circular dependency that could not be resolved,
-		// let the user know that they should reorder their columns.
-		ensureAlwaysMsgf(FinalGroupOrder.Num() == ColumnCount,
-				TEXT("The TEDS Outliner Column order contains a circular priority relation, "
-				"the order will be incorrect until fixed."));
-
 		
 		// Assign Priority Indexes
 		int32 PriorityIndexOffset = 0;
@@ -179,9 +172,19 @@ namespace UE::Editor::Outliner::Helpers
 				++PriorityIndexOffset;
 			}
 		}
+
+		// If the column sizes are different, there was a circular dependency that could not be resolved,
+		// let the user know that they should reorder their columns.
+		if (ensureAlwaysMsgf(FinalGroupOrder.Num() == ColumnCount,
+				TEXT("The TEDS Outliner Column order contains a circular priority relation, "
+				"the order will be incorrect until fixed.")))
+		{
+			return true;
+		}
+		return false;
 	}
 		
-	void OrderColumns(const TArray<TWeakObjectPtr<const UScriptStruct>>& InColumns,
+	bool OrderColumns(const TArray<TWeakObjectPtr<const UScriptStruct>>& InColumns,
 		const FTedsOutlinerColumnDescription& InColumnDescription, TMap<TWeakObjectPtr<const UScriptStruct>, uint8>& OutOrderedColumnMap)
 	{
 		using EColumnPriorityGroup = FTedsOutlinerColumnParams::EColumnPriorityGroup;
@@ -220,9 +223,12 @@ namespace UE::Editor::Outliner::Helpers
 		constexpr uint8 DefaultMiddlePriorityIndex = 85;
 		constexpr uint8 DefaultRightPriorityIndex = 170;
 		
-		OrderColumnGroup(LeftPriorityColumns, InColumnDescription, DefaultLeftPriorityIndex, OutOrderedColumnMap);
-		OrderColumnGroup(MiddlePriorityColumns, InColumnDescription, DefaultMiddlePriorityIndex, OutOrderedColumnMap);
-		OrderColumnGroup(RightPriorityColumns, InColumnDescription, DefaultRightPriorityIndex, OutOrderedColumnMap);
+		bool bColumnsOrderedSuccessfully = true;
+		
+		bColumnsOrderedSuccessfully &= OrderColumnGroup(LeftPriorityColumns, InColumnDescription, DefaultLeftPriorityIndex, OutOrderedColumnMap);
+		bColumnsOrderedSuccessfully &= OrderColumnGroup(MiddlePriorityColumns, InColumnDescription, DefaultMiddlePriorityIndex, OutOrderedColumnMap);
+		bColumnsOrderedSuccessfully &= OrderColumnGroup(RightPriorityColumns, InColumnDescription, DefaultRightPriorityIndex, OutOrderedColumnMap);
+		return bColumnsOrderedSuccessfully;
 	}
 	
 	FName FindOutlinerColumnFromTedsColumns(const DataStorage::ICoreProvider* Storage, TConstArrayView<TWeakObjectPtr<const UScriptStruct>> TEDSColumns)
diff --git a/Engine/Plugins/Experimental/EditorDataStorageFeatures/Source/TedsOutliner/Public/TedsOutlinerHelpers.h b/Engine/Plugins/Experimental/EditorDataStorageFeatures/Source/TedsOutliner/Public/TedsOutlinerHelpers.h
--- a/Engine/Plugins/Experimental/EditorDataStorageFeatures/Source/TedsOutliner/Public/TedsOutlinerHelpers.h
+++ b/Engine/Plugins/Experimental/EditorDataStorageFeatures/Source/TedsOutliner/Public/TedsOutlinerHelpers.h
@@ -26,14 +26,14 @@ namespace UE::Editor::Outliner::Helpers
 	TEDSOUTLINER_API void RefreshLevelEditorOutliners();
 
 	// Orders a given array of Columns within the same priority group and relations defined by the
-	// ColumnDescription and returns evaluated priorities to the OutOrderedColumnMap
-	void TEDSOUTLINER_API OrderColumnGroup(const TArray<TWeakObjectPtr<const UScriptStruct>>& InColumnGroup,
+	// ColumnDescription and returns evaluated priorities to the OutOrderedColumnMap. Returns true if the ordering was successful
+	bool TEDSOUTLINER_API OrderColumnGroup(const TArray<TWeakObjectPtr<const UScriptStruct>>& InColumnGroup,
     	const FTedsOutlinerColumnDescription& InColumnDescription, const uint8 InDefaultStartingPriorityIndex,
     	TMap<TWeakObjectPtr<const UScriptStruct>, uint8>& OutOrderedColumnMap);
 
 	// Orders a given array of Columns using priority groups and relations defined by the
-	// ColumnDescription and returns evaluated priorities to the OutOrderedColumnMap
-	void TEDSOUTLINER_API OrderColumns(const TArray<TWeakObjectPtr<const UScriptStruct>>& InColumns,
+	// ColumnDescription and returns evaluated priorities to the OutOrderedColumnMap. Returns true if the ordering was successful
+	bool TEDSOUTLINER_API OrderColumns(const TArray<TWeakObjectPtr<const UScriptStruct>>& InColumns,
 		const FTedsOutlinerColumnDescription& InColumnDescription, TMap<TWeakObjectPtr<const UScriptStruct>, uint8>& OutOrderedColumnMap);
 	
 	// Get the name of the Outliner column corresponding to the given TEDS column (if any)
