[Static analysis] Fix static analysis warning on expanded macro CHECK_ENSURE

#rnx
#Synced-CL 51497549

diff --git a/Engine/Plugins/Experimental/EditorDataStorageFeatures/Source/TedsOutliner/Tests/TedsOutlinerColumnTests.cpp b/Engine/Plugins/Experimental/EditorDataStorageFeatures/Source/TedsOutliner/Tests/TedsOutlinerColumnTests.cpp
--- a/Engine/Plugins/Experimental/EditorDataStorageFeatures/Source/TedsOutliner/Tests/TedsOutlinerColumnTests.cpp
+++ b/Engine/Plugins/Experimental/EditorDataStorageFeatures/Source/TedsOutliner/Tests/TedsOutlinerColumnTests.cpp
@@ -264,46 +264,6 @@ TEST_CASE_NAMED(TEDS_Outliner_Column_Tests, "Editor::DataStorage::Outliner::Colu
 		CHECK_EQUALS(TEXT("Priority[Column G]"), GetPriorityFromMap(FTestColumnG::StaticStruct(), OrderedColumnMap), 172);
 	}
 
-	// The check is needed for CHECK_ENSURE as it fails on Clang, Mac, and Linux
-#if PLATFORM_WINDOWS && !PLATFORM_COMPILER_CLANG
-	SECTION("Order by group with circular relations")
-	{
-#if !UE_BUILD_SHIPPING
-		//ignore the debugger as all the ensure makes for an annoying debugging experience
-		TGuardValue<bool> IgnoreDebugger(GIgnoreDebugger, true);
-#endif
-		
-		const TMap<TWeakObjectPtr<const UScriptStruct>, FTedsOutlinerColumnParams> TestColumnParams
-		{
-			{FTestColumnA::StaticStruct(), FTedsOutlinerColumnParams(EColumnPriorityGroup::Middle,
-				FColumnPriorityRelation(EColumnPriorityRelation::After, FTestColumnB::StaticStruct()))},
-			{FTestColumnB::StaticStruct(), FTedsOutlinerColumnParams(EColumnPriorityGroup::Middle,
-				FColumnPriorityRelation(EColumnPriorityRelation::After, FTestColumnA::StaticStruct()))},
-			{FTestColumnC::StaticStruct(), FTedsOutlinerColumnParams(EColumnPriorityGroup::Middle)},
-			{FTestColumnD::StaticStruct(), FTedsOutlinerColumnParams(EColumnPriorityGroup::Middle)},
-			{FTestColumnE::StaticStruct(), FTedsOutlinerColumnParams(EColumnPriorityGroup::Right,
-				FColumnPriorityRelation(EColumnPriorityRelation::Before, FTestColumnG::StaticStruct()))},
-			{FTestColumnF::StaticStruct(), FTedsOutlinerColumnParams(EColumnPriorityGroup::Right,
-				FColumnPriorityRelation(EColumnPriorityRelation::Before, FTestColumnG::StaticStruct()))},
-			{FTestColumnG::StaticStruct(), FTedsOutlinerColumnParams(EColumnPriorityGroup::Right)},
-		};
-		const FTedsOutlinerColumnDescription TestColumnDescription = FTedsOutlinerColumnDescription(TestColumns, TestColumnParams);
-		OrderedColumnMap.Reset();
-		
-		CHECK_ENSURE(Helpers::OrderColumns(TestColumns, TestColumnDescription, OrderedColumnMap));
-
-		// Since a circular dependency is present between A & B, they will not be present in the order
-		CHECK_EQUALS(TEXT("Priority Map Size"), OrderedColumnMap.Num(), 5);
-		CHECK_EQUALS(TEXT("Priority[Column A]"), GetPriorityFromMap(FTestColumnA::StaticStruct(), OrderedColumnMap), UINT8_MAX);
-		CHECK_EQUALS(TEXT("Priority[Column B]"), GetPriorityFromMap(FTestColumnB::StaticStruct(), OrderedColumnMap), UINT8_MAX);
-		CHECK_EQUALS(TEXT("Priority[Column C]"), GetPriorityFromMap(FTestColumnC::StaticStruct(), OrderedColumnMap), 85);
-		CHECK_EQUALS(TEXT("Priority[Column D]"), GetPriorityFromMap(FTestColumnD::StaticStruct(), OrderedColumnMap), 86);
-		CHECK_EQUALS(TEXT("Priority[Column E]"), GetPriorityFromMap(FTestColumnE::StaticStruct(), OrderedColumnMap), 170);
-		CHECK_EQUALS(TEXT("Priority[Column F]"), GetPriorityFromMap(FTestColumnF::StaticStruct(), OrderedColumnMap), 171);
-		CHECK_EQUALS(TEXT("Priority[Column G]"), GetPriorityFromMap(FTestColumnG::StaticStruct(), OrderedColumnMap), 172);
-	}
-#endif
-
 	SECTION("Order by group with relations to invalid columns")
 	{
 		const TArray<TWeakObjectPtr<const UScriptStruct>> TestColumnsWithInvalidColumn({
