[TEDS Outliner] Add a Socket Column to TEDS and the TEDS Outliner

#Synced-CL 51541229

diff --git a/Engine/Plugins/Experimental/EditorDataStorageFeatures/Source/TedsActorCompatibility/Private/Compatibility/TedsActorSocketQueries.cpp b/Engine/Plugins/Experimental/EditorDataStorageFeatures/Source/TedsActorCompatibility/Private/Compatibility/TedsActorSocketQueries.cpp
new file mode 100644
--- /dev/null
+++ b/Engine/Plugins/Experimental/EditorDataStorageFeatures/Source/TedsActorCompatibility/Private/Compatibility/TedsActorSocketQueries.cpp
@@ -0,0 +1,73 @@
+// Copyright Epic Games, Inc. All Rights Reserved.
+
+#include "Compatibility/TedsActorSocketQueries.h"
+
+#include "Elements/Columns/TypedElementCompatibilityColumns.h"
+#include "Elements/Columns/TypedElementMiscColumns.h"
+#include "Columns/TedsActorSocketColumns.h"
+#include "Elements/Framework/TypedElementQueryBuilder.h"
+#include "GameFramework/Actor.h"
+
+#include UE_INLINE_GENERATED_CPP_BY_NAME(TedsActorSocketQueries)
+
+void UActorSocketDataStorageFactory::RegisterQueries(UE::Editor::DataStorage::ICoreProvider& DataStorage)
+{
+	RegisterActorAddSocketColumn(DataStorage);
+	RegisterActorSocketToColumnQuery(DataStorage);
+}
+
+void UActorSocketDataStorageFactory::RegisterActorAddSocketColumn(UE::Editor::DataStorage::ICoreProvider& DataStorage) const
+{
+	using namespace UE::Editor::DataStorage::Queries;
+
+	DataStorage.RegisterQuery(
+		Select(
+			TEXT("Add the attached Socket of an Actor to a New Column"),
+			FProcessor(EQueryTickPhase::PrePhysics, DataStorage.GetQueryTickGroupName(EQueryTickGroups::SyncExternalToDataStorage))
+			.SetExecutionMode(EExecutionMode::GameThread),
+			[](IQueryContext& Context, RowHandle Row, const FTypedElementUObjectColumn& Object)
+			{
+				if (const AActor* ActorInstance = Cast<AActor>(Object.Object); ActorInstance != nullptr)
+				{
+					if (const FName SocketName = ActorInstance->GetAttachParentSocketName(); !SocketName.IsNone())
+					{
+						Context.AddColumn(Row, FTedsActorSocketColumn
+							{ .AttachedSocket = SocketName });
+					}
+				}
+			})
+		.Where()
+			.All<FTypedElementActorTag, FTypedElementSyncFromWorldTag>()
+			.None<FTedsActorSocketColumn>()
+		.Compile()
+	);
+}
+
+void UActorSocketDataStorageFactory::RegisterActorSocketToColumnQuery(UE::Editor::DataStorage::ICoreProvider& DataStorage) const
+{
+	using namespace UE::Editor::DataStorage::Queries;
+
+	DataStorage.RegisterQuery(
+		Select(
+			TEXT("Sync Actor's Attached Socket to Column"),
+			FProcessor(EQueryTickPhase::PrePhysics, DataStorage.GetQueryTickGroupName(EQueryTickGroups::SyncExternalToDataStorage))
+			.SetExecutionMode(EExecutionMode::GameThread),
+			[](IQueryContext& Context, RowHandle Row, const FTypedElementUObjectColumn& Object, FTedsActorSocketColumn& SocketColumn)
+			{
+				if (const AActor* ActorInstance = Cast<AActor>(Object.Object); ActorInstance != nullptr)
+				{
+					if (const FName SocketName = ActorInstance->GetAttachParentSocketName(); !SocketName.IsNone())
+					{
+						SocketColumn.AttachedSocket = SocketName;
+					}
+					else
+					{
+						Context.RemoveColumns<FTedsActorSocketColumn>(Row);
+					}
+				}
+			})
+		.Where()
+			.All<FTypedElementActorTag, FTypedElementSyncFromWorldTag>()
+		.Compile()
+	);
+}
\ No newline at end of file
diff --git a/Engine/Plugins/Experimental/EditorDataStorageFeatures/Source/TedsActorCompatibility/Private/Compatibility/TedsActorSocketQueries.h b/Engine/Plugins/Experimental/EditorDataStorageFeatures/Source/TedsActorCompatibility/Private/Compatibility/TedsActorSocketQueries.h
new file mode 100644
--- /dev/null
+++ b/Engine/Plugins/Experimental/EditorDataStorageFeatures/Source/TedsActorCompatibility/Private/Compatibility/TedsActorSocketQueries.h
@@ -0,0 +1,28 @@
+// Copyright Epic Games, Inc. All Rights Reserved.
+
+#pragma once
+
+#include "Elements/Interfaces/TypedElementDataStorageFactory.h"
+
+#include "TedsActorSocketQueries.generated.h"
+
+UCLASS()
+class UActorSocketDataStorageFactory : public UEditorDataStorageFactory
+{
+	GENERATED_BODY()
+
+public:
+	~UActorSocketDataStorageFactory() override = default;
+
+	void RegisterQueries(UE::Editor::DataStorage::ICoreProvider& DataStorage) override;
+private:
+
+	/**
+	 * Adds the socket column to new actors that do not have one already.
+	 */
+	void RegisterActorAddSocketColumn(UE::Editor::DataStorage::ICoreProvider& DataStorage) const;
+	/**
+	 * Takes the socket name set on an actor and copies it to the Data Storage.
+	 */
+	void RegisterActorSocketToColumnQuery(UE::Editor::DataStorage::ICoreProvider& DataStorage) const;
+};
diff --git a/Engine/Plugins/Experimental/EditorDataStorageFeatures/Source/TedsActorCompatibility/Public/Columns/TedsActorSocketColumns.h b/Engine/Plugins/Experimental/EditorDataStorageFeatures/Source/TedsActorCompatibility/Public/Columns/TedsActorSocketColumns.h
new file mode 100644
--- /dev/null
+++ b/Engine/Plugins/Experimental/EditorDataStorageFeatures/Source/TedsActorCompatibility/Public/Columns/TedsActorSocketColumns.h
@@ -0,0 +1,19 @@
+// Copyright Epic Games, Inc. All Rights Reserved.
+
+#pragma once
+
+#include "Elements/Interfaces/TypedElementDataStorageInterface.h"
+
+#include "TedsActorSocketColumns.generated.h"
+
+/**
+ * Column that stores the name of the socket the actor is attached to.
+ */
+USTRUCT(meta = (DisplayName = "Socket"))
+struct FTedsActorSocketColumn final : public FEditorDataStorageColumn
+{
+	GENERATED_BODY()
+
+	UPROPERTY(meta = (Sortable, Searchable))
+	FName AttachedSocket = NAME_None;
+};
\ No newline at end of file
diff --git a/Engine/Plugins/Experimental/EditorDataStorageFeatures/Source/TedsOutliner/Private/Compatibility/SceneOutlinerTedsBridge.cpp b/Engine/Plugins/Experimental/EditorDataStorageFeatures/Source/TedsOutliner/Private/Compatibility/SceneOutlinerTedsBridge.cpp
--- a/Engine/Plugins/Experimental/EditorDataStorageFeatures/Source/TedsOutliner/Private/Compatibility/SceneOutlinerTedsBridge.cpp
+++ b/Engine/Plugins/Experimental/EditorDataStorageFeatures/Source/TedsOutliner/Private/Compatibility/SceneOutlinerTedsBridge.cpp
@@ -10,6 +10,7 @@
 #include "Elements/Columns/TypedElementTypeInfoColumns.h"
 #include "Elements/Columns/TypedElementVisibilityColumns.h"
 #include "Columns/TedsActorMobilityColumns.h"
+#include "Columns/TedsActorSocketColumns.h"
 #include "Elements/Interfaces/TypedElementDataStorageInterface.h"
 #include "Elements/Interfaces/TypedElementDataStorageUiInterface.h"
 #include "SceneOutlinerPublicTypes.h"
@@ -37,6 +38,7 @@ void USceneOutlinerTedsBridgeFactory::PreRegister(UE::Editor::DataStorage::ICore
 	TEDSToOutlinerDefaultColumnMapping.Add(FTypedElementLabelColumn::StaticStruct(), FSceneOutlinerBuiltInColumnTypes::Label());
 	TEDSToOutlinerDefaultColumnMapping.Add(FLevelColumn::StaticStruct(), FSceneOutlinerBuiltInColumnTypes::Level());
 	TEDSToOutlinerDefaultColumnMapping.Add(FTedsPrimaryPackageObjectTag::StaticStruct(), FSceneOutlinerBuiltInColumnTypes::Unsaved());
+	TEDSToOutlinerDefaultColumnMapping.Add(FTedsActorSocketColumn::StaticStruct(), FSceneOutlinerBuiltInColumnTypes::Socket());
 }
 
 void USceneOutlinerTedsBridgeFactory::RegisterTables(UE::Editor::DataStorage::ICoreProvider& DataStorage)
diff --git a/Engine/Plugins/Experimental/EditorDataStorageFeatures/Source/TedsOutliner/Private/TedsOutlinerModule.cpp b/Engine/Plugins/Experimental/EditorDataStorageFeatures/Source/TedsOutliner/Private/TedsOutlinerModule.cpp
--- a/Engine/Plugins/Experimental/EditorDataStorageFeatures/Source/TedsOutliner/Private/TedsOutlinerModule.cpp
+++ b/Engine/Plugins/Experimental/EditorDataStorageFeatures/Source/TedsOutliner/Private/TedsOutlinerModule.cpp
@@ -20,6 +20,7 @@
 #include "Elements/Columns/TypedElementVisibilityColumns.h"
 #include "Columns/TedsActorComponentCompatibilityColumns.h"
 #include "Columns/TedsActorMobilityColumns.h"
+#include "Columns/TedsActorSocketColumns.h"
 #include "TedsAlertColumns.h"
 #include "DataStorage/Features.h"
 #include "Elements/Interfaces/TypedElementDataStorageInterface.h"
@@ -153,6 +154,7 @@ FTedsOutlinerColumnDescription FTedsOutlinerModule::GetLevelEditorTedsOutlinerCo
 		FTypedElementClassTypeInfoColumn::StaticStruct(),
 		FLevelColumn::StaticStruct(), // PackageShortName
 		FTedsActorMobilityColumn::StaticStruct(),
+		FTedsActorSocketColumn::StaticStruct(),
 		FUObjectIdNameColumn::StaticStruct(),
 		FAlertColumn::StaticStruct(),
 		FChildAlertColumn::StaticStruct()});
diff --git a/Engine/Plugins/Experimental/EditorDataStorageFeatures/Source/TedsOutliner/Private/Widgets/SocketWidget.cpp b/Engine/Plugins/Experimental/EditorDataStorageFeatures/Source/TedsOutliner/Private/Widgets/SocketWidget.cpp
new file mode 100644
--- /dev/null
+++ b/Engine/Plugins/Experimental/EditorDataStorageFeatures/Source/TedsOutliner/Private/Widgets/SocketWidget.cpp
@@ -0,0 +1,64 @@
+// Copyright Epic Games, Inc. All Rights Reserved.
+
+#include "Widgets/SocketWidget.h"
+
+#include "Columns/TedsActorSocketColumns.h"
+#include "Elements/Columns/TypedElementCompatibilityColumns.h"
+#include "Elements/Framework/TypedElementAttributeBinding.h"
+#include "Widgets/Text/STextBlock.h"
+#include "Widgets/SBoxPanel.h"
+
+#include UE_INLINE_GENERATED_CPP_BY_NAME(SocketWidget)
+
+#define LOCTEXT_NAMESPACE "SocketWidget"
+
+void USocketWidgetFactory::RegisterWidgetConstructors(UE::Editor::DataStorage::ICoreProvider& DataStorage,
+	UE::Editor::DataStorage::IUiProvider& DataStorageUi) const
+{
+	using namespace UE::Editor::DataStorage::Queries;
+
+	DataStorageUi.RegisterWidgetFactory<FSocketWidgetConstructor>(
+	DataStorageUi.FindPurpose(DataStorageUi.GetGeneralWidgetPurposeID()),
+		TColumn<FTedsActorSocketColumn>());
+}
+
+FSocketWidgetConstructor::FSocketWidgetConstructor()
+	: FSimpleWidgetConstructor(StaticStruct())
+{
+}
+
+TSharedPtr<SWidget> FSocketWidgetConstructor::CreateWidget(
+	UE::Editor::DataStorage::ICoreProvider* DataStorage,
+	UE::Editor::DataStorage::IUiProvider* DataStorageUi,
+	UE::Editor::DataStorage::RowHandle TargetRow,
+	UE::Editor::DataStorage::RowHandle WidgetRow,
+	const UE::Editor::DataStorage::FMetaDataView& Arguments)
+{
+	using namespace UE::Editor::DataStorage;
+	if (!DataStorage->IsRowAvailable(TargetRow))
+	{
+		return SNew(STextBlock)
+			.Text(LOCTEXT("MissingRowReferenceColumn", "Unable to retrieve row reference."));
+	}
+
+	FAttributeBinder Binder(TargetRow, DataStorage);
+
+	return SNew(SHorizontalBox)
+			+ SHorizontalBox::Slot()
+			.AutoWidth()
+			.VAlign(VAlign_Center)
+			.Padding(8, 0, 0, 0)
+			[
+				SNew(STextBlock)
+					.ColorAndOpacity(FSlateColor::UseSubduedForeground())
+					.Text(Binder.BindText(&FTedsActorSocketColumn::AttachedSocket))
+			];
+}
+
+FText FSocketWidgetConstructor::CreateWidgetDisplayNameText(UE::Editor::DataStorage::ICoreProvider* DataStorage,
+	UE::Editor::DataStorage::RowHandle Row) const
+{
+	return LOCTEXT("SocketHeaderDisplayName", "Socket");
+}
+
+#undef LOCTEXT_NAMESPACE
diff --git a/Engine/Plugins/Experimental/EditorDataStorageFeatures/Source/TedsOutliner/Public/Widgets/SocketWidget.h b/Engine/Plugins/Experimental/EditorDataStorageFeatures/Source/TedsOutliner/Public/Widgets/SocketWidget.h
new file mode 100644
--- /dev/null
+++ b/Engine/Plugins/Experimental/EditorDataStorageFeatures/Source/TedsOutliner/Public/Widgets/SocketWidget.h
@@ -0,0 +1,45 @@
+// Copyright Epic Games, Inc. All Rights Reserved.
+
+#pragma once
+
+#include "Elements/Interfaces/TypedElementDataStorageFactory.h"
+#include "Elements/Interfaces/TypedElementDataStorageInterface.h"
+#include "Elements/Interfaces/TypedElementDataStorageUiInterface.h"
+
+#include "SocketWidget.generated.h"
+
+#define UE_API TEDSOUTLINER_API
+
+UCLASS()
+class USocketWidgetFactory : public UEditorDataStorageFactory
+{
+	GENERATED_BODY()
+
+public:
+	~USocketWidgetFactory() override = default;
+
+	void RegisterWidgetConstructors(UE::Editor::DataStorage::ICoreProvider& DataStorage,
+		UE::Editor::DataStorage::IUiProvider& DataStorageUi) const override;
+};
+
+USTRUCT()
+struct FSocketWidgetConstructor : public FSimpleWidgetConstructor
+{
+	GENERATED_BODY()
+
+public:
+	UE_API FSocketWidgetConstructor();
+	~FSocketWidgetConstructor() override = default;
+
+	UE_API virtual TSharedPtr<SWidget> CreateWidget(
+		UE::Editor::DataStorage::ICoreProvider* DataStorage,
+		UE::Editor::DataStorage::IUiProvider* DataStorageUi,
+		UE::Editor::DataStorage::RowHandle TargetRow,
+		UE::Editor::DataStorage::RowHandle WidgetRow,
+		const UE::Editor::DataStorage::FMetaDataView& Arguments) override;
+
+	UE_API virtual FText CreateWidgetDisplayNameText(
+		UE::Editor::DataStorage::ICoreProvider* DataStorage, UE::Editor::DataStorage::RowHandle Row = UE::Editor::DataStorage::InvalidRowHandle) const override;
+};
+
+#undef UE_API
