TedsUI: Change the OnSelection event on TedsTableViewer and TedsHierarchyViewer to also include the selection info type, allowing more fine grained control over handling selection change events based on cause.

#Synced-CL 51607007

diff --git a/Engine/Plugins/Experimental/EditorDataStorageFeatures/Source/TedsContentBrowser/Private/TedsContentBrowserModule.cpp b/Engine/Plugins/Experimental/EditorDataStorageFeatures/Source/TedsContentBrowser/Private/TedsContentBrowserModule.cpp
--- a/Engine/Plugins/Experimental/EditorDataStorageFeatures/Source/TedsContentBrowser/Private/TedsContentBrowserModule.cpp
+++ b/Engine/Plugins/Experimental/EditorDataStorageFeatures/Source/TedsContentBrowser/Private/TedsContentBrowserModule.cpp
@@ -236,7 +236,7 @@ namespace UE::Editor::ContentBrowser
 						DynamicSkeletalMeshSkeletonColumn, DynamicBlueprintColumn, DynamicSourceTextureColumn,
 						DynamicPhysicsAssetColumn, DynamicShadowPhysicsAssetColumn })
 					.ListSelectionMode(ESelectionMode::Multi)
-					.OnSelectionChanged_Lambda([this](RowHandle Row)
+					.OnSelectionChanged_Lambda([this](RowHandle Row, ESelectInfo::Type SelectInfo)
 					{
 						UpdateAssetPreviewTargetRow();
 
diff --git a/Engine/Plugins/Experimental/EditorDataStorageFeatures/Source/TedsDebugger/Private/Widgets/QueryEditor/TedsQueryEditorResultsView.cpp b/Engine/Plugins/Experimental/EditorDataStorageFeatures/Source/TedsDebugger/Private/Widgets/QueryEditor/TedsQueryEditorResultsView.cpp
--- a/Engine/Plugins/Experimental/EditorDataStorageFeatures/Source/TedsDebugger/Private/Widgets/QueryEditor/TedsQueryEditorResultsView.cpp
+++ b/Engine/Plugins/Experimental/EditorDataStorageFeatures/Source/TedsDebugger/Private/Widgets/QueryEditor/TedsQueryEditorResultsView.cpp
@@ -227,7 +227,7 @@ namespace UE::Editor::DataStorage::Debug::QueryEditor
 					.QueryStack(RowQueryStack)
 					.EmptyRowsMessage(LOCTEXT("EmptyRowsMessage", "The provided query has no results."))
 					.OnSelectionChanged(STedsTableViewer::FOnSelectionChanged::CreateLambda(
-						[this](RowHandle SelectedRow)
+						[this](RowHandle SelectedRow, ESelectInfo::Type SelectInfo)
 							{
 								if(RowDetailsWidget)
 								{
@@ -252,7 +252,7 @@ namespace UE::Editor::DataStorage::Debug::QueryEditor
 					.AllNodeProvider(RowQueryStack)
 					.EmptyRowsMessage(LOCTEXT("EmptyRowsMessage", "The provided query has no results."))
 					.OnSelectionChanged(STedsTableViewer::FOnSelectionChanged::CreateLambda(
-						[this](RowHandle SelectedRow)
+						[this](RowHandle SelectedRow, ESelectInfo::Type SelectInfo)
 							{
 								if(RowDetailsWidget)
 								{
diff --git a/Engine/Plugins/Experimental/EditorDataStorageFeatures/Source/TedsDebugger/Private/Widgets/STedsDiscoveryDebugger.cpp b/Engine/Plugins/Experimental/EditorDataStorageFeatures/Source/TedsDebugger/Private/Widgets/STedsDiscoveryDebugger.cpp
--- a/Engine/Plugins/Experimental/EditorDataStorageFeatures/Source/TedsDebugger/Private/Widgets/STedsDiscoveryDebugger.cpp
+++ b/Engine/Plugins/Experimental/EditorDataStorageFeatures/Source/TedsDebugger/Private/Widgets/STedsDiscoveryDebugger.cpp
@@ -198,7 +198,7 @@ namespace UE::Editor::DataStorage::Debug
 			.QueryStack(RowQueryStack)
 			.EmptyRowsMessage(LOCTEXT("EmptyRowsMessage", "No results for the current search."))
 			.OnSelectionChanged(STedsTableViewer::FOnSelectionChanged::CreateLambda(
-				[this](RowHandle SelectedRow)
+				[this](RowHandle SelectedRow, ESelectInfo::Type SelectInfo)
 					{
 						if(RowDetailsWidget)
 						{
diff --git a/Engine/Plugins/Experimental/EditorDataStorageFeatures/Source/TedsEverythingPicker/Public/Context/TedsPickerContextUtil.h b/Engine/Plugins/Experimental/EditorDataStorageFeatures/Source/TedsEverythingPicker/Public/Context/TedsPickerContextUtil.h
--- a/Engine/Plugins/Experimental/EditorDataStorageFeatures/Source/TedsEverythingPicker/Public/Context/TedsPickerContextUtil.h
+++ b/Engine/Plugins/Experimental/EditorDataStorageFeatures/Source/TedsEverythingPicker/Public/Context/TedsPickerContextUtil.h
@@ -18,7 +18,7 @@ namespace UE::Editor::DataStorage::QueryStack
 
 namespace UE::Editor::DataStorage::Picker
 {
-	DECLARE_DELEGATE_OneParam(FOnSelectionChanged, RowHandle);
+	DECLARE_DELEGATE_TwoParams(FOnSelectionChanged, RowHandle, ESelectInfo::Type);
 
 	class SObjectReferenceContextView : public SCompoundWidget
 	{
diff --git a/Engine/Plugins/Experimental/EditorDataStorageFeatures/Source/TedsTableViewer/Private/Widgets/STedsHierarchyViewer.cpp b/Engine/Plugins/Experimental/EditorDataStorageFeatures/Source/TedsTableViewer/Private/Widgets/STedsHierarchyViewer.cpp
--- a/Engine/Plugins/Experimental/EditorDataStorageFeatures/Source/TedsTableViewer/Private/Widgets/STedsHierarchyViewer.cpp
+++ b/Engine/Plugins/Experimental/EditorDataStorageFeatures/Source/TedsTableViewer/Private/Widgets/STedsHierarchyViewer.cpp
@@ -348,7 +348,7 @@ namespace UE::Editor::DataStorage
 	{
 		if(OnSelectionChanged.IsBound())
 		{
-			OnSelectionChanged.Execute(Item);
+			OnSelectionChanged.Execute(Item, SelectInfo);
 		}
 
 		// Some systems can only access the table viewer by the row handle so also fire the selection changed event on the column
diff --git a/Engine/Plugins/Experimental/EditorDataStorageFeatures/Source/TedsTableViewer/Private/Widgets/STedsTableViewer.cpp b/Engine/Plugins/Experimental/EditorDataStorageFeatures/Source/TedsTableViewer/Private/Widgets/STedsTableViewer.cpp
--- a/Engine/Plugins/Experimental/EditorDataStorageFeatures/Source/TedsTableViewer/Private/Widgets/STedsTableViewer.cpp
+++ b/Engine/Plugins/Experimental/EditorDataStorageFeatures/Source/TedsTableViewer/Private/Widgets/STedsTableViewer.cpp
@@ -245,7 +245,7 @@ namespace UE::Editor::DataStorage
 	{
 		if(OnSelectionChanged.IsBound())
 		{
-			OnSelectionChanged.Execute(Item);
+			OnSelectionChanged.Execute(Item, SelectInfo);
 		}
 
 		// Some systems can only access the table viewer by the row handle so also fire the selection changed event on the column
diff --git a/Engine/Plugins/Experimental/EditorDataStorageFeatures/Source/TedsTableViewer/Public/Widgets/STedsHierarchyViewer.h b/Engine/Plugins/Experimental/EditorDataStorageFeatures/Source/TedsTableViewer/Public/Widgets/STedsHierarchyViewer.h
--- a/Engine/Plugins/Experimental/EditorDataStorageFeatures/Source/TedsTableViewer/Public/Widgets/STedsHierarchyViewer.h
+++ b/Engine/Plugins/Experimental/EditorDataStorageFeatures/Source/TedsTableViewer/Public/Widgets/STedsHierarchyViewer.h
@@ -61,7 +61,7 @@ namespace UE::Editor::DataStorage
 		};
 
 		// Delegate fired when the selection in the table viewer changes
-		DECLARE_DELEGATE_OneParam(FOnSelectionChanged, RowHandle)
+		DECLARE_DELEGATE_TwoParams(FOnSelectionChanged, RowHandle, ESelectInfo::Type)
 
 		SLATE_BEGIN_ARGS(SHierarchyViewer)
 			: _CanSelectGeneratedColumn(true)
diff --git a/Engine/Plugins/Experimental/EditorDataStorageFeatures/Source/TedsTableViewer/Public/Widgets/STedsTableViewer.h b/Engine/Plugins/Experimental/EditorDataStorageFeatures/Source/TedsTableViewer/Public/Widgets/STedsTableViewer.h
--- a/Engine/Plugins/Experimental/EditorDataStorageFeatures/Source/TedsTableViewer/Public/Widgets/STedsTableViewer.h
+++ b/Engine/Plugins/Experimental/EditorDataStorageFeatures/Source/TedsTableViewer/Public/Widgets/STedsTableViewer.h
@@ -46,7 +46,7 @@ namespace UE::Editor::DataStorage
 	public:
 		
 		// Delegate fired when the selection in the table viewer changes
-		DECLARE_DELEGATE_OneParam(FOnSelectionChanged, RowHandle)
+		DECLARE_DELEGATE_TwoParams(FOnSelectionChanged, RowHandle, ESelectInfo::Type)
 
 		SLATE_BEGIN_ARGS(STedsTableViewer)
 			: _CanSelectGeneratedColumn(true)
diff --git a/Engine/Plugins/Experimental/InstanceDataObjectFixupTool/Source/Private/DataRecoveryTool.cpp b/Engine/Plugins/Experimental/InstanceDataObjectFixupTool/Source/Private/DataRecoveryTool.cpp
--- a/Engine/Plugins/Experimental/InstanceDataObjectFixupTool/Source/Private/DataRecoveryTool.cpp
+++ b/Engine/Plugins/Experimental/InstanceDataObjectFixupTool/Source/Private/DataRecoveryTool.cpp
@@ -424,7 +424,7 @@ TSharedRef<SWidget> SDataRecoveryTool::MakePanel(const FText& InTitle, const TSh
 		];
 }
 
-void SDataRecoveryTool::OnListSelectionChanged(UE::Editor::DataStorage::RowHandle RowHandle)
+void SDataRecoveryTool::OnListSelectionChanged(UE::Editor::DataStorage::RowHandle RowHandle, ESelectInfo::Type Info)
 {
 	using namespace UE::Editor::DataStorage;
 
diff --git a/Engine/Plugins/Experimental/InstanceDataObjectFixupTool/Source/Public/DataRecoveryTool.h b/Engine/Plugins/Experimental/InstanceDataObjectFixupTool/Source/Public/DataRecoveryTool.h
--- a/Engine/Plugins/Experimental/InstanceDataObjectFixupTool/Source/Public/DataRecoveryTool.h
+++ b/Engine/Plugins/Experimental/InstanceDataObjectFixupTool/Source/Public/DataRecoveryTool.h
@@ -50,7 +50,7 @@ private:
 
 	UE_API TSharedRef<SWidget> MakeClassesPanel();
 	UE_API TSharedRef<SWidget> MakeTEDSTableViewer();
-	UE_API void OnListSelectionChanged(UE::Editor::DataStorage::RowHandle RowItem);
+	UE_API void OnListSelectionChanged(UE::Editor::DataStorage::RowHandle RowItem, ESelectInfo::Type Info);
 
 	void SetSelectedClass();
 
diff --git a/Engine/Plugins/Experimental/MeshPartition/Source/MeshPartitionEditorUI/Private/UI/MeshPartitionSettingsWidget.cpp b/Engine/Plugins/Experimental/MeshPartition/Source/MeshPartitionEditorUI/Private/UI/MeshPartitionSettingsWidget.cpp
--- a/Engine/Plugins/Experimental/MeshPartition/Source/MeshPartitionEditorUI/Private/UI/MeshPartitionSettingsWidget.cpp
+++ b/Engine/Plugins/Experimental/MeshPartition/Source/MeshPartitionEditorUI/Private/UI/MeshPartitionSettingsWidget.cpp
@@ -685,7 +685,7 @@ TSharedRef<UE::Editor::DataStorage::STedsCompositeHierarchyViewer> SMegaMeshSett
 			.CellWidgetPurpose(CellWidgetPurpose)
 			.HeaderWidgetPurpose(HeaderWidgetPurpose)
 			.PrimaryColumn(FTypedElementLabelColumn::StaticStruct())
-			.OnSelectionChanged_Lambda([this](RowHandle Row)
+			.OnSelectionChanged_Lambda([this](RowHandle Row, ESelectInfo::Type SelectInfo)
 				{
 					OnOutlinerSelectionChanged();
 				})
