[TEDS Filters] Refactor to allow TedsFilters to get passed in as SharedPtrs instead of TedsFilterData. Allows for more customization as the base TedsFilter can be inherited and built on for custom workflows

#Synced-CL 51586067

diff --git a/Engine/Plugins/Experimental/EditorDataStorageFeatures/Source/TedsOutliner/Private/TedsOutlinerFilter.cpp b/Engine/Plugins/Experimental/EditorDataStorageFeatures/Source/TedsOutliner/Private/TedsOutlinerFilter.cpp
--- a/Engine/Plugins/Experimental/EditorDataStorageFeatures/Source/TedsOutliner/Private/TedsOutlinerFilter.cpp
+++ b/Engine/Plugins/Experimental/EditorDataStorageFeatures/Source/TedsOutliner/Private/TedsOutlinerFilter.cpp
@@ -16,16 +16,9 @@ namespace UE::Editor::Outliner
 {
 FTedsOutlinerFilter::FTedsOutlinerFilter(const FName& InFilterName, const FText& InFilterDisplayName, const FText& InFilterToolTip, const FName& InFilterIconName, 
 	const TSharedPtr<FFilterCategory>& InCategory, const QueryHandle& InFilterQuery, const bool bInteractiveFilter)
-	: FFilterBase(InCategory)
-	, FilterName(InFilterName)
-	, FilterDisplayName(InFilterDisplayName)
-	, FilterToolTip(InFilterToolTip)
-	, FilterIconName(InFilterIconName)
-	, bIsClassFilter(false)
+	: FTedsFilterBase(InFilterName, InFilterDisplayName, InFilterToolTip, InFilterIconName, InCategory, InFilterQuery)
 	, bInteractiveFilter(bInteractiveFilter)
-{
-	FilterQuery.Set<QueryHandle>(Helpers::CheckValidFilterQueryHandle(InFilterQuery) ? InFilterQuery : InvalidQueryHandle);
-}
+{}
 
 FTedsOutlinerFilter::FTedsOutlinerFilter(const FName& InFilterName, const FText& InFilterDisplayName, const QueryHandle& InFilterQuery,
 	const bool bInteractiveFilter)
@@ -36,16 +29,9 @@ FTedsOutlinerFilter::FTedsOutlinerFilter(const FName& InFilterName, const FText&
 FTedsOutlinerFilter::FTedsOutlinerFilter(const FName& InFilterName, const FText& InFilterDisplayName, const FText& InFilterToolTip, 
 	const FName& InFilterIconName, const TSharedPtr<FFilterCategory>& InCategory, const TQueryFunction<bool>& InFilterQuery,
 	const bool bInteractiveFilter)
-	: FFilterBase(InCategory)
-	, FilterName(InFilterName)
-	, FilterDisplayName(InFilterDisplayName)
-	, FilterToolTip(InFilterToolTip)
-	, FilterIconName(InFilterIconName)
-	, bIsClassFilter(false)
+	: FTedsFilterBase(InFilterName, InFilterDisplayName, InFilterToolTip, InFilterIconName, InCategory, InFilterQuery)
 	, bInteractiveFilter(bInteractiveFilter)
-{
-	FilterQuery.Set<TQueryFunction<bool>>(InFilterQuery);
-}
+{}
 
 FTedsOutlinerFilter::FTedsOutlinerFilter(const FName& InFilterName, const FText& InFilterDisplayName, const TQueryFunction<bool>& InFilterQuery,
 	const bool bInteractiveFilter)
@@ -54,64 +40,15 @@ FTedsOutlinerFilter::FTedsOutlinerFilter(const FName& InFilterName, const FText&
 {}
 
 FTedsOutlinerFilter::FTedsOutlinerFilter(const UClass* InClass, const TSharedPtr<FFilterCategory>& InCategory, const bool bInteractiveFilter)
-	: FFilterBase(InCategory)
-	, FilterName(InClass->GetFName())
-	, FilterDisplayName(InClass->GetDisplayNameText())
-	, FilterToolTip(FText::Format(LOCTEXT("FilterClassTooltip", "Filter by {0}"), InClass->GetDisplayNameText()))
-	, FilterIconName(FSlateIconFinder::FindIconForClass(InClass).GetStyleName())
-	, bIsClassFilter(true)
+	: FTedsFilterBase(InClass, InCategory)
 	, bInteractiveFilter(bInteractiveFilter)
-{
-	FilterQuery.Set<TQueryFunction<bool>>(
-		Queries::BuildQueryFunction<bool>([InClass](TQueryContext<RowBatchInfo> Context, TResult<bool>& Result, TConstBatch<FTypedElementClassTypeInfoColumn> TypeInfoColumns)
-		{
-			Context.ForEachRow([&Result, InClass](const FTypedElementClassTypeInfoColumn& TypeInfoColumn)
-			{
-				Result.Add(TypeInfoColumn.TypeInfo->IsChildOf(InClass));
-			}, TypeInfoColumns);
-		}));
-}
+{}
 
 void FTedsOutlinerFilter::SetSceneOutlinerImpl(const TSharedPtr<FTedsOutlinerImpl>& InTedsOutlinerImpl)
 {
 	TedsOutlinerImpl = InTedsOutlinerImpl;
 }
 
-void FTedsOutlinerFilter::SetCategory(const TSharedPtr<FFilterCategory>& InCategory)
-{
-	FilterCategory = InCategory;
-}
-	
-FString FTedsOutlinerFilter::GetName() const
-{
-	return FilterName.ToString();
-}
-
-FText FTedsOutlinerFilter::GetDisplayName() const
-{
-	return FilterDisplayName;
-}
-
-FText FTedsOutlinerFilter::GetToolTipText() const
-{
-	return FilterToolTip;
-}
-
-FLinearColor FTedsOutlinerFilter::GetColor() const
-{
-	return FLinearColor();	
-}
-
-FName FTedsOutlinerFilter::GetIconName() const
-{
-	return FilterIconName;
-}
-
-bool FTedsOutlinerFilter::IsInverseFilter() const
-{
-	return false;
-}
-
 void FTedsOutlinerFilter::ActiveStateChanged(bool bActive)
 {
 	using namespace UE::Editor::DataStorage::Queries;
@@ -153,21 +90,6 @@ void FTedsOutlinerFilter::ActiveStateChanged(bool bActive)
 	}
 }
 
-void FTedsOutlinerFilter::ModifyContextMenu(FMenuBuilder& MenuBuilder)
-{
-	
-}
-
-void FTedsOutlinerFilter::SaveSettings(const FString& IniFilename, const FString& IniSection, const FString& SettingsString) const
-{
-	
-}
-
-void FTedsOutlinerFilter::LoadSettings(const FString& IniFilename, const FString& IniSection, const FString& SettingsString)
-{
-	
-}
-
 bool FTedsOutlinerFilter::PassesFilter(SceneOutliner::FilterBarType InItem) const
 {
 	if (TSharedPtr<FTedsOutlinerImpl> TedsOutlinerImplPin = TedsOutlinerImpl.Pin())
diff --git a/Engine/Plugins/Experimental/EditorDataStorageFeatures/Source/TedsOutliner/Public/TedsOutlinerFilter.h b/Engine/Plugins/Experimental/EditorDataStorageFeatures/Source/TedsOutliner/Public/TedsOutlinerFilter.h
--- a/Engine/Plugins/Experimental/EditorDataStorageFeatures/Source/TedsOutliner/Public/TedsOutlinerFilter.h
+++ b/Engine/Plugins/Experimental/EditorDataStorageFeatures/Source/TedsOutliner/Public/TedsOutlinerFilter.h
@@ -7,7 +7,7 @@
 #include <type_traits>
 #include "Filters/FilterBase.h"
 #include "SceneOutlinerPublicTypes.h"
-#include "DataStorage/Queries/Description.h"
+#include "TedsFilter.h"
 #include "Elements/Framework/TypedElementQueryBuilder.h"
 
 namespace UE::Editor::Outliner
@@ -15,7 +15,7 @@ namespace UE::Editor::Outliner
 	using namespace UE::Editor::DataStorage::Queries;
 	class FTedsOutlinerImpl;
 		
-	class FTedsOutlinerFilter : public FFilterBase<SceneOutliner::FilterBarType>
+	class FTedsOutlinerFilter : public FTedsFilterBase<SceneOutliner::FilterBarType>
 	{
 		friend class FTedsOutlinerImpl;
 		
@@ -125,43 +125,10 @@ namespace UE::Editor::Outliner
 		 */
 		TEDSOUTLINER_API FTedsOutlinerFilter(const UClass* InClass, const TSharedPtr<FFilterCategory>& InCategory = nullptr, const bool bInteractiveFilter = true);
 
-		virtual ~FTedsOutlinerFilter() override = default;
-
-		void SetCategory(const TSharedPtr<FFilterCategory>& InCategory);
-		
-		bool IsClassFilter() const { return bIsClassFilter; }
-
 		bool IsInteractiveFilter() const { return bInteractiveFilter; }
-		
-		/** Returns the system name for this filter */
-		virtual TEDSOUTLINER_API FString GetName() const override;
-
-		/** Returns the human-readable name for this filter */
-		virtual TEDSOUTLINER_API FText GetDisplayName() const override;
-
-		/** Returns the tooltip for this filter, shown in the filters menu */
-		virtual TEDSOUTLINER_API FText GetToolTipText() const override;
-
-		/** Returns the color this filter button will be when displayed as a button */
-		virtual TEDSOUTLINER_API FLinearColor GetColor() const override;
-
-		/** Returns the name of the icon to use in menu entries */
-		virtual TEDSOUTLINER_API FName GetIconName() const override;
-
-		/** If true, the filter will be active in the FilterBar when it is inactive in the UI (i.e the filter pill is grayed out) */
-		virtual TEDSOUTLINER_API bool IsInverseFilter() const override;
 
 		/** Notification that the filter became active or inactive */
 		virtual TEDSOUTLINER_API void ActiveStateChanged(bool bActive) override;
-
-		/** Called when the right-click context menu is being built for this filter */
-		virtual TEDSOUTLINER_API void ModifyContextMenu(FMenuBuilder& MenuBuilder) override;
-		
-		/** Can be overriden for custom FilterBar subclasses to save settings, currently not implemented in any generic Filter Bar */
-		virtual TEDSOUTLINER_API void SaveSettings(const FString& IniFilename, const FString& IniSection, const FString& SettingsString) const override;
-
-		/** Can be overridden for custom FilterBar subclasses to load settings, currently not implemented in any generic Filter Bar */
-		virtual TEDSOUTLINER_API void LoadSettings(const FString& IniFilename, const FString& IniSection, const FString& SettingsString) override;
 		
 		/** Returns whether the specified Item passes the Filter's restrictions */
 		virtual TEDSOUTLINER_API bool PassesFilter(SceneOutliner::FilterBarType InItem ) const override;
@@ -169,17 +136,10 @@ namespace UE::Editor::Outliner
 	protected:
 		/** Sets a SceneOutlinerImpl to be used by this Filter */
 		void SetSceneOutlinerImpl(const TSharedPtr<FTedsOutlinerImpl>& InTedsOutlinerImpl);
-	
-		FName FilterName;
-		FText FilterDisplayName;
-		FText FilterToolTip;
-		FName FilterIconName;
-
-		const bool bIsClassFilter;
+		
 		const bool bInteractiveFilter;
 
 		TWeakPtr<FTedsOutlinerImpl> TedsOutlinerImpl;
-		TVariant<QueryHandle, TQueryFunction<bool>> FilterQuery;
 	};
 } // namespace UE::Editor::Outliner
 
diff --git a/Engine/Plugins/Experimental/EditorDataStorageFeatures/Source/TedsTableViewer/Private/TedsFilter.cpp b/Engine/Plugins/Experimental/EditorDataStorageFeatures/Source/TedsTableViewer/Private/TedsFilter.cpp
--- a/Engine/Plugins/Experimental/EditorDataStorageFeatures/Source/TedsTableViewer/Private/TedsFilter.cpp
+++ b/Engine/Plugins/Experimental/EditorDataStorageFeatures/Source/TedsTableViewer/Private/TedsFilter.cpp
@@ -8,79 +8,9 @@ namespace UE::Editor::DataStorage
 {
 	class STedsFilterBar;
 
-	FTedsFilterData::FTedsFilterData(const FName& InFilterName, const FText& InFilterDisplayName, const FText& InFilterToolTip, const FName& InFilterIconName,
-    	const TSharedPtr<FFilterCategory>& InCategory, const QueryHandle& InFilterQuery)
-    	: FilterName(InFilterName)
-    	, FilterDisplayName(InFilterDisplayName)
-    	, FilterToolTip(InFilterToolTip)
-    	, FilterIconName(InFilterIconName)
-    	, FilterCategory(InCategory)
-    {
-    	FilterQuery.Set<QueryHandle>(CheckValidFilterQueryHandle(InFilterQuery) ? InFilterQuery : InvalidQueryHandle);
-    }
-    
-    FTedsFilterData::FTedsFilterData(const FName& InFilterName, const FText& InFilterDisplayName, const QueryHandle& InFilterQuery)
-    	: FTedsFilterData(
-    		InFilterName, 
-    		InFilterDisplayName, 
-    		FText::FromString(InFilterName.ToString()), 
-    		FName(), 
-    		nullptr, 
-    		InFilterQuery)
-    {}
-    
-    FTedsFilterData::FTedsFilterData(const FName& InFilterName, const FText& InFilterDisplayName, const FText& InFilterToolTip, const FName& InFilterIconName,
-    	const TSharedPtr<FFilterCategory>& InCategory, const Queries::TQueryFunction<bool>& InFilterQuery)
-    	: FilterName(InFilterName)
-    	, FilterDisplayName(InFilterDisplayName)
-    	, FilterToolTip(InFilterToolTip)
-    	, FilterIconName(InFilterIconName)
-    	, FilterCategory(InCategory)
-    {
-    	FilterQuery.Set<Queries::TQueryFunction<bool>>(InFilterQuery);
-    }
-    
-    FTedsFilterData::FTedsFilterData(const FName& InFilterName, const FText& InFilterDisplayName, const Queries::TQueryFunction<bool>& InFilterQuery)
-    	: FTedsFilterData(
-    		InFilterName, 
-    		InFilterDisplayName, 
-    		FText::FromString(InFilterName.ToString()), 
-    		FName(), 
-    		nullptr, 
-    		InFilterQuery)
-    {}
-    
-    FTedsFilterData::FTedsFilterData(const FTedsFilterData&) = default;
-    FTedsFilterData::~FTedsFilterData() = default;
-		
-	FString FTedsFilter::GetName() const
-	{
-		return FilterName.ToString();
-	}
-
-	FText FTedsFilter::GetDisplayName() const
-	{
-		return FilterDisplayName;
-	}
-
-	FText FTedsFilter::GetToolTipText() const
-	{
-		return FilterToolTip;
-	}
-
-	FLinearColor FTedsFilter::GetColor() const
-	{
-		return FLinearColor();	
-	}
-
-	FName FTedsFilter::GetIconName() const
-	{
-		return FilterIconName;
-	}
-
-	bool FTedsFilter::IsInverseFilter() const
+	void FTedsFilter::SetTedsFilterBar(const TSharedPtr<STedsFilterBar>& InTedsFilterBar)
 	{
-		return false;
+		TedsFilterBar = InTedsFilterBar;
 	}
 
 	void FTedsFilter::ActiveStateChanged(bool bActive)
@@ -124,21 +54,6 @@ namespace UE::Editor::DataStorage
 		}
 	}
 
-	void FTedsFilter::ModifyContextMenu(FMenuBuilder& MenuBuilder)
-	{
-		
-	}
-
-	void FTedsFilter::SaveSettings(const FString& IniFilename, const FString& IniSection, const FString& SettingsString) const
-	{
-		
-	}
-
-	void FTedsFilter::LoadSettings(const FString& IniFilename, const FString& IniSection, const FString& SettingsString)
-	{
-		
-	}
-
 	bool FTedsFilter::PassesFilter(FTedsRowHandle& InItem) const
 	{
 		// The filter is applied through a TEDS query and this is just a dummy to activate it, so we can just run a simple valid check on if the
diff --git a/Engine/Plugins/Experimental/EditorDataStorageFeatures/Source/TedsTableViewer/Private/Widgets/Composite/STedsCompositeHierarchyViewer.cpp b/Engine/Plugins/Experimental/EditorDataStorageFeatures/Source/TedsTableViewer/Private/Widgets/Composite/STedsCompositeHierarchyViewer.cpp
--- a/Engine/Plugins/Experimental/EditorDataStorageFeatures/Source/TedsTableViewer/Private/Widgets/Composite/STedsCompositeHierarchyViewer.cpp
+++ b/Engine/Plugins/Experimental/EditorDataStorageFeatures/Source/TedsTableViewer/Private/Widgets/Composite/STedsCompositeHierarchyViewer.cpp
@@ -102,7 +102,6 @@ namespace UE::Editor::DataStorage
 					.OutFilteredNode(&FilterNode)
 					.CommonSectionFilters(InArgs._CommonSectionFilters)
 					.Filters(InArgs._Filters)
-					.ClassFilters(InArgs._ClassFilters)
 					.UseSectionsForCategories(InArgs._UseSectionsForCategories)
 					.OnPostFiltersChanged_Lambda([this]()
 					{
@@ -284,19 +283,19 @@ namespace UE::Editor::DataStorage
 		return HierarchyViewer->IsSelectedExclusively(InRow);
 	}
 
-	void STedsCompositeHierarchyViewer::AddFilter(FTedsFilterData& InFilterData) const
+	void STedsCompositeHierarchyViewer::AddFilter(const TSharedPtr<FTedsFilter> InFilter) const
 	{
 		if (FilterBar)
 		{
-			FilterBar->AddFilter(InFilterData);
+			FilterBar->AddFilter(InFilter);
 		}
 	}
 
-	void STedsCompositeHierarchyViewer::AddFilters(const TArray<FTedsFilterData>& InFiltersData) const
+	void STedsCompositeHierarchyViewer::AddFilters(const TArray<TSharedPtr<FTedsFilter>>& InFilter) const
 	{
 		if (FilterBar)
 		{
-			FilterBar->AddFilters(InFiltersData);
+			FilterBar->AddFilters(InFilter);
 		}
 	}
 
@@ -324,7 +323,7 @@ namespace UE::Editor::DataStorage
 		}
 	}
 
-	void STedsCompositeHierarchyViewer::ReplaceFilters(const TArray<FTedsFilterData>& InFilters, const bool bPersistFilterStates /*= true*/) const
+	void STedsCompositeHierarchyViewer::ReplaceFilters(const TArray<TSharedPtr<FTedsFilter>>& InFilters, const bool bPersistFilterStates /*= true*/) const
 	{
 		if (FilterBar)
 		{
diff --git a/Engine/Plugins/Experimental/EditorDataStorageFeatures/Source/TedsTableViewer/Private/Widgets/STedsFilterBar.cpp b/Engine/Plugins/Experimental/EditorDataStorageFeatures/Source/TedsTableViewer/Private/Widgets/STedsFilterBar.cpp
--- a/Engine/Plugins/Experimental/EditorDataStorageFeatures/Source/TedsTableViewer/Private/Widgets/STedsFilterBar.cpp
+++ b/Engine/Plugins/Experimental/EditorDataStorageFeatures/Source/TedsTableViewer/Private/Widgets/STedsFilterBar.cpp
@@ -29,17 +29,17 @@ namespace UE::Editor::DataStorage
 		Storage = GetMutableDataStorageFeature<DataStorage::ICoreProvider>(StorageFeatureName);
 
 		OtherFiltersCategory = MakeShared<FFilterCategory>(LOCTEXT("OtherFilters", "Other Filters"), LOCTEXT("OtherFiltersToolTip", "Other Filters"));
-		for (FTedsFilterData FilterData : InArgs._Filters)
+		for (const TSharedPtr<FTedsFilter>& Filter : InArgs._Filters)
 		{
-			if (!FilterData.FilterCategory)
+			if (Filter)
 			{
-				FilterData.FilterCategory = OtherFiltersCategory;
+				if (!Filter->GetCategory())
+				{
+					Filter->SetCategory(OtherFiltersCategory);
+				}
+				Filter->SetTedsFilterBar(STedsFilterBarRef);
+				Args._CustomFilters.Add(Filter.ToSharedRef());
 			}
-			Args._CustomFilters.Add(MakeShared<DataStorage::FTedsFilter>(FilterData, STedsFilterBarRef));
-		}
-		for (UClass* ClassFilter : InArgs._ClassFilters)
-		{
-			Args._CustomFilters.Add(MakeShared<DataStorage::FTedsFilter>(ClassFilter, OtherFiltersCategory, STedsFilterBarRef));
 		}
 		
 		if(ensureMsgf(InArgs._OutFilteredNode && InArgs._InFilterableRowNode, TEXT("In the TEDSFilterBar, a valid Row Node must be passed in to search on and a valid Out Node pointer must be given")))
@@ -63,27 +63,32 @@ namespace UE::Editor::DataStorage
 		}
 	}
 
-	TSharedRef<FTedsFilter> STedsFilterBar::AddFilter(FTedsFilterData& InFilterData)
+	TSharedPtr<FTedsFilter> STedsFilterBar::AddFilter(const TSharedPtr<FTedsFilter> InFilter)
 	{
-		if (!InFilterData.FilterCategory)
+		if (ensureMsgf(InFilter, TEXT("The given TedsFilter was invalid and cannot be added.")))
 		{
-			InFilterData.FilterCategory = OtherFiltersCategory;
-		}
-		
-		if (const TSharedPtr<FFilterBase<FTedsRowHandle&>> FoundFilter = GetFilter(InFilterData.FilterName.ToString()))
-		{
-			// If we did find an existing filter with that name, delete the old version so we can overwrite it
-			DeleteFilter(InFilterData.FilterName);
-		}
+			if (!InFilter->GetCategory())
+			{
+				InFilter->SetCategory(OtherFiltersCategory);
+			}
+			
+			if (const TSharedPtr<FFilterBase<FTedsRowHandle&>> FoundFilter = GetFilter(InFilter->GetName()))
+			{
+				// If we did find an existing filter with that name, delete the old version so we can overwrite it
+				DeleteFilter(InFilter->GetIdentifier());
+			}
+
+			InFilter->SetTedsFilterBar(SharedThis(this));
 
-		const TSharedRef<FTedsFilter> FilterToAdd = MakeShared<FTedsFilter>(InFilterData, SharedThis(this));
-		SBasicFilterBar<FTedsRowHandle&>::AddFilter(FilterToAdd);
-		return FilterToAdd;
+			SBasicFilterBar<FTedsRowHandle&>::AddFilter(InFilter.ToSharedRef());
+			return InFilter;
+		}
+		return nullptr;
 	}
 
-	void STedsFilterBar::AddFilters(const TArray<FTedsFilterData>& InFiltersData)
+	void STedsFilterBar::AddFilters(const TArray<TSharedPtr<FTedsFilter>>& InFilters)
 	{
-		for (FTedsFilterData FilterData : InFiltersData)
+		for (const TSharedPtr<FTedsFilter>& FilterData : InFilters)
 		{
 			AddFilter(FilterData);
 		}
@@ -122,14 +127,14 @@ namespace UE::Editor::DataStorage
 		CommonSectionFilters = InFilters;
 	}
 
-	void STedsFilterBar::ReplaceFilters(const TArray<FTedsFilterData>& InFilters, const bool bPersistFilterStates /*= true*/)
+	void STedsFilterBar::ReplaceFilters(const TArray<TSharedPtr<FTedsFilter>>& InFilters, const bool bPersistFilterStates /*= true*/)
 	{
 		// Find all current filters on the filter bar and which ones are active so that state can be persisted if 
 		// filters with the same name are re-added
 		TMap<FString, bool> PreviousFilterStates;
 		if (bPersistFilterStates)
 		{
-			for(TSharedRef<SFilter> FilterOnBar : Filters)
+			for(const TSharedRef<SFilter>& FilterOnBar : Filters)
 			{
 				PreviousFilterStates.Add(FilterOnBar->GetFilterName(), FilterOnBar->IsEnabled());
 			}
@@ -138,16 +143,18 @@ namespace UE::Editor::DataStorage
 		DeleteAllFilters();
 
 		// Add the new given filters
-		for (FTedsFilterData FilterData : InFilters)
+		for (const TSharedPtr<FTedsFilter>& Filter : InFilters)
 		{
-			const TSharedRef<FTedsFilter> NewFilter = AddFilter(FilterData);
-			if (const bool* FilterState = PreviousFilterStates.Find(NewFilter->GetName()))
+			if (const TSharedPtr<FTedsFilter> NewFilter = AddFilter(Filter))
 			{
-				const TSharedRef<SFilter> AddedFilter = AddFilterToBar(NewFilter);
-				if (*FilterState) // Enabled defaults to false when a filter is created so only perform if we need to enable
+				if (const bool* FilterState = PreviousFilterStates.Find(NewFilter->GetName()))
 				{
-					AddedFilter->SetEnabled(true);
-					SetFrontendFilterActive(NewFilter, true);
+					const TSharedRef<SFilter> AddedFilter = AddFilterToBar(NewFilter.ToSharedRef());
+					if (*FilterState) // Enabled defaults to false when a filter is created so only perform if we need to enable
+					{
+						AddedFilter->SetEnabled(true);
+						SetFrontendFilterActive(NewFilter.ToSharedRef(), true);
+					}
 				}
 			}
 		}
diff --git a/Engine/Plugins/Experimental/EditorDataStorageFeatures/Source/TedsTableViewer/Public/TedsFilter.h b/Engine/Plugins/Experimental/EditorDataStorageFeatures/Source/TedsTableViewer/Public/TedsFilter.h
--- a/Engine/Plugins/Experimental/EditorDataStorageFeatures/Source/TedsTableViewer/Public/TedsFilter.h
+++ b/Engine/Plugins/Experimental/EditorDataStorageFeatures/Source/TedsTableViewer/Public/TedsFilter.h
@@ -2,210 +2,35 @@
 
 #pragma once
 
-#include <type_traits>
-
-#include "Filters/FilterBase.h"
-#include "DataStorage/Features.h"
+#include "TedsFilter.inl"
 #include "DataStorage/Handles.h"
-#include "DataStorage/Queries/Description.h"
-#include "Elements/Columns/TypedElementTypeInfoColumns.h"
-#include "Elements/Framework/TypedElementQueryBuilder.h"
-#include "Elements/Framework/TypedElementQueryContext.h"
-#include "Styling/SlateIconFinder.h"
 
 #define UE_API TEDSTABLEVIEWER_API
 
-#define LOCTEXT_NAMESPACE "TEDSFilter"
-
 namespace UE::Editor::DataStorage
 {
 	class STedsFilterBar;
-
-	// Helper function to check if the given QueryHandle is valid to be used in a filter (is not an observer).
-	inline bool CheckValidFilterQueryHandle(const QueryHandle& InQueryHandle)
-	{
-		// Check if the given Query Handle is valid for a filter (Not an Observer Query Handle)
-		DataStorage::ICoreProvider* Storage = GetMutableDataStorageFeature<DataStorage::ICoreProvider>(StorageFeatureName);
-		if (ensureMsgf(Storage, TEXT("TEDS must be initialized before TEDS Filters")))
-		{
-			const EQueryCallbackType QueryHandleCallbackType = Storage->GetQueryDescription(InQueryHandle).Callback.Type;
-			return ensureMsgf(QueryHandleCallbackType == EQueryCallbackType::None || QueryHandleCallbackType == EQueryCallbackType::Processor,
-				TEXT("TEDS Filters cannot accept Observer Query Handles."));
-		}
-		return false;
-	}
-
-	struct FTedsFilterData
-	{
-		UE_API FTedsFilterData(const FName& InFilterName, const FText& InFilterDisplayName, const FText& InFilterToolTip, const FName& InFilterIconName,
-			const TSharedPtr<FFilterCategory>& InCategory, const QueryHandle& InFilterQuery);
-
-		UE_API FTedsFilterData(const FName& InFilterName, const FText& InFilterDisplayName, const QueryHandle& InFilterQuery);
 		
-		UE_API FTedsFilterData(const FName& InFilterName, const FText& InFilterDisplayName, const FText& InFilterToolTip, const FName& InFilterIconName,
-			const TSharedPtr<FFilterCategory>& InCategory, const Queries::TQueryFunction<bool>& InFilterQuery);
-
-		UE_API FTedsFilterData(const FName& InFilterName, const FText& InFilterDisplayName, const Queries::TQueryFunction<bool>& InFilterQuery);
-
-		/*
-		 * Require that the template type cannot be one of the pre-defined FilterQuery types so it will pass to the defined
-		 * constructor instead of this one.
-		 */
-		template< typename FilterFunction UE_REQUIRES(
-			!std::is_same_v<std::remove_cvref_t<FilterFunction>, QueryHandle> &&
-			!std::is_same_v<std::remove_cvref_t<FilterFunction>, Queries::TQueryFunction<bool>>) >
-		FTedsFilterData(const FName& InFilterName, const FText& InFilterDisplayName, const FText& InFilterToolTip, const FName& InFilterIconName,
-			const TSharedPtr<FFilterCategory>& InCategory, FilterFunction&& InFilterQuery)
-			: FTedsFilterData(
-				InFilterName, 
-				InFilterDisplayName, 
-				InFilterToolTip, 
-				InFilterIconName, 
-				InCategory,
-				Queries::BuildQueryFunction<bool>(Forward<FilterFunction>(InFilterQuery)))
-		{}
-
-		/*
-		 * Require that the template type cannot be one of the pre-defined FilterQuery types so it will pass to the defined
-		 * constructor instead of this one.
-		 */
-		template< typename FilterFunction UE_REQUIRES(
-			!std::is_same_v<std::remove_cvref_t<FilterFunction>, QueryHandle> &&
-			!std::is_same_v<std::remove_cvref_t<FilterFunction>, Queries::TQueryFunction<bool>>) >
-	    FTedsFilterData(const FName& InFilterName, const FText& InFilterDisplayName, FilterFunction&& InFilterQuery)
-			: FTedsFilterData(
-				InFilterName, 
-				InFilterDisplayName, 
-				FText::FromString(InFilterName.ToString()), 
-				FName(), 
-				nullptr, 
-				Queries::BuildQueryFunction<bool>(Forward<FilterFunction>(InFilterQuery)))
-		{}
-
-		UE_API FTedsFilterData(const FTedsFilterData&);
-		UE_API ~FTedsFilterData(); 
-		
-		FName FilterName;
-		FText FilterDisplayName;
-		FText FilterToolTip;
-		FName FilterIconName;
-		
-		// The Category can be set as a nullptr in FTedsFilterData to go to the default 'Other Filters' category when adding,
-		// otherwise the filter will not display if it's not given a category
-		TSharedPtr<FFilterCategory> FilterCategory;
-		TVariant<QueryHandle, Queries::TQueryFunction<bool>> FilterQuery;
-	};
-		
-	class FTedsFilter : public FFilterBase<FTedsRowHandle&>
+	class FTedsFilter : public FTedsFilterBase<FTedsRowHandle&>
 	{
+		friend class STedsFilterBar;
+		
 	public:
-		FTedsFilter(const FTedsFilterData& InTedsFilterData, const TSharedPtr<STedsFilterBar>& InTedsFilterBar )
-			: FFilterBase(InTedsFilterData.FilterCategory)
-			, FilterName(InTedsFilterData.FilterName)
-			, FilterDisplayName(InTedsFilterData.FilterDisplayName)
-			, FilterToolTip(InTedsFilterData.FilterToolTip)
-			, FilterIconName(InTedsFilterData.FilterIconName)
-			, bIsClassFilter(false)
-			, FilterQuery(InTedsFilterData.FilterQuery)
-			, TedsFilterBar(InTedsFilterBar)
-		{}
-
-		FTedsFilter(const FName& InFilterName, const FText& InFilterDisplayName, const FText& InFilterToolTip, const FName& InFilterIconName, 
-			const TSharedPtr<FFilterCategory>& InCategory, const QueryHandle& InFilterQuery, const TSharedPtr<STedsFilterBar>& InTedsFilterBar)
-			: FFilterBase(InCategory)
-			, FilterName(InFilterName)
-			, FilterDisplayName(InFilterDisplayName)
-			, FilterToolTip(InFilterToolTip)
-			, FilterIconName(InFilterIconName)
-			, bIsClassFilter(false)
-			, TedsFilterBar(InTedsFilterBar)
-		{
-			FilterQuery.Set<QueryHandle>(CheckValidFilterQueryHandle(InFilterQuery) ? InFilterQuery : InvalidQueryHandle);
-		}
+		// Inherit all constructors from the base class
+		using FTedsFilterBase::FTedsFilterBase;
 		
-		FTedsFilter(const FName& InFilterName, const FText& InFilterDisplayName, const FText& InFilterToolTip, 
-			const FName& InFilterIconName, const TSharedPtr<FFilterCategory>& InCategory, const Queries::TQueryFunction<bool>& InFilterQuery, 
-			const TSharedPtr<STedsFilterBar>& InTedsFilterBar)
-			: FFilterBase(InCategory)
-			, FilterName(InFilterName)
-			, FilterDisplayName(InFilterDisplayName)
-			, FilterToolTip(InFilterToolTip)
-			, FilterIconName(InFilterIconName)
-			, bIsClassFilter(false)
-			, TedsFilterBar(InTedsFilterBar)
-		{
-			FilterQuery.Set<Queries::TQueryFunction<bool>>(InFilterQuery);
-		}
-
-		FTedsFilter(const UClass* InClass, const TSharedPtr<FFilterCategory>& InCategory, const TSharedPtr<STedsFilterBar>& InTedsFilterBar)
-			: FFilterBase(InCategory)
-			, FilterName(InClass->GetFName())
-			, FilterDisplayName(InClass->GetDisplayNameText())
-			, FilterToolTip(FText::Format(LOCTEXT("FilterClassTooltip", "Filter by {0}"), InClass->GetDisplayNameText()))
-			, FilterIconName(FSlateIconFinder::FindIconForClass(InClass).GetStyleName())
-			, bIsClassFilter(true)
-			, TedsFilterBar(InTedsFilterBar)
-		{
-			FilterQuery.Set<Queries::TQueryFunction<bool>>(
-				Queries::BuildQueryFunction<bool>([InClass](
-					Queries::TQueryContext<Queries::RowBatchInfo> Context,
-					Queries::TResult<bool>& Result,
-					Queries::TConstBatch<FTypedElementClassTypeInfoColumn> TypeInfoColumns)
-				{
-					Context.ForEachRow([&Result, InClass](const FTypedElementClassTypeInfoColumn& TypeInfoColumn)
-					{
-						Result.Add(TypeInfoColumn.TypeInfo->IsChildOf(InClass));
-					}, TypeInfoColumns);
-				}));
-		}
-
-		/** Returns the system name for this filter */
-		virtual FString GetName() const override;
-
-		/** Returns the human-readable name for this filter */
-		virtual FText GetDisplayName() const override;
-
-		/** Returns the tooltip for this filter, shown in the filters menu */
-		virtual FText GetToolTipText() const override;
-
-		/** Returns the color this filter button will be when displayed as a button */
-		virtual FLinearColor GetColor() const override;
-
-		/** Returns the name of the icon to use in menu entries */
-		virtual FName GetIconName() const override;
-
-		/** If true, the filter will be active in the FilterBar when it is inactive in the UI (i.e the filter pill is grayed out) */
-		virtual bool IsInverseFilter() const override;
-
 		/** Notification that the filter became active or inactive */
-		virtual void ActiveStateChanged(bool bActive) override;
-
-		/** Called when the right-click context menu is being built for this filter */
-		virtual void ModifyContextMenu(FMenuBuilder& MenuBuilder) override;
-		
-		/** Can be overriden for custom FilterBar subclasses to save settings, currently not implemented in any generic Filter Bar */
-		virtual void SaveSettings(const FString& IniFilename, const FString& IniSection, const FString& SettingsString) const override;
-
-		/** Can be overridden for custom FilterBar subclasses to load settings, currently not implemented in any generic Filter Bar */
-		virtual void LoadSettings(const FString& IniFilename, const FString& IniSection, const FString& SettingsString) override;
+		virtual UE_API void ActiveStateChanged(bool bActive) override;
 		
 		/** Returns whether the specified Item passes the Filter's restrictions */
-		virtual bool PassesFilter(FTedsRowHandle& InItem ) const override;
+		virtual UE_API bool PassesFilter(FTedsRowHandle& InItem ) const override;
 
 	protected:
-
-		FName FilterName;
-		FText FilterDisplayName;
-		FText FilterToolTip;
-		FName FilterIconName;
-
-		const bool bIsClassFilter;
-
-		TVariant<QueryHandle, Queries::TQueryFunction<bool>> FilterQuery;
+		/** Sets a TedsFilterBar to be used by this Filter */
+		UE_API void SetTedsFilterBar(const TSharedPtr<STedsFilterBar>& InTedsFilterBar);
+		
 		TWeakPtr<STedsFilterBar> TedsFilterBar;
 	};
 } // namespace UE::Editor::DataStorage
 
-#undef LOCTEXT_NAMESPACE
-
 #undef UE_API
diff --git a/Engine/Plugins/Experimental/EditorDataStorageFeatures/Source/TedsTableViewer/Public/TedsFilter.inl b/Engine/Plugins/Experimental/EditorDataStorageFeatures/Source/TedsTableViewer/Public/TedsFilter.inl
new file mode 100644
--- /dev/null
+++ b/Engine/Plugins/Experimental/EditorDataStorageFeatures/Source/TedsTableViewer/Public/TedsFilter.inl
@@ -0,0 +1,290 @@
+// Copyright Epic Games, Inc. All Rights Reserved.
+
+#pragma once
+
+#include <type_traits>
+
+#include "DataStorage/Features.h"
+#include "DataStorage/Handles.h"
+#include "DataStorage/Queries/Description.h"
+#include "Elements/Columns/TypedElementTypeInfoColumns.h"
+#include "Elements/Framework/TypedElementQueryContext.h"
+#include "Elements/Framework/TypedElementQueryBuilder.h"
+#include "Filters/FilterBase.h"
+#include "Styling/SlateIconFinder.h"
+
+#define LOCTEXT_NAMESPACE "TEDSFilter"
+
+namespace UE::Editor::DataStorage
+{
+	// Helper function to check if the given QueryHandle is valid to be used in a filter (is not an observer).
+	inline bool CheckValidFilterQueryHandle(const QueryHandle& InQueryHandle)
+	{
+		// Check if the given Query Handle is valid for a filter (Not an Observer Query Handle)
+		DataStorage::ICoreProvider* Storage = GetMutableDataStorageFeature<DataStorage::ICoreProvider>(StorageFeatureName);
+		if (ensureMsgf(Storage, TEXT("TEDS must be initialized before TEDS Filters")))
+		{
+			const EQueryCallbackType QueryHandleCallbackType = Storage->GetQueryDescription(InQueryHandle).Callback.Type;
+			return ensureMsgf(QueryHandleCallbackType == EQueryCallbackType::None || QueryHandleCallbackType == EQueryCallbackType::Processor,
+				TEXT("TEDS Filters cannot accept Observer Query Handles."));
+		}
+		return false;
+	}
+
+	template<typename FilterType>
+	class FTedsFilterBase : public FFilterBase<FilterType>
+	{
+	public:
+		/**
+		 * Constructor for a TEDS Filter
+		 * @param InFilterName FName Identifier of the Filter
+		 * @param InFilterDisplayName Display Name to show on the Add Filter Menu and toggle button
+		 * @param InFilterToolTip Tooltip to describe the function of the Filter
+		 * @param InFilterIconName Name of an icon in the App Style Set to display in the Add Filter Menu
+		 * @param InCategory Group to place the filter in the Add Filter Menu
+		 * @param InFilterQuery QueryHandle of a registered query to filter for (must be the same syntax as any other Query used in the viewer)
+		 */
+		FTedsFilterBase(const FName& InFilterName, const FText& InFilterDisplayName, const FText& InFilterToolTip, const FName& InFilterIconName, 
+			const TSharedPtr<FFilterCategory>& InCategory, const QueryHandle& InFilterQuery);
+
+		/**
+		 * Constructor for a TEDS Filter
+		 * @param InFilterName FName Identifier of the Filter
+		 * @param InFilterDisplayName Display Name to show on the Add Filter Menu and toggle button
+		 * @param InFilterQuery QueryHandle of a registered query to filter for (must be the same syntax as any other Query used in the viewer)
+		 */
+		FTedsFilterBase(const FName& InFilterName, const FText& InFilterDisplayName, const QueryHandle& InFilterQuery);
+
+		/**
+		 * Constructor for a TEDS Filter
+		 * @param InFilterName FName Identifier of the Filter
+		 * @param InFilterDisplayName Display Name to show on the Add Filter Menu and toggle button
+		 * @param InFilterToolTip Tooltip to describe the function of the Filter
+		 * @param InFilterIconName Name of an icon in the App Style Set to display in the Add Filter Menu
+		 * @param InCategory Group to place the filter in the Add Filter Menu
+		 * @param InFilterQuery QueryFunction to filter for, Batching should be used for improved performance if possible
+		 */
+		FTedsFilterBase(const FName& InFilterName, const FText& InFilterDisplayName, const FText& InFilterToolTip, const FName& InFilterIconName,
+			const TSharedPtr<FFilterCategory>& InCategory, const Queries::TQueryFunction<bool>& InFilterQuery);
+
+		/**
+		 * Constructor for a TEDS Filter
+		 * @param InFilterName FName Identifier of the Filter
+		 * @param InFilterDisplayName Display Name to show on the Add Filter Menu and toggle button
+		 * @param InFilterQuery QueryFunction to filter for, Batching should be used for improved performance if possible
+		 */
+		FTedsFilterBase(const FName& InFilterName, const FText& InFilterDisplayName, const Queries::TQueryFunction<bool>& InFilterQuery);
+		
+		/**
+		 * Constructor for a TEDS Filter
+		 * Require that the template type cannot be one of the pre-defined FilterQuery types so it will pass to the defined
+		 * constructor instead of this one.
+		 * @param InFilterName FName Identifier of the Filter
+		 * @param InFilterDisplayName Display Name to show on the Add Filter Menu and toggle button
+		 * @param InFilterToolTip Tooltip to describe the function of the Filter
+		 * @param InFilterIconName Name of an icon in the App Style Set to display in the Add Filter Menu
+		 * @param InCategory Group to place the filter in the Add Filter Menu
+		 * @param InFilterQuery Function callback that checks the value of a column(s) and returns a bool of whether that row passes the filter
+		 */
+		template< typename FilterFunction UE_REQUIRES(
+			!std::is_same_v<std::remove_cvref_t<FilterFunction>, QueryHandle> &&
+			!std::is_same_v<std::remove_cvref_t<FilterFunction>, Queries::TQueryFunction<bool>>) >
+		FTedsFilterBase(const FName& InFilterName, const FText& InFilterDisplayName, const FText& InFilterToolTip, const FName& InFilterIconName,
+			const TSharedPtr<FFilterCategory>& InCategory, FilterFunction&& InFilterQuery)
+			: FTedsFilterBase(
+				InFilterName, 
+				InFilterDisplayName, 
+				InFilterToolTip, 
+				InFilterIconName, 
+				InCategory,
+				Queries::BuildQueryFunction<bool>(Forward<FilterFunction>(InFilterQuery)))
+		{}
+
+		/**
+		 * Constructor for a TEDS Filter
+		 * Require that the template type cannot be one of the pre-defined FilterQuery types so it will pass to the defined
+		 * constructor instead of this one.
+		 * @param InFilterName FName Identifier of the Filter
+		 * @param InFilterDisplayName Display Name to show on the Add Filter Menu and toggle button
+		 * @param InFilterQuery Function callback that checks the value of a column(s) and returns a bool of whether that row passes the filter
+		 */
+		template< typename FilterFunction UE_REQUIRES(
+			!std::is_same_v<std::remove_cvref_t<FilterFunction>, QueryHandle> &&
+			!std::is_same_v<std::remove_cvref_t<FilterFunction>, Queries::TQueryFunction<bool>>) >
+		FTedsFilterBase(const FName& InFilterName, const FText& InFilterDisplayName, FilterFunction&& InFilterQuery)
+			: FTedsFilterBase(
+				InFilterName, 
+				InFilterDisplayName, 
+				FText::FromString(InFilterName.ToString()), 
+				FName(), 
+				nullptr,
+				Queries::BuildQueryFunction<bool>(Forward<FilterFunction>(InFilterQuery)))
+		{}
+
+		/**
+		 * Constructor for a TEDS Class Filter
+		 * @param InClass UClass to filter for on the TypeInfo Column
+		 * @param InCategory Group to place the filter in the Add Filter Menu
+		 */
+		FTedsFilterBase(const UClass* InClass, const TSharedPtr<FFilterCategory>& InCategory = nullptr);
+
+		virtual ~FTedsFilterBase() override = default;
+		
+		void SetCategory(const TSharedPtr<FFilterCategory>& InCategory);
+        
+		bool IsClassFilter() const { return bIsClassFilter; }
+		
+		FName GetIdentifier() const { return FilterName; }
+
+		/** Returns the system name for this filter */
+		virtual FString GetName() const override;
+
+		/** Returns the human-readable name for this filter */
+		virtual FText GetDisplayName() const override;
+
+		/** Returns the tooltip for this filter, shown in the filters menu */
+		virtual FText GetToolTipText() const override;
+
+		/** Returns the color this filter button will be when displayed as a button */
+		virtual FLinearColor GetColor() const override;
+
+		/** Returns the name of the icon to use in menu entries */
+		virtual FName GetIconName() const override;
+
+		/** If true, the filter will be active in the FilterBar when it is inactive in the UI (i.e the filter pill is grayed out) */
+		virtual bool IsInverseFilter() const override;
+
+		/** Called when the right-click context menu is being built for this filter */
+		virtual void ModifyContextMenu(FMenuBuilder& MenuBuilder) override;
+		
+		/** Can be overriden for custom FilterBar subclasses to save settings, currently not implemented in any generic Filter Bar */
+		virtual void SaveSettings(const FString& IniFilename, const FString& IniSection, const FString& SettingsString) const override;
+
+		/** Can be overridden for custom FilterBar subclasses to load settings, currently not implemented in any generic Filter Bar */
+		virtual void LoadSettings(const FString& IniFilename, const FString& IniSection, const FString& SettingsString) override;
+		
+	protected:
+		FName FilterName;
+		FText FilterDisplayName;
+		FText FilterToolTip;
+		FName FilterIconName;
+
+		const bool bIsClassFilter;
+
+		TVariant<QueryHandle, Queries::TQueryFunction<bool>> FilterQuery;
+	};
+
+	template<typename FilterType>
+	FTedsFilterBase<FilterType>::FTedsFilterBase(const FName& InFilterName, const FText& InFilterDisplayName, const FText& InFilterToolTip, const FName& InFilterIconName, 
+		const TSharedPtr<FFilterCategory>& InCategory, const QueryHandle& InFilterQuery)
+		: FFilterBase<FilterType>(InCategory)
+		, FilterName(InFilterName)
+		, FilterDisplayName(InFilterDisplayName)
+		, FilterToolTip(InFilterToolTip)
+		, FilterIconName(InFilterIconName)
+		, bIsClassFilter(false)
+	{
+		FilterQuery.Set<QueryHandle>(CheckValidFilterQueryHandle(InFilterQuery) ? InFilterQuery : InvalidQueryHandle);
+	}
+
+	template<typename FilterType>
+	FTedsFilterBase<FilterType>::FTedsFilterBase(const FName& InFilterName, const FText& InFilterDisplayName, const QueryHandle& InFilterQuery)
+		: FTedsFilterBase(InFilterName, InFilterDisplayName, FText::FromString(InFilterName.ToString()), FName(), 
+			nullptr, InFilterQuery)
+	{}
+
+	template<typename FilterType>
+	FTedsFilterBase<FilterType>::FTedsFilterBase(const FName& InFilterName, const FText& InFilterDisplayName, const FText& InFilterToolTip, 
+		const FName& InFilterIconName, const TSharedPtr<FFilterCategory>& InCategory, const Queries::TQueryFunction<bool>& InFilterQuery)
+		: FFilterBase<FilterType>(InCategory)
+		, FilterName(InFilterName)
+		, FilterDisplayName(InFilterDisplayName)
+		, FilterToolTip(InFilterToolTip)
+		, FilterIconName(InFilterIconName)
+		, bIsClassFilter(false)
+	{
+		FilterQuery.Set<Queries::TQueryFunction<bool>>(InFilterQuery);
+	}
+
+	template<typename FilterType>
+	FTedsFilterBase<FilterType>::FTedsFilterBase(const FName& InFilterName, const FText& InFilterDisplayName, const Queries::TQueryFunction<bool>& InFilterQuery)
+		: FTedsFilterBase(InFilterName, InFilterDisplayName, FText::FromString(InFilterName.ToString()), FName(), 
+			nullptr, InFilterQuery)
+	{}
+
+	template<typename FilterType>
+	FTedsFilterBase<FilterType>::FTedsFilterBase(const UClass* InClass, const TSharedPtr<FFilterCategory>& InCategory)
+		: FFilterBase<FilterType>(InCategory)
+		, FilterName(InClass->GetFName())
+		, FilterDisplayName(InClass->GetDisplayNameText())
+		, FilterToolTip(FText::Format(LOCTEXT("FilterClassTooltip", "Filter by {0}"), InClass->GetDisplayNameText()))
+		, FilterIconName(FSlateIconFinder::FindIconForClass(InClass).GetStyleName())
+		, bIsClassFilter(true)
+	{
+		using namespace UE::Editor::DataStorage::Queries;
+		FilterQuery.Set<TQueryFunction<bool>>(
+			Queries::BuildQueryFunction<bool>([InClass](TQueryContext<RowBatchInfo> Context, TResult<bool>& Result, TConstBatch<FTypedElementClassTypeInfoColumn> TypeInfoColumns)
+			{
+				Context.ForEachRow([&Result, InClass](const FTypedElementClassTypeInfoColumn& TypeInfoColumn)
+				{
+					Result.Add(TypeInfoColumn.TypeInfo->IsChildOf(InClass));
+				}, TypeInfoColumns);
+			}));
+	}
+
+	template<typename FilterType>
+	void FTedsFilterBase<FilterType>::SetCategory(const TSharedPtr<FFilterCategory>& InCategory)
+	{
+		FFilterBase<FilterType>::FilterCategory = InCategory;
+	}
+	
+	template<typename FilterType>
+	FString FTedsFilterBase<FilterType>::GetName() const
+	{
+		return FilterName.ToString();
+	}
+	
+	template<typename FilterType>
+	FText FTedsFilterBase<FilterType>::GetDisplayName() const
+	{
+		return FilterDisplayName;
+	}
+
+	template<typename FilterType>
+	FText FTedsFilterBase<FilterType>::GetToolTipText() const
+	{
+		return FilterToolTip;
+	}
+
+	template<typename FilterType>
+	FLinearColor FTedsFilterBase<FilterType>::GetColor() const
+	{
+		return FLinearColor();	
+	}
+
+	template<typename FilterType>
+	FName FTedsFilterBase<FilterType>::GetIconName() const
+	{
+		return FilterIconName;
+	}
+
+	template<typename FilterType>
+	bool FTedsFilterBase<FilterType>::IsInverseFilter() const
+	{
+		return false;
+	}
+
+	template<typename FilterType>
+	void FTedsFilterBase<FilterType>::ModifyContextMenu(FMenuBuilder& MenuBuilder)
+	{}
+
+	template<typename FilterType>
+	void FTedsFilterBase<FilterType>::SaveSettings(const FString& IniFilename, const FString& IniSection, const FString& SettingsString) const
+	{}
+
+	template<typename FilterType>
+	void FTedsFilterBase<FilterType>::LoadSettings(const FString& IniFilename, const FString& IniSection, const FString& SettingsString)
+	{}
+}// namespace UE::Editor::DataStorage
+
+#undef LOCTEXT_NAMESPACE
diff --git a/Engine/Plugins/Experimental/EditorDataStorageFeatures/Source/TedsTableViewer/Public/Widgets/Composite/STedsCompositeHierarchyViewer.h b/Engine/Plugins/Experimental/EditorDataStorageFeatures/Source/TedsTableViewer/Public/Widgets/Composite/STedsCompositeHierarchyViewer.h
--- a/Engine/Plugins/Experimental/EditorDataStorageFeatures/Source/TedsTableViewer/Public/Widgets/Composite/STedsCompositeHierarchyViewer.h
+++ b/Engine/Plugins/Experimental/EditorDataStorageFeatures/Source/TedsTableViewer/Public/Widgets/Composite/STedsCompositeHierarchyViewer.h
@@ -97,11 +97,8 @@ namespace UE::Editor::DataStorage
 		/** Array of filter names to also add to the 'common filters' section (section at the top of the menu) */
 		SLATE_ARGUMENT(TArray<FName>, CommonSectionFilters)
 
-		// Array of filters to add to the filter bar *AND Filters
-		SLATE_ARGUMENT(TArray<FTedsFilterData>, Filters)
-
-		// Array of class filters to add to the filter bar *OR Filters
-		SLATE_ARGUMENT(TArray<UClass*>, ClassFilters)
+		// Array of filters to add to the filter bar
+		SLATE_ARGUMENT(TArray<TSharedPtr<FTedsFilter>>, Filters)
 
 		// Whether to use submenus or sections for categories in the filter menu
 		SLATE_ARGUMENT(bool, UseSectionsForCategories)
@@ -152,9 +149,9 @@ namespace UE::Editor::DataStorage
 
 		TEDSTABLEVIEWER_API virtual bool IsSelectedExclusively(RowHandle InRow) const override;
 
-		TEDSTABLEVIEWER_API void AddFilter(FTedsFilterData& InFilterData) const;
+		TEDSTABLEVIEWER_API void AddFilter(const TSharedPtr<FTedsFilter> InFilter) const;
 
-		TEDSTABLEVIEWER_API void AddFilters(const TArray<FTedsFilterData>& InFiltersData) const;
+		TEDSTABLEVIEWER_API void AddFilters(const TArray<TSharedPtr<FTedsFilter>>& InFilters) const;
 
 		TEDSTABLEVIEWER_API void DeleteFilter(const FName& InFilterName) const;
 
@@ -162,7 +159,7 @@ namespace UE::Editor::DataStorage
 		
 		TEDSTABLEVIEWER_API void SetCommonFilters(const TArray<FName>& InFilters) const;
 
-		TEDSTABLEVIEWER_API void ReplaceFilters(const TArray<FTedsFilterData>& InFilters, const bool bPersistFilterStates = true) const;
+		TEDSTABLEVIEWER_API void ReplaceFilters(const TArray<TSharedPtr<FTedsFilter>>& InFilters, const bool bPersistFilterStates = true) const;
 		
 		// Recursively Expand all rows currently present in the viewer
 		TEDSTABLEVIEWER_API void ExpandAll() const;
diff --git a/Engine/Plugins/Experimental/EditorDataStorageFeatures/Source/TedsTableViewer/Public/Widgets/STedsFilterBar.h b/Engine/Plugins/Experimental/EditorDataStorageFeatures/Source/TedsTableViewer/Public/Widgets/STedsFilterBar.h
--- a/Engine/Plugins/Experimental/EditorDataStorageFeatures/Source/TedsTableViewer/Public/Widgets/STedsFilterBar.h
+++ b/Engine/Plugins/Experimental/EditorDataStorageFeatures/Source/TedsTableViewer/Public/Widgets/STedsFilterBar.h
@@ -50,11 +50,8 @@ namespace UE::Editor::DataStorage
 			/** The filter results node that will contain the filtered rows */
 			SLATE_ARGUMENT(TSharedPtr<QueryStack::IRowNode>*, OutFilteredNode)
 
-			/** Array of filters to add to the filter bar *AND Filters */
-			SLATE_ARGUMENT(TArray<FTedsFilterData>, Filters)
-		
-			/** Array of class filters to add to the filter bar *OR Filters */
-			SLATE_ARGUMENT(TArray<UClass*>, ClassFilters)
+			/** Array of filters to add to the filter bar */
+			SLATE_ARGUMENT(TArray<TSharedPtr<FTedsFilter>>, Filters)
 			
 			/** Array of filter names to also add to the 'common filters' section (section at the top of the menu) */
             SLATE_ARGUMENT(TArray<FName>, CommonSectionFilters)
@@ -82,10 +79,10 @@ namespace UE::Editor::DataStorage
 		/** Override tick so that filters are only recompiled once if multiple were added at once */
 		UE_API virtual void Tick(const FGeometry& AllottedGeometry, const double InCurrentTime, const float InDeltaTime) override;
 
-		/** Add an array of new Filters to the menu post-initialization, will overwrite any filters with the same name */
-		UE_API TSharedRef<FTedsFilter> AddFilter(FTedsFilterData& InFilterData);
 		/** Add a new Filter to the menu post-initialization, will overwrite any filters with the same name */
-		UE_API void AddFilters(const TArray<FTedsFilterData>& InFiltersData);
+		UE_API TSharedPtr<FTedsFilter> AddFilter(const TSharedPtr<FTedsFilter> InFilter);
+		/** Add an array of new Filters to the menu post-initialization, will overwrite any filters with the same name */
+		UE_API void AddFilters(const TArray<TSharedPtr<FTedsFilter>>& InFilters);
 		/** Remove a Filter from the menu and remove any place it may be referenced */
 		UE_API void DeleteFilter(const FName& InFilterName);
 		/** Remove all Filters from the menu and all references */
@@ -97,7 +94,7 @@ namespace UE::Editor::DataStorage
 		 * Remove all filters and replace them with a new array of filters with an option to persist their state (is pinned & is active) for new
 		 * filters with the same FName identifier
 		 */
-		UE_API void ReplaceFilters(const TArray<FTedsFilterData>& InFilters, const bool bPersistFilterStates = true);
+		UE_API void ReplaceFilters(const TArray<TSharedPtr<FTedsFilter>>& InFilters, const bool bPersistFilterStates = true);
 		
 		// Add an external query description to the FilterBar
 		UE_API void AddExternalQuery(const FName& QueryName, const QueryHandle& InQueryHandle);
diff --git a/Engine/Plugins/Experimental/MeshPartition/Source/MeshPartitionEditorUI/Private/UI/MeshPartitionSettingsWidget.cpp b/Engine/Plugins/Experimental/MeshPartition/Source/MeshPartitionEditorUI/Private/UI/MeshPartitionSettingsWidget.cpp
--- a/Engine/Plugins/Experimental/MeshPartition/Source/MeshPartitionEditorUI/Private/UI/MeshPartitionSettingsWidget.cpp
+++ b/Engine/Plugins/Experimental/MeshPartition/Source/MeshPartitionEditorUI/Private/UI/MeshPartitionSettingsWidget.cpp
@@ -652,8 +652,8 @@ TSharedRef<UE::Editor::DataStorage::STedsCompositeHierarchyViewer> SMegaMeshSett
 		LOCTEXT("MeshPartitionFilterCategory_Tooltip","Filters related to the Mesh Partition")
 	);
 	
-	TArray<UE::Editor::DataStorage::FTedsFilterData> Filters;
-	
+	TArray<TSharedPtr<DataStorage::FTedsFilter>> Filters;
+
 	{
 		TQueryFunction<bool> FilterQuery = 
 			BuildQueryFunction<bool>([](TQueryContext<SingleRowInfo> Context)
@@ -665,16 +665,14 @@ TSharedRef<UE::Editor::DataStorage::STedsCompositeHierarchyViewer> SMegaMeshSett
 				return false;
 			});
 
-		UE::Editor::DataStorage::FTedsFilterData Filter(
+		Filters.Add(MakeShared<DataStorage::FTedsFilter>(
 			FName(TEXT("ActorBoundsFilter")),
 			LOCTEXT("ActorBoundsFilter", "Filter By Actor Bounds"),
 			LOCTEXT("ActorBoundsFilter_Tooltip", "Use the selected Actor's bonuds to filter only those modifiers affecting that region."),
 			FName(),
 			MeshPartitionFilterCategory,
 			FilterQuery
-			);
-
-		Filters.Add(Filter);
+			));
 	}
 
 
@@ -698,7 +696,7 @@ TSharedRef<UE::Editor::DataStorage::STedsCompositeHierarchyViewer> SMegaMeshSett
 					OnOutlinerSelectionChanged(SelectInfo);
 				})
 			.EmptyRowsMessage(LOCTEXT("MegaMeshSettings_NoRows", "No modifiers found for this Mesh Partition instance.")))
-		.Filters(Filters)
+			.Filters(Filters)
 		;
 
 	ReferenceView->ExpandAll();
