[Unified Favorites] Add UnifiedFavorites module to EditorDataStorageFeatures

Extracts and generalizes the favorites widget system from TedsOutliner into a
dedicated UnifiedFavorites module, making it reusable across different editor
contexts (Outliner, Table, Prefab Editor, etc.).

[Added] UnifiedFavorites editor module (UnifiedFavorites.Build.cs) with
dependencies on TedsOutliner, TypedElementFramework, Slate, and TEDS
subsystems

[Added] SUnifiedFavoritesWidget (UnifiedFavoritesWidget.h/.cpp) ? base Slate
widget extending SImage that manages favorite state toggling on TEDS
RowHandles, with hover/click handling and active/inactive color states

[Added] UUnifiedFavoritesTableFactory and FTableFavoriteWidgetConstructor /
FTableFavoriteHeaderConstructor (UnifiedFavoritesTableWidget.h/.cpp) ?
table-specific widget constructors with column sorting support

[Added] UUnifiedFavoritesOutlinerFactory and FOutlinerFavoriteHeaderConstructor
(UnifiedFavoritesOutlinerWidget.h/.cpp) ? outliner-specific header
widget constructor override

[Added] FUnifiedFavoritesModule (UnifiedFavoritesModule.cpp) ? module
startup/shutdown boilerplate

[Modified] EditorDataStorageFeatures.uplugin to register UnifiedFavorites as a
new editor-only module

[Modified] TedsOutlinerModule.cpp ? renamed favorites CVar

[Modified] EntityPrefabToolkit.cpp ? renamed favorites CVar

[Removed] OutlinerFavoriteWidget.h and OutlinerFavoriteWidget.cpp from
TedsOutliner (superseded by UnifiedFavorites)

#Synced-CL 51672242

diff --git a/Engine/Plugins/Experimental/EditorDataStorageFeatures/EditorDataStorageFeatures.uplugin b/Engine/Plugins/Experimental/EditorDataStorageFeatures/EditorDataStorageFeatures.uplugin
--- a/Engine/Plugins/Experimental/EditorDataStorageFeatures/EditorDataStorageFeatures.uplugin
+++ b/Engine/Plugins/Experimental/EditorDataStorageFeatures/EditorDataStorageFeatures.uplugin
@@ -79,6 +79,10 @@
 		{
 			"Name": "TedsOperations",
 			"Type": "Editor"
+		},
+		{
+			"Name": "UnifiedFavorites",
+			"Type": "Editor"
 		}
 	],
 	"Plugins": [
diff --git a/Engine/Plugins/Experimental/EditorDataStorageFeatures/Source/TedsOutliner/Private/TedsOutlinerModule.cpp b/Engine/Plugins/Experimental/EditorDataStorageFeatures/Source/TedsOutliner/Private/TedsOutlinerModule.cpp
--- a/Engine/Plugins/Experimental/EditorDataStorageFeatures/Source/TedsOutliner/Private/TedsOutlinerModule.cpp
+++ b/Engine/Plugins/Experimental/EditorDataStorageFeatures/Source/TedsOutliner/Private/TedsOutlinerModule.cpp
@@ -64,7 +64,7 @@ namespace UE::Editor::Outliner
 		}));
 
 	static FAutoConsoleVariableRef CVarUseFavoritesWidgets(
-		TEXT("TEDS.UI.UseFavoritesWidgets.LevelEditorOutliner"),
+		TEXT("UnifiedFavorites.LevelEditorOutliner.Widget"),
 		Private::bUseFavoritesWidgets,
 		TEXT("Use unified favorites in the Level Editor Outliner (requires TEDS-Outliner to be enabled)")
 		, FConsoleVariableDelegate::CreateLambda([](IConsoleVariable*)
diff --git a/Engine/Plugins/Experimental/EditorDataStorageFeatures/Source/TedsOutliner/Private/Widgets/OutlinerFavoriteWidget.cpp b/Engine/Plugins/Experimental/EditorDataStorageFeatures/Source/TedsOutliner/Private/Widgets/OutlinerFavoriteWidget.cpp
deleted file mode 100644
diff --git a/Engine/Plugins/Experimental/EditorDataStorageFeatures/Source/TedsOutliner/Private/Widgets/OutlinerFavoriteWidget.h b/Engine/Plugins/Experimental/EditorDataStorageFeatures/Source/TedsOutliner/Private/Widgets/OutlinerFavoriteWidget.h
deleted file mode 100644
diff --git a/Engine/Plugins/Experimental/EditorDataStorageFeatures/Source/UnifiedFavorites/Private/UnifiedFavoritesModule.cpp b/Engine/Plugins/Experimental/EditorDataStorageFeatures/Source/UnifiedFavorites/Private/UnifiedFavoritesModule.cpp
new file mode 100644
--- /dev/null
+++ b/Engine/Plugins/Experimental/EditorDataStorageFeatures/Source/UnifiedFavorites/Private/UnifiedFavoritesModule.cpp
@@ -0,0 +1,20 @@
+// Copyright Epic Games, Inc. All Rights Reserved.
+
+#include "Modules/ModuleManager.h"
+#include "Modules/ModuleInterface.h"
+
+class FUnifiedFavoritesModule : public IModuleInterface
+{
+public:
+
+	virtual void StartupModule() override
+	{
+	}
+
+	virtual void ShutdownModule() override
+	{
+	}
+};
+
+IMPLEMENT_MODULE(FUnifiedFavoritesModule, UnifiedFavorites);
+
diff --git a/Engine/Plugins/Experimental/EditorDataStorageFeatures/Source/UnifiedFavorites/Private/Widgets/UnifiedFavoritesOutlinerWidget.cpp b/Engine/Plugins/Experimental/EditorDataStorageFeatures/Source/UnifiedFavorites/Private/Widgets/UnifiedFavoritesOutlinerWidget.cpp
new file mode 100644
--- /dev/null
+++ b/Engine/Plugins/Experimental/EditorDataStorageFeatures/Source/UnifiedFavorites/Private/Widgets/UnifiedFavoritesOutlinerWidget.cpp
@@ -0,0 +1,69 @@
+// Copyright Epic Games, Inc. All Rights Reserved.
+
+#include "Widgets/UnifiedFavoritesOutlinerWidget.h"
+
+#include "Columns/SlateHeaderColumns.h"
+#include "Elements/Columns/TypedElementCompatibilityColumns.h"
+#include "Elements/Columns/TypedElementFolderColumns.h"
+#include "Elements/Columns/TypedElementMiscColumns.h"
+#include "Styling/AppStyle.h"
+#include "Widgets/Images/SImage.h"
+#include "Widgets/Layout/SBox.h"
+
+#include UE_INLINE_GENERATED_CPP_BY_NAME(UnifiedFavoritesOutlinerWidget)
+
+#define LOCTEXT_NAMESPACE "UnifiedFavoritesOutlinerWidget"
+
+//
+// Outliner Favorite Header Widget
+//
+
+FOutlinerFavoriteHeaderConstructor::FOutlinerFavoriteHeaderConstructor()
+	: FSimpleWidgetConstructor(StaticStruct())
+{
+}
+
+TSharedPtr<SWidget> FOutlinerFavoriteHeaderConstructor::CreateWidget(
+	UE::Editor::DataStorage::ICoreProvider* DataStorage,
+	UE::Editor::DataStorage::IUiProvider* DataStorageUi,
+	UE::Editor::DataStorage::RowHandle TargetRow,
+	UE::Editor::DataStorage::RowHandle WidgetRow,
+	const UE::Editor::DataStorage::FMetaDataView& Arguments)
+{
+	using namespace UE::Editor::DataStorage;
+	DataStorage->AddColumn(WidgetRow, FHeaderWidgetSizeColumn
+		{
+			.ColumnSizeMode = EColumnSizeMode::Fixed,
+			.Width = 24.0f
+		});
+
+	return SNew(SBox)
+		.VAlign(VAlign_Center)
+		.HAlign(HAlign_Center)
+		.HeightOverride(18.f)
+		.WidthOverride(18.f)
+		[
+			SNew(SImage)
+			.ColorAndOpacity(FAppStyle::Get().GetSlateColor("Colors.Foreground"))
+			.Image(FAppStyle::Get().GetBrush("Icons.Star"))
+			.ToolTipText(LOCTEXT("OutlinerFavoriteHeaderTooltip", "Favorite"))
+		];
+}
+
+//
+// Outliner Factory
+//
+
+void UUnifiedFavoritesOutlinerFactory::RegisterWidgetConstructors(UE::Editor::DataStorage::ICoreProvider& DataStorage,
+	UE::Editor::DataStorage::IUiProvider& DataStorageUi) const
+{
+	using namespace UE::Editor::DataStorage;
+	using namespace UE::Editor::DataStorage::Queries;
+
+	FConditions OutlinerColumns = TColumn<FFavoriteTag>() || TColumn<FTypedElementUObjectColumn>() || TColumn<FFolderTag>();
+
+	DataStorageUi.RegisterWidgetFactory<FOutlinerFavoriteHeaderConstructor>(
+		DataStorageUi.FindPurpose(IUiProvider::FPurposeInfo("SceneOutliner", "Header", NAME_None).GeneratePurposeID()), OutlinerColumns);
+}
+
+#undef LOCTEXT_NAMESPACE
diff --git a/Engine/Plugins/Experimental/EditorDataStorageFeatures/Source/UnifiedFavorites/Private/Widgets/UnifiedFavoritesOutlinerWidget.h b/Engine/Plugins/Experimental/EditorDataStorageFeatures/Source/UnifiedFavorites/Private/Widgets/UnifiedFavoritesOutlinerWidget.h
new file mode 100644
--- /dev/null
+++ b/Engine/Plugins/Experimental/EditorDataStorageFeatures/Source/UnifiedFavorites/Private/Widgets/UnifiedFavoritesOutlinerWidget.h
@@ -0,0 +1,40 @@
+// Copyright Epic Games, Inc. All Rights Reserved.
+
+#pragma once
+
+#include "Elements/Interfaces/TypedElementDataStorageFactory.h"
+#include "Elements/Interfaces/TypedElementDataStorageUiInterface.h"
+
+#include "UnifiedFavoritesOutlinerWidget.generated.h"
+
+// Outliner-specific header override
+USTRUCT()
+struct FOutlinerFavoriteHeaderConstructor : public FSimpleWidgetConstructor
+{
+	GENERATED_BODY()
+
+public:
+	FOutlinerFavoriteHeaderConstructor();
+	~FOutlinerFavoriteHeaderConstructor() override = default;
+
+	virtual TSharedPtr<SWidget> CreateWidget(
+		UE::Editor::DataStorage::ICoreProvider* DataStorage,
+		UE::Editor::DataStorage::IUiProvider* DataStorageUi,
+		UE::Editor::DataStorage::RowHandle TargetRow,
+		UE::Editor::DataStorage::RowHandle WidgetRow,
+		const UE::Editor::DataStorage::FMetaDataView& Arguments) override;
+};
+
+// Factory — registers only the Outliner header override
+UCLASS()
+class UUnifiedFavoritesOutlinerFactory : public UEditorDataStorageFactory
+{
+	GENERATED_BODY()
+
+public:
+	~UUnifiedFavoritesOutlinerFactory() override = default;
+
+	virtual void RegisterWidgetConstructors(
+		UE::Editor::DataStorage::ICoreProvider& DataStorage,
+		UE::Editor::DataStorage::IUiProvider& DataStorageUi) const override;
+};
diff --git a/Engine/Plugins/Experimental/EditorDataStorageFeatures/Source/UnifiedFavorites/Private/Widgets/UnifiedFavoritesTableWidget.cpp b/Engine/Plugins/Experimental/EditorDataStorageFeatures/Source/UnifiedFavorites/Private/Widgets/UnifiedFavoritesTableWidget.cpp
new file mode 100644
--- /dev/null
+++ b/Engine/Plugins/Experimental/EditorDataStorageFeatures/Source/UnifiedFavorites/Private/Widgets/UnifiedFavoritesTableWidget.cpp
@@ -0,0 +1,154 @@
+// Copyright Epic Games, Inc. All Rights Reserved.
+
+#include "UnifiedFavoritesTableWidget.h"
+
+#include "TedsAssetDataColumns.h"
+#include "UnifiedFavoritesWidget.h"
+#include "Widgets/AssetDataLabelWidget.h"
+#include "Columns/SlateHeaderColumns.h"
+#include "Elements/Columns/TypedElementCompatibilityColumns.h"
+#include "Elements/Columns/TypedElementFolderColumns.h"
+#include "Elements/Columns/TypedElementMiscColumns.h"
+#include "Styling/AppStyle.h"
+#include "Widgets/Images/SImage.h"
+#include "Widgets/Layout/SBox.h"
+#include "Widgets/Text/STextBlock.h"
+
+#include UE_INLINE_GENERATED_CPP_BY_NAME(UnifiedFavoritesTableWidget)
+
+#define LOCTEXT_NAMESPACE "UnifiedFavoritesTableWidget"
+
+namespace UE::Editor::DataStorage::Private
+{
+	//
+	// Column Sorter
+	//
+
+	class FFavoriteWidgetSorter final : public FColumnSorterInterface
+	{
+	public:
+		virtual int32 Compare(const ICoreProvider& Storage, RowHandle Left, RowHandle Right) const override
+		{
+			return 0;
+		}
+
+		virtual FPrefixInfo CalculatePrefix(const ICoreProvider& Storage, RowHandle Row, uint32 ByteIndex) const override
+		{
+			return CreateSortPrefix(ByteIndex, static_cast<int32>(!Storage.HasColumns<FFavoriteTag>(Row)));
+		}
+
+		virtual ESortType GetSortType() const override
+		{
+			return ESortType::FixedSize64;
+		}
+
+		virtual FText GetShortName() const override
+		{
+			return LOCTEXT("FavoriteWidgetSorter", "Favorite Sorter");
+		}
+	};
+
+	//
+	// Helpers
+	//
+
+	TSharedRef<SWidget> MakeFavoritesWidgetContainer(const TSharedRef<SWidget>& FavoritesWidget)
+	{
+		return SNew(SBox)
+			.VAlign(VAlign_Center)
+			.HAlign(HAlign_Center)
+			.HeightOverride(18.f)
+			.WidthOverride(18.f)
+			[
+				FavoritesWidget
+			];
+	}
+}
+
+//
+// Table Favorite Header Widget
+//
+
+FTableFavoriteHeaderConstructor::FTableFavoriteHeaderConstructor()
+	: FSimpleWidgetConstructor(StaticStruct())
+{
+}
+
+TSharedPtr<SWidget> FTableFavoriteHeaderConstructor::CreateWidget(UE::Editor::DataStorage::ICoreProvider* DataStorage,
+	UE::Editor::DataStorage::IUiProvider* DataStorageUi,
+	UE::Editor::DataStorage::RowHandle TargetRow, UE::Editor::DataStorage::RowHandle WidgetRow, const UE::Editor::DataStorage::FMetaDataView& Arguments)
+{
+	using namespace UE::Editor::DataStorage;
+	DataStorage->AddColumn(WidgetRow, FHeaderWidgetSizeColumn
+		{
+			.ColumnSizeMode = EColumnSizeMode::Fixed,
+			.Width = 24.0f
+		});
+
+	return Private::MakeFavoritesWidgetContainer(
+		SNew(SImage)
+		.ColorAndOpacity(FAppStyle::Get().GetSlateColor("Colors.Foreground"))
+		.Image(FAppStyle::Get().GetBrush("Icons.Star"))
+		.ToolTipText(LOCTEXT("TableFavoriteHeaderTooltip", "Favorite"))
+	);
+}
+
+//
+// Table Favorite Widget
+//
+
+FTableFavoriteWidgetConstructor::FTableFavoriteWidgetConstructor()
+	: FSimpleWidgetConstructor(StaticStruct())
+{
+}
+
+TSharedPtr<SWidget> FTableFavoriteWidgetConstructor::CreateWidget(
+	UE::Editor::DataStorage::ICoreProvider* DataStorage,
+	UE::Editor::DataStorage::IUiProvider* DataStorageUi,
+	UE::Editor::DataStorage::RowHandle TargetRow,
+	UE::Editor::DataStorage::RowHandle WidgetRow,
+	const UE::Editor::DataStorage::FMetaDataView& Arguments)
+{
+
+	using namespace UE::Editor::DataStorage;
+	if (!DataStorage->IsRowAvailable(TargetRow))
+	{
+		return SNew(STextBlock)
+			.Text(LOCTEXT("MissingRowReferenceColumn", "Unable to retrieve row reference."));
+	}
+
+	return Private::MakeFavoritesWidgetContainer(
+		SNew(SUnifiedFavoritesWidget, DataStorage)
+		.RowHandle(TargetRow)
+	);
+}
+
+TArray<TSharedPtr<const UE::Editor::DataStorage::FColumnSorterInterface>> FTableFavoriteWidgetConstructor::ConstructColumnSorters(
+	UE::Editor::DataStorage::ICoreProvider* DataStorage, UE::Editor::DataStorage::IUiProvider* DataStorageUi,
+	const UE::Editor::DataStorage::FMetaDataView& Arguments)
+{
+	using namespace UE::Editor::DataStorage;
+	return TArray<TSharedPtr<const FColumnSorterInterface>>(
+		{
+			MakeShared<Private::FFavoriteWidgetSorter>()
+		});
+}
+
+//
+// Table Factory
+//
+
+void UUnifiedFavoritesTableFactory::RegisterWidgetConstructors(UE::Editor::DataStorage::ICoreProvider& DataStorage,
+	UE::Editor::DataStorage::IUiProvider& DataStorageUi) const
+{
+	using namespace UE::Editor::DataStorage;
+	using namespace UE::Editor::DataStorage::Queries;
+
+	// Outliner
+	FConditions OutlinerColumns = TColumn<FFavoriteTag>() || TColumn<FTypedElementUObjectColumn>() || TColumn<FFolderTag>();
+
+	DataStorageUi.RegisterWidgetFactory<FTableFavoriteWidgetConstructor>(
+		DataStorageUi.FindPurpose(IUiProvider::FPurposeInfo("SceneOutliner", "Cell", NAME_None).GeneratePurposeID()), OutlinerColumns);
+}
+
+#undef LOCTEXT_NAMESPACE
diff --git a/Engine/Plugins/Experimental/EditorDataStorageFeatures/Source/UnifiedFavorites/Private/Widgets/UnifiedFavoritesTableWidget.h b/Engine/Plugins/Experimental/EditorDataStorageFeatures/Source/UnifiedFavorites/Private/Widgets/UnifiedFavoritesTableWidget.h
new file mode 100644
--- /dev/null
+++ b/Engine/Plugins/Experimental/EditorDataStorageFeatures/Source/UnifiedFavorites/Private/Widgets/UnifiedFavoritesTableWidget.h
@@ -0,0 +1,64 @@
+// Copyright Epic Games, Inc. All Rights Reserved.
+
+#pragma once
+
+#include "Elements/Interfaces/TypedElementDataStorageFactory.h"
+#include "Elements/Interfaces/TypedElementDataStorageUiInterface.h"
+
+#include "UnifiedFavoritesTableWidget.generated.h"
+
+//
+// Table Constructors
+//
+
+USTRUCT()
+struct FTableFavoriteWidgetConstructor : public FSimpleWidgetConstructor
+{
+	GENERATED_BODY()
+
+public:
+	FTableFavoriteWidgetConstructor();
+	~FTableFavoriteWidgetConstructor() override = default;
+
+	virtual TSharedPtr<SWidget> CreateWidget(
+		UE::Editor::DataStorage::ICoreProvider* DataStorage,
+		UE::Editor::DataStorage::IUiProvider* DataStorageUi,
+		UE::Editor::DataStorage::RowHandle TargetRow,
+		UE::Editor::DataStorage::RowHandle WidgetRow,
+		const UE::Editor::DataStorage::FMetaDataView& Arguments) override;
+	virtual TArray<TSharedPtr<const UE::Editor::DataStorage::FColumnSorterInterface>> ConstructColumnSorters(
+		UE::Editor::DataStorage::ICoreProvider* DataStorage, UE::Editor::DataStorage::IUiProvider* DataStorageUi,
+		const UE::Editor::DataStorage::FMetaDataView& Arguments) override;
+};
+
+USTRUCT()
+struct FTableFavoriteHeaderConstructor : public FSimpleWidgetConstructor
+{
+	GENERATED_BODY()
+
+public:
+	FTableFavoriteHeaderConstructor();
+	~FTableFavoriteHeaderConstructor() override = default;
+
+	virtual TSharedPtr<SWidget> CreateWidget(
+		UE::Editor::DataStorage::ICoreProvider* DataStorage,
+		UE::Editor::DataStorage::IUiProvider* DataStorageUi,
+		UE::Editor::DataStorage::RowHandle TargetRow,
+		UE::Editor::DataStorage::RowHandle WidgetRow,
+		const UE::Editor::DataStorage::FMetaDataView& Arguments) override;
+};
+
+//
+// Factory
+//
+
+UCLASS()
+class UUnifiedFavoritesTableFactory : public UEditorDataStorageFactory
+{
+	GENERATED_BODY()
+
+public:
+	~UUnifiedFavoritesTableFactory() override = default;
+
+	virtual void RegisterWidgetConstructors(UE::Editor::DataStorage::ICoreProvider& DataStorage, UE::Editor::DataStorage::IUiProvider& DataStorageUi) const override;
+};
diff --git a/Engine/Plugins/Experimental/EditorDataStorageFeatures/Source/UnifiedFavorites/Private/Widgets/UnifiedFavoritesWidget.cpp b/Engine/Plugins/Experimental/EditorDataStorageFeatures/Source/UnifiedFavorites/Private/Widgets/UnifiedFavoritesWidget.cpp
new file mode 100644
--- /dev/null
+++ b/Engine/Plugins/Experimental/EditorDataStorageFeatures/Source/UnifiedFavorites/Private/Widgets/UnifiedFavoritesWidget.cpp
@@ -0,0 +1,128 @@
+// Copyright Epic Games, Inc. All Rights Reserved.
+
+#include "UnifiedFavoritesWidget.h"
+
+#include "Elements/Columns/TypedElementMiscColumns.h"
+
+namespace UE::Editor::DataStorage
+{
+	void SUnifiedFavoritesWidget::Construct(const FArguments& InArgs, ICoreProvider* InDataStorage)
+	{
+		check(InDataStorage);
+		DataStorage = InDataStorage;
+		TargetRow = InArgs._RowHandle;
+
+		FavoriteBrush    = FAppStyle::GetBrush("Icons.Star");
+		NotFavoriteBrush = FAppStyle::GetBrush("PropertyWindow.Favorites_Disabled");
+
+		ActiveColor = InArgs._ActiveColor;
+		NotActiveColor = InArgs._NotActiveColor;
+
+		SImage::Construct(
+			SImage::FArguments()
+			.ColorAndOpacity(this, &SUnifiedFavoritesWidget::GetForegroundColor)
+			.Image(this, &SUnifiedFavoritesWidget::GetBrush)
+		);
+	}
+
+
+	bool SUnifiedFavoritesWidget::IsFavorite() const
+	{
+		return DataStorage->HasColumns<FFavoriteTag>(TargetRow);
+	}
+
+
+	void SUnifiedFavoritesWidget::ToggleFavorite() const
+	{
+		if (!IsInteractable())
+		{
+			return;
+		}
+
+		DataStorage->AddColumn<FTypedElementSyncFromWorldTag>(TargetRow);
+
+		if (DataStorage->HasColumns<FFavoriteTag>(TargetRow))
+		{
+			DataStorage->RemoveColumn<FFavoriteTag>(TargetRow);
+		}
+		else
+		{
+			DataStorage->AddColumn<FFavoriteTag>(TargetRow);
+		}
+	}
+
+	bool SUnifiedFavoritesWidget::IsInteractable() const
+	{
+		return DataStorage->IsRowAvailable(TargetRow);
+	}
+
+	void SUnifiedFavoritesWidget::SetTargetRow(RowHandle InTargetRow)
+	{
+		TargetRow = InTargetRow;
+	}
+
+	FReply SUnifiedFavoritesWidget::HandleClick() const
+	{
+		if (!IsEnabled())
+		{
+			return FReply::Unhandled();
+		}
+
+		ToggleFavorite();
+
+		return FReply::Handled();
+	}
+
+	FReply SUnifiedFavoritesWidget::OnMouseButtonDoubleClick(const FGeometry& InMyGeometry, const FPointerEvent& InMouseEvent)
+	{
+		return HandleClick();
+	}
+
+	FReply SUnifiedFavoritesWidget::OnMouseButtonDown(const FGeometry& MyGeometry, const FPointerEvent& MouseEvent)
+	{
+		if (MouseEvent.GetEffectingButton() == EKeys::LeftMouseButton)
+		{
+			return HandleClick();
+		}
+
+		return FReply::Unhandled();
+	}
+
+	FReply SUnifiedFavoritesWidget::OnMouseButtonUp(const FGeometry& MyGeometry, const FPointerEvent& MouseEvent)
+	{
+		if (MouseEvent.GetEffectingButton() == EKeys::LeftMouseButton)
+		{
+			return FReply::Handled();
+		}
+
+		return FReply::Unhandled();
+	}
+
+	const FSlateBrush* SUnifiedFavoritesWidget::GetBrush() const
+	{
+		return IsFavorite() ? FavoriteBrush : NotFavoriteBrush;
+	}
+
+	FSlateColor SUnifiedFavoritesWidget::GetForegroundColor() const
+	{
+		if (!IsInteractable())
+		{
+			return FLinearColor::Transparent;
+		}
+
+		if (!IsFavorite() && !IsHovered())
+		{
+			return NotActiveColor;
+		}
+
+		FLinearColor Color = ActiveColor.GetSpecifiedColor();
+
+		if (IsHovered())
+		{
+			constexpr float Darken = .8f;
+			Color *= Darken;
+		}
+
+		return Color;
+	}
+}
diff --git a/Engine/Plugins/Experimental/EditorDataStorageFeatures/Source/UnifiedFavorites/Private/Widgets/UnifiedFavoritesWidget.h b/Engine/Plugins/Experimental/EditorDataStorageFeatures/Source/UnifiedFavorites/Private/Widgets/UnifiedFavoritesWidget.h
new file mode 100644
--- /dev/null
+++ b/Engine/Plugins/Experimental/EditorDataStorageFeatures/Source/UnifiedFavorites/Private/Widgets/UnifiedFavoritesWidget.h
@@ -0,0 +1,61 @@
+// Copyright Epic Games, Inc. All Rights Reserved.
+
+#pragma once
+
+#include "Elements/Interfaces/TypedElementDataStorageInterface.h"
+#include "Widgets/Images/SImage.h"
+
+
+namespace UE::Editor::DataStorage
+{
+	/** Widget responsible for managing the favorite state of a TEDS row handle */
+	class SUnifiedFavoritesWidget final : public SImage
+	{
+	public:
+		SLATE_BEGIN_ARGS(SUnifiedFavoritesWidget) :
+			_ActiveColor(FAppStyle::Get().GetSlateColor("Colors.Foreground")),
+			_NotActiveColor(FAppStyle::Get().GetSlateColor("Colors.Transparent"))
+			{
+			}
+			SLATE_ARGUMENT(FSlateColor, ActiveColor)
+			SLATE_ARGUMENT(FSlateColor, NotActiveColor)
+			SLATE_ARGUMENT(UE::Editor::DataStorage::RowHandle, RowHandle)
+		SLATE_END_ARGS()
+
+			/** Construct this widget */
+			void Construct(const FArguments& InArgs, ICoreProvider* InDataStorage);
+
+			virtual bool IsInteractable() const override;
+
+			/** Directly sets the resolved target row (called by the mapping resolution processor). */
+			void SetTargetRow(RowHandle InTargetRow);
+
+	private:
+		bool IsFavorite() const;
+
+		void ToggleFavorite() const;
+
+		FReply HandleClick() const;
+
+		virtual FReply OnMouseButtonDoubleClick(const FGeometry& InMyGeometry, const FPointerEvent& InMouseEvent) override;
+
+		virtual FReply OnMouseButtonDown(const FGeometry& MyGeometry, const FPointerEvent& MouseEvent) override;
+
+		virtual FReply OnMouseButtonUp(const FGeometry& MyGeometry, const FPointerEvent& MouseEvent) override;
+
+		virtual FSlateColor GetForegroundColor() const override;
+
+		const FSlateBrush* GetBrush() const;
+
+	private:
+		ICoreProvider* DataStorage;
+
+		RowHandle TargetRow = InvalidRowHandle;
+
+		FSlateColor ActiveColor;
+		FSlateColor NotActiveColor;
+
+		const FSlateBrush* FavoriteBrush;
+		const FSlateBrush* NotFavoriteBrush;
+	};
+}
diff --git a/Engine/Plugins/Experimental/EditorDataStorageFeatures/Source/UnifiedFavorites/UnifiedFavorites.Build.cs b/Engine/Plugins/Experimental/EditorDataStorageFeatures/Source/UnifiedFavorites/UnifiedFavorites.Build.cs
new file mode 100644
--- /dev/null
+++ b/Engine/Plugins/Experimental/EditorDataStorageFeatures/Source/UnifiedFavorites/UnifiedFavorites.Build.cs
@@ -0,0 +1,51 @@
+// Copyright Epic Games, Inc. All Rights Reserved.
+
+using UnrealBuildTool;
+
+public class UnifiedFavorites : ModuleRules
+{
+	public UnifiedFavorites(ReadOnlyTargetRules Target) : base(Target)
+	{
+		// Enable truncation warnings in this module
+		CppCompileWarningSettings.UnsafeTypeCastWarningLevel = WarningLevel.Warning;
+
+		if (Target.bBuildEditor)
+		{
+			PublicIncludePaths.AddRange(new string[] {});
+			PrivateIncludePaths.AddRange(new string[] {});
+
+			PublicDependencyModuleNames.AddRange(
+				new string[]
+				{
+					"Core",
+					"CoreUObject",
+					"EditorFramework",
+					"Engine",
+					"SceneOutliner",
+					"Slate",
+					"SlateCore",
+					"TypedElementFramework",
+					"TypedElementRuntime"
+				});
+				
+			PrivateDependencyModuleNames.AddRange(
+				new string[]
+				{
+					"WorkspaceMenuStructure",
+					"ToolMenus",
+					"ApplicationCore",
+					"InputCore",
+					"TedsQueryStack",
+					"TedsActorCompatibility",
+					"TedsOperations",
+					"LevelEditor",
+					"LevelInstanceEditor",
+					"TedsOutliner",
+					"ToolWidgets",
+					"TedsTableViewer",
+					"ContentBrowser",
+					"TedsAssetData"
+				});
+		}
+	}
+}
\ No newline at end of file
