TEDS editor hierarchies:

Consolidate actor/component/world/level hierarchies into one "EditorObjectHierarchy" and update all old references to the individual hierarchies
Add a hierarchy from ULevel -> UWorld
Add a column to identify the persistent level in TEDS and hide it from UI since we don't want it to be shown in the Outliner
Modify GetParentRow() logic in TEDS Outliner to walk up the parent chain until it finds a valid parent

#Synced-CL 51760981

diff --git a/Engine/Plugins/Experimental/EditorDataStorageFeatures/Source/TedsActorCompatibility/Private/ActorComponentDebugHierarchyWidget/ActorComponentTree.cpp b/Engine/Plugins/Experimental/EditorDataStorageFeatures/Source/TedsActorCompatibility/Private/ActorComponentDebugHierarchyWidget/ActorComponentTree.cpp
--- a/Engine/Plugins/Experimental/EditorDataStorageFeatures/Source/TedsActorCompatibility/Private/ActorComponentDebugHierarchyWidget/ActorComponentTree.cpp
+++ b/Engine/Plugins/Experimental/EditorDataStorageFeatures/Source/TedsActorCompatibility/Private/ActorComponentDebugHierarchyWidget/ActorComponentTree.cpp
@@ -14,6 +14,7 @@
 #include "Elements/Columns/TypedElementCompatibilityColumns.h"
 #include "Elements/Columns/TypedElementSelectionColumns.h"
 #include "DataStorage/Features.h"
+#include "Factories/TedsEditorHierarchyFactory.h"
 #include "Framework/Docking/TabManager.h"
 #include "QueryStackNodes/RowQueryCallbackResultsNode.h"
 #include "Widgets/STedsTableViewer.h"
@@ -39,7 +40,7 @@ namespace UE::Editor::DataStorage::ActorCompatibility
 
 				ICoreProvider* DataStorageInterface = GetMutableDataStorageFeature<ICoreProvider>(StorageFeatureName);
 				
-				FHierarchyHandle HierarchyHandle = DataStorageInterface->FindHierarchyByName(TEXT("ObjectComponent"));
+				FHierarchyHandle HierarchyHandle = DataStorageInterface->FindHierarchyByName(UTedsEditorHierarchyFactory::EditorObjectHierarchyName);
 
 				if (DataStorageInterface->IsValidHierarchyHandle(HierarchyHandle))
 				{
diff --git a/Engine/Plugins/Experimental/EditorDataStorageFeatures/Source/TedsActorCompatibility/Private/Compatibility/TedsActorComponentFactory.cpp b/Engine/Plugins/Experimental/EditorDataStorageFeatures/Source/TedsActorCompatibility/Private/Compatibility/TedsActorComponentFactory.cpp
--- a/Engine/Plugins/Experimental/EditorDataStorageFeatures/Source/TedsActorCompatibility/Private/Compatibility/TedsActorComponentFactory.cpp
+++ b/Engine/Plugins/Experimental/EditorDataStorageFeatures/Source/TedsActorCompatibility/Private/Compatibility/TedsActorComponentFactory.cpp
@@ -10,6 +10,7 @@
 #include "Elements/Columns/TypedElementTypeInfoColumns.h"
 #include "Elements/Framework/TypedElementQueryBuilder.h"
 #include "Elements/Interfaces/TypedElementDataStorageCompatibilityInterface.h"
+#include "Factories/TedsEditorHierarchyFactory.h"
 #include "GameFramework/Actor.h"
 
 #include UE_INLINE_GENERATED_CPP_BY_NAME(TedsActorComponentFactory)
@@ -38,22 +39,6 @@ void UTedsActorComponentFactory::RegisterTables(UE::Editor::DataStorage::ICorePr
 	}
 }
 
-void UTedsActorComponentFactory::RegisterHierarchies(UE::Editor::DataStorage::ICoreProvider& DataStorage)
-{
-	using namespace UE::Editor::DataStorage;
-	
-	if (bEnableActorComponentCompatibility)
-	{
-		const FHierarchyRegistrationParams HierarchyParams
-			{
-				.Name = TEXT("ObjectComponent"),
-				.bEnableParentChangedColumn = true
-			};
-
-		DataStorage.RegisterHierarchy(HierarchyParams);
-	}
-}
-
 void UTedsActorComponentFactory::RegisterQueries(UE::Editor::DataStorage::ICoreProvider& DataStorage)
 {
 	using namespace UE::Editor::DataStorage;
@@ -157,7 +142,7 @@ void UTedsActorComponentFactory::RegisterQueries(UE::Editor::DataStorage::ICoreP
 						}
 					}
 				})
-				.AccessesHierarchy(TEXT("ObjectComponent"))
+				.AccessesHierarchy(UTedsEditorHierarchyFactory::EditorObjectHierarchyName)
 				.Where()
 					.All<FActorComponentTypeTag, FTypedElementSyncFromWorldTag>()
 				.Compile());
diff --git a/Engine/Plugins/Experimental/EditorDataStorageFeatures/Source/TedsActorCompatibility/Private/Compatibility/TedsActorComponentFactory.h b/Engine/Plugins/Experimental/EditorDataStorageFeatures/Source/TedsActorCompatibility/Private/Compatibility/TedsActorComponentFactory.h
--- a/Engine/Plugins/Experimental/EditorDataStorageFeatures/Source/TedsActorCompatibility/Private/Compatibility/TedsActorComponentFactory.h
+++ b/Engine/Plugins/Experimental/EditorDataStorageFeatures/Source/TedsActorCompatibility/Private/Compatibility/TedsActorComponentFactory.h
@@ -15,6 +15,5 @@ class UTedsActorComponentFactory : public UEditorDataStorageFactory
 
 public:
 	virtual void RegisterTables(UE::Editor::DataStorage::ICoreProvider& DataStorage) override;
-	virtual void RegisterHierarchies(UE::Editor::DataStorage::ICoreProvider& DataStorage) override;
 	virtual void RegisterQueries(UE::Editor::DataStorage::ICoreProvider& DataStorage) override;
 };
\ No newline at end of file
diff --git a/Engine/Plugins/Experimental/EditorDataStorageFeatures/Source/TedsActorCompatibility/Private/Compatibility/TedsActorParentQueries.cpp b/Engine/Plugins/Experimental/EditorDataStorageFeatures/Source/TedsActorCompatibility/Private/Compatibility/TedsActorParentQueries.cpp
--- a/Engine/Plugins/Experimental/EditorDataStorageFeatures/Source/TedsActorCompatibility/Private/Compatibility/TedsActorParentQueries.cpp
+++ b/Engine/Plugins/Experimental/EditorDataStorageFeatures/Source/TedsActorCompatibility/Private/Compatibility/TedsActorParentQueries.cpp
@@ -8,13 +8,14 @@
 #include "Elements/Columns/TypedElementMiscColumns.h"
 #include "Elements/Framework/TypedElementQueryBuilder.h"
 #include "Engine/Engine.h"
+#include "Engine/Level.h"
+#include "Factories/TedsEditorHierarchyFactory.h"
 #include "GameFramework/Actor.h"
 
 #include UE_INLINE_GENERATED_CPP_BY_NAME(TedsActorParentQueries)
 
 namespace UE::Editor::DataStorage::ActorParentQueries::Private
 {
-	static FName ActorEditorHierarchyName(TEXT("ActorEditorHierarchy"));
 	static bool bAddParentColumnToActors = true;
 
 	// Cvar to add parenting info for actors to TEDS
@@ -22,14 +23,33 @@ namespace UE::Editor::DataStorage::ActorParentQueries::Private
 		TEXT("TEDS.AddParentColumnToActors"),
 		bAddParentColumnToActors,
 		TEXT("Mirror parent information for actors to TEDS (only works when set on startup)"));
+	
+	// Helper to get the first viable parent for an actor in the editor
+	const UObject* GetEditorParentForActor(const AActor* InActor)
+	{
+		// For the editor hierarchy, the order to check is Parent Actor -> Owning Level -> Owning World
+		if (const AActor* ParentActor = InActor->GetAttachParentActor())
+		{
+			return ParentActor;
+		}
+		else if (const ULevel* ParentLevel = InActor->GetLevel())
+		{
+			return ParentLevel;
+		}
+		else if (const UWorld* ParentWorld = InActor->GetWorld())
+		{
+			return ParentWorld;
+		}
+		
+		return nullptr;
+	}
 };
 
 void UActorParentDataStorageFactory::RegisterQueries(UE::Editor::DataStorage::ICoreProvider& InDataStorage)
 {
 	if (UE::Editor::DataStorage::ActorParentQueries::Private::bAddParentColumnToActors)
 	{
-		RegisterAddParentColumn(InDataStorage);
-		RegisterUpdateOrRemoveParentColumn(InDataStorage);
+		RegisterActorParentQueries(InDataStorage);
 	}
 }
 
@@ -41,25 +61,12 @@ void UActorParentDataStorageFactory::PreRegister(UE::Editor::DataStorage::ICoreP
 	DataStorage = &InDataStorage;
 }
 
-void UActorParentDataStorageFactory::RegisterHierarchies(UE::Editor::DataStorage::ICoreProvider& InDataStorage)
-{
-	using namespace UE::Editor::DataStorage;
-	
-	UE::Editor::DataStorage::FHierarchyRegistrationParams ActorHierarchyParams
-		{
-			.Name = ActorParentQueries::Private::ActorEditorHierarchyName, 
-			.bEnableParentChangedColumn = true
-		};
-	
-	EditorActorHierarchyHandle = InDataStorage.RegisterHierarchy(ActorHierarchyParams);
-}
-
 void UActorParentDataStorageFactory::PreShutdown(UE::Editor::DataStorage::ICoreProvider& InDataStorage)
 {
 	GEngine->OnLevelActorDetached().RemoveAll(this);
 }
 
-void UActorParentDataStorageFactory::RegisterAddParentColumn(UE::Editor::DataStorage::ICoreProvider& InDataStorage) const
+void UActorParentDataStorageFactory::RegisterActorParentQueries(UE::Editor::DataStorage::ICoreProvider& InDataStorage) const
 {
 	using namespace UE::Editor::DataStorage;
 	using namespace UE::Editor::DataStorage::Queries;
@@ -73,7 +80,7 @@ void UActorParentDataStorageFactory::RegisterAddParentColumn(UE::Editor::DataSto
 			{
 				if (const AActor* ActorInstance = Cast<AActor>(Actor.Object))
 				{
-					if (AActor* Parent = ActorInstance->GetAttachParentActor())
+					if (const UObject* Parent = ActorParentQueries::Private::GetEditorParentForActor(ActorInstance))
 					{
 						FMapKeyView IdKey = FMapKeyView(Parent);
 						RowHandle ParentRow = Context.LookupMappedRow(ICompatibilityProvider::ObjectMappingDomain, IdKey);
@@ -88,17 +95,11 @@ void UActorParentDataStorageFactory::RegisterAddParentColumn(UE::Editor::DataSto
 					}
 				}
 			})
-		.AccessesHierarchy(ActorParentQueries::Private::ActorEditorHierarchyName)
+		.AccessesHierarchy(UTedsEditorHierarchyFactory::EditorObjectHierarchyName)
 		.Where()
 			.All<FTypedElementSyncFromWorldTag, FTypedElementActorTag>()
 			.None(InDataStorage.GetChildTagType(EditorActorHierarchyHandle))
 		.Compile());
-}
-
-void UActorParentDataStorageFactory::RegisterUpdateOrRemoveParentColumn(UE::Editor::DataStorage::ICoreProvider& InDataStorage) const
-{
-	using namespace UE::Editor::DataStorage;
-	using namespace UE::Editor::DataStorage::Queries;
 
 	InDataStorage.RegisterQuery(
 		Select(
@@ -109,9 +110,9 @@ void UActorParentDataStorageFactory::RegisterUpdateOrRemoveParentColumn(UE::Edit
 			{
 				if (const AActor* ActorInstance = Cast<AActor>(Actor.Object))
 				{
-					if (AActor* ParentActor = ActorInstance->GetAttachParentActor())
+					if (const UObject* Parent = ActorParentQueries::Private::GetEditorParentForActor(ActorInstance))
 					{
-						FMapKeyView IdKey = FMapKeyView(ParentActor);
+						FMapKeyView IdKey = FMapKeyView(Parent);
 						RowHandle NewParentRow = Context.LookupMappedRow(ICompatibilityProvider::ObjectMappingDomain, IdKey);
 						RowHandle CurrentParentRow = Context.GetParentRow(TargetRow);
 
@@ -123,7 +124,7 @@ void UActorParentDataStorageFactory::RegisterUpdateOrRemoveParentColumn(UE::Edit
 							}
 							else
 							{
-								Context.SetUnresolvedParent(TargetRow, FMapKey(ParentActor), ICompatibilityProvider::ObjectMappingDomain);
+								Context.SetUnresolvedParent(TargetRow, FMapKey(Parent), ICompatibilityProvider::ObjectMappingDomain);
 							}
 						}
 						return;
@@ -133,7 +134,7 @@ void UActorParentDataStorageFactory::RegisterUpdateOrRemoveParentColumn(UE::Edit
 				// If we have reached here, the actor does not have a parent
 				Context.SetParentRow(TargetRow, InvalidRowHandle);
 			})
-		.AccessesHierarchy(ActorParentQueries::Private::ActorEditorHierarchyName)
+		.AccessesHierarchy(UTedsEditorHierarchyFactory::EditorObjectHierarchyName)
 		.Where()
 			.All<FTypedElementActorTag, FTypedElementSyncFromWorldTag>()
 		.Compile());
diff --git a/Engine/Plugins/Experimental/EditorDataStorageFeatures/Source/TedsActorCompatibility/Private/Compatibility/TedsActorParentQueries.h b/Engine/Plugins/Experimental/EditorDataStorageFeatures/Source/TedsActorCompatibility/Private/Compatibility/TedsActorParentQueries.h
--- a/Engine/Plugins/Experimental/EditorDataStorageFeatures/Source/TedsActorCompatibility/Private/Compatibility/TedsActorParentQueries.h
+++ b/Engine/Plugins/Experimental/EditorDataStorageFeatures/Source/TedsActorCompatibility/Private/Compatibility/TedsActorParentQueries.h
@@ -19,21 +19,14 @@ public:
 
 	virtual void RegisterQueries(UE::Editor::DataStorage::ICoreProvider& DataStorage) override;
 	virtual void PreRegister(UE::Editor::DataStorage::ICoreProvider& DataStorage) override;
-	virtual void RegisterHierarchies(UE::Editor::DataStorage::ICoreProvider& DataStorage) override;
 	virtual void PreShutdown(UE::Editor::DataStorage::ICoreProvider& DataStorage) override;
 
 private:
 	/**
-	 * Checks rows with actors that don't have a parent column yet if one needs to be added whenever
-	 * the row is marked for updates.
+	 * Sets the hierarchy for an actor in the editor
 	 */
-	void RegisterAddParentColumn(UE::Editor::DataStorage::ICoreProvider& DataStorage) const;
-	/**
-	 * Updates the parent column with the parent from the actor or removes it if there's no parent associated
-	 * with the actor anymore.
-	 */
-	void RegisterUpdateOrRemoveParentColumn(UE::Editor::DataStorage::ICoreProvider& DataStorage) const;
-
+	void RegisterActorParentQueries(UE::Editor::DataStorage::ICoreProvider& DataStorage) const;
+	
 	/**
 	 * Function called when an actor is detached from its parent in the world
 	 */
diff --git a/Engine/Plugins/Experimental/EditorDataStorageFeatures/Source/TedsEditorCompatibility/Private/Factories/TedsEditorHierarchyFactory.cpp b/Engine/Plugins/Experimental/EditorDataStorageFeatures/Source/TedsEditorCompatibility/Private/Factories/TedsEditorHierarchyFactory.cpp
new file mode 100644
--- /dev/null
+++ b/Engine/Plugins/Experimental/EditorDataStorageFeatures/Source/TedsEditorCompatibility/Private/Factories/TedsEditorHierarchyFactory.cpp
@@ -0,0 +1,20 @@
+﻿// Copyright Epic Games, Inc. All Rights Reserved.
+
+#include "Factories/TedsEditorHierarchyFactory.h"
+
+#include "Elements/Interfaces/TypedElementDataStorageInterface.h"
+
+FName UTedsEditorHierarchyFactory::EditorObjectHierarchyName("EditorObjectHierarchy");
+
+void UTedsEditorHierarchyFactory::RegisterHierarchies(UE::Editor::DataStorage::ICoreProvider& DataStorage)
+{
+	using namespace UE::Editor::DataStorage;
+	
+	UE::Editor::DataStorage::FHierarchyRegistrationParams WorldHierarchyParams
+	{
+		.Name = EditorObjectHierarchyName, 
+		.bEnableParentChangedColumn = true
+	};
+	
+	DataStorage.RegisterHierarchy(WorldHierarchyParams);
+}
\ No newline at end of file
diff --git a/Engine/Plugins/Experimental/EditorDataStorageFeatures/Source/TedsEditorCompatibility/Private/Level/TedsLevelFactory.cpp b/Engine/Plugins/Experimental/EditorDataStorageFeatures/Source/TedsEditorCompatibility/Private/Level/TedsLevelFactory.cpp
--- a/Engine/Plugins/Experimental/EditorDataStorageFeatures/Source/TedsEditorCompatibility/Private/Level/TedsLevelFactory.cpp
+++ b/Engine/Plugins/Experimental/EditorDataStorageFeatures/Source/TedsEditorCompatibility/Private/Level/TedsLevelFactory.cpp
@@ -7,11 +7,13 @@
 #include "Elements/Columns/TypedElementIconOverrideColumns.h"
 #include "Elements/Columns/TypedElementLabelColumns.h"
 #include "Elements/Columns/TypedElementMiscColumns.h"
+#include "Elements/Columns/TypedElementUIColumns.h"
 #include "Elements/Framework/TypedElementQueryBuilder.h"
 #include "Elements/Interfaces/TypedElementDataStorageCompatibilityInterface.h"
 #include "Elements/Interfaces/TypedElementDataStorageInterface.h"
 #include "Engine/Level.h"
 #include "Engine/World.h"
+#include "Factories/TedsEditorHierarchyFactory.h"
 #include "Misc/PackageName.h"
 #include "Styling/SlateIconFinder.h"
 #include "UObject/Package.h"
@@ -84,4 +86,62 @@ void UTedsLevelFactory::RegisterQueries(UE::Editor::DataStorage::ICoreProvider&
 		.Where()
 			.All<FLevelTag>()
 		.Compile());
+	
+	DataStorage.RegisterQuery(
+		Select(TEXT("Hide persistent level from UI"),
+			FObserver::OnAdd<FTypedElementUObjectColumn>()
+				.SetExecutionMode(EExecutionMode::GameThread),
+			[](IQueryContext& Context, RowHandle Row, const FTypedElementUObjectColumn& LevelColumn)
+			{
+				if (const ULevel* Level = Cast<ULevel>(LevelColumn.Object); Level != nullptr)
+				{
+					if (Level->IsPersistentLevel())
+					{
+						// The persistent level is hidden from editor UI like the Outliner
+						Context.AddColumns<FHideRowFromUITag>(Row);
+					}
+				}
+			})
+		.Where()
+			.All<FLevelTag>()
+		.Compile());
+	
+	DataStorage.RegisterQuery(
+		Select(
+			TEXT("Sync Level -> World hierarchy"),
+			FProcessor(EQueryTickPhase::PrePhysics, DataStorage.GetQueryTickGroupName(EQueryTickGroups::SyncExternalToDataStorage))
+			.SetExecutionMode(EExecutionMode::GameThread),
+			[](IQueryContext& Context, RowHandle TargetRow, const FTypedElementUObjectColumn& LevelColumn)
+			{
+				if (const ULevel* Level = Cast<ULevel>(LevelColumn.Object); Level != nullptr)
+				{
+					if (const UWorld* ParentWorld = Level->GetWorld())
+					{
+						FMapKeyView IdKey = FMapKeyView(ParentWorld);
+						RowHandle NewParentRow = Context.LookupMappedRow(ICompatibilityProvider::ObjectMappingDomain, IdKey);
+						RowHandle CurrentParentRow = Context.GetParentRow(TargetRow);
+
+						if (CurrentParentRow != NewParentRow)
+						{
+							if (Context.IsRowAvailable(NewParentRow))
+							{
+								Context.SetParentRow(TargetRow, NewParentRow);
+							}
+							else
+							{
+								Context.SetUnresolvedParent(TargetRow, FMapKey(ParentWorld), ICompatibilityProvider::ObjectMappingDomain);
+							}
+						}
+					}
+					// In case this Level doesn't have a world for some reason, we add a fallback to reset the hierarchy
+					else
+					{
+						Context.SetParentRow(TargetRow, InvalidRowHandle);
+					}
+				}
+			})
+		.AccessesHierarchy(UTedsEditorHierarchyFactory::EditorObjectHierarchyName)
+		.Where()
+			.All<FTypedElementSyncFromWorldTag, FLevelTag>()
+		.Compile());
 }
diff --git a/Engine/Plugins/Experimental/EditorDataStorageFeatures/Source/TedsEditorCompatibility/Public/Factories/TedsEditorHierarchyFactory.h b/Engine/Plugins/Experimental/EditorDataStorageFeatures/Source/TedsEditorCompatibility/Public/Factories/TedsEditorHierarchyFactory.h
new file mode 100644
--- /dev/null
+++ b/Engine/Plugins/Experimental/EditorDataStorageFeatures/Source/TedsEditorCompatibility/Public/Factories/TedsEditorHierarchyFactory.h
@@ -0,0 +1,24 @@
+// Copyright Epic Games, Inc. All Rights Reserved.
+
+#pragma once
+
+#include "DataStorage/CommonTypes.h"
+#include "DataStorage/Handles.h"
+#include "Elements/Interfaces/TypedElementDataStorageFactory.h"
+#include "UObject/ObjectMacros.h"
+
+#include "TedsEditorHierarchyFactory.generated.h"
+
+UCLASS()
+class UTedsEditorHierarchyFactory : public UEditorDataStorageFactory
+{
+	GENERATED_BODY()
+
+public:
+	~UTedsEditorHierarchyFactory() override = default;
+	
+	// Hierarchy for objects in the editor (e.g UWorld, ULevel, AActor)
+	TEDSEDITORCOMPATIBILITY_API static FName EditorObjectHierarchyName;
+
+	TEDSEDITORCOMPATIBILITY_API virtual void RegisterHierarchies(UE::Editor::DataStorage::ICoreProvider& DataStorage) override;
+};
\ No newline at end of file
diff --git a/Engine/Plugins/Experimental/EditorDataStorageFeatures/Source/TedsOutliner/Private/TedsOutlinerImpl.cpp b/Engine/Plugins/Experimental/EditorDataStorageFeatures/Source/TedsOutliner/Private/TedsOutlinerImpl.cpp
--- a/Engine/Plugins/Experimental/EditorDataStorageFeatures/Source/TedsOutliner/Private/TedsOutlinerImpl.cpp
+++ b/Engine/Plugins/Experimental/EditorDataStorageFeatures/Source/TedsOutliner/Private/TedsOutlinerImpl.cpp
@@ -872,7 +872,6 @@ DataStorage::RowHandle FTedsOutlinerImpl::GetParentRow(DataStorage::RowHandle In
 		return DataStorage::InvalidRowHandle;
 	}
 	
-	RowHandle ParentRowHandle = DataStorage::InvalidRowHandle;
 	DataStorage::FRowHandleArrayView Rows;
 
 	if (ShouldForceShowParentRows())
@@ -891,19 +890,22 @@ DataStorage::RowHandle FTedsOutlinerImpl::GetParentRow(DataStorage::RowHandle In
 		Rows = RowHandleSortNode->GetRows();
 	}
 	
-	// Iterate through each parent in case there are multiple hierarchies, so that we don't stop if the first parent doesn't match our condition and
-	// instead keep looking
-	HierarchyData->ForEachImmediateParent(*Storage, InRowHandle, [&ParentRowHandle, &Rows](const ICoreProvider& Storage, RowHandle ParentRow)
+	RowHandle ParentRowHandle = InRowHandle;
+
+	// Sanity check in case there are loops in the hierarchy
+	int32 MaxIterations = 10000;
+	do
+	{
+		if (!ensureMsgf(MaxIterations-- > 0, TEXT("Reached max iterations when looking for parent of row %llu"), InRowHandle))
 		{
-			// The results are sorted by row handle so a binary search is done here
-			if (Rows.Contains(ParentRow))
-			{
-				ParentRowHandle = ParentRow;
-				return false; // We have found a valid parent, no need to continue iteration
-			}
-			
-			return true;
-		});
+			ParentRowHandle = InvalidRowHandle;
+			break;
+		}
+		
+		ParentRowHandle = HierarchyData->GetParent(*Storage, ParentRowHandle);
+	}
+	// Walk up the parent chain until we find a parent that is present our query stack results or we reached the root of this hierarchy
+	while (Storage->IsRowAvailable(ParentRowHandle) && !Rows.Contains(ParentRowHandle));
 	
 	if (!Storage->IsRowAvailable(ParentRowHandle))
 	{
diff --git a/Engine/Plugins/Experimental/EditorDataStorageFeatures/Source/TedsOutliner/Private/TedsOutlinerModule.cpp b/Engine/Plugins/Experimental/EditorDataStorageFeatures/Source/TedsOutliner/Private/TedsOutlinerModule.cpp
--- a/Engine/Plugins/Experimental/EditorDataStorageFeatures/Source/TedsOutliner/Private/TedsOutlinerModule.cpp
+++ b/Engine/Plugins/Experimental/EditorDataStorageFeatures/Source/TedsOutliner/Private/TedsOutlinerModule.cpp
@@ -29,6 +29,7 @@
 #include "Compatibility/SceneOutlinerTedsBridge.h"
 #include "Compatibility/SceneOutlinerRowHandleColumn.h"
 #include "Layers/Columns/LayersColumns.h"
+#include "Factories/TedsEditorHierarchyFactory.h"
 #include "Widgets/Docking/SDockTab.h"
 
 #define LOCTEXT_NAMESPACE "TedsOutlinerModule"
@@ -211,8 +212,7 @@ TSharedRef<SWidget> FTedsOutlinerModule::CreateLevelEditorTedsOutliner()
 	Params.ColumnDescription = GetLevelEditorTedsOutlinerColumnDescription();
 	
 	const TArray<FHierarchyHandle> Hierarchies{
-		Storage->FindHierarchyByName(TEXT("ActorEditorHierarchy")),
-		Storage->FindHierarchyByName(TEXT("ObjectComponent"))
+		Storage->FindHierarchyByName(UTedsEditorHierarchyFactory::EditorObjectHierarchyName)
 		};
 	
 	Params.HierarchyData = MakeShared<FTedsOutlinerMultiHierarchyInterface>(MakeShared<FHierarchyViewerMultiData>(*Storage, Hierarchies));
diff --git a/Engine/Plugins/Experimental/EditorDataStorageFeatures/Source/TedsOutliner/TedsOutliner.Build.cs b/Engine/Plugins/Experimental/EditorDataStorageFeatures/Source/TedsOutliner/TedsOutliner.Build.cs
--- a/Engine/Plugins/Experimental/EditorDataStorageFeatures/Source/TedsOutliner/TedsOutliner.Build.cs
+++ b/Engine/Plugins/Experimental/EditorDataStorageFeatures/Source/TedsOutliner/TedsOutliner.Build.cs
@@ -45,7 +45,8 @@ public class TedsOutliner : ModuleRules
 					"TedsQueryStack",
 					"TedsActorCompatibility",
 					"TedsOperations",
-					"LevelInstanceEditor"
+					"LevelInstanceEditor",
+					"TedsEditorCompatibility"
 				});
 
 			DynamicallyLoadedModuleNames.AddRange(new string[] {});
