STedsTableViewer: Add argument for generic metadata and fix a bug where the primary column arg wasn't passed into the model

#Synced-CL 51760981

diff --git a/Engine/Plugins/Experimental/EditorDataStorageFeatures/Source/TedsTableViewer/Private/Widgets/STedsTableViewer.cpp b/Engine/Plugins/Experimental/EditorDataStorageFeatures/Source/TedsTableViewer/Private/Widgets/STedsTableViewer.cpp
--- a/Engine/Plugins/Experimental/EditorDataStorageFeatures/Source/TedsTableViewer/Private/Widgets/STedsTableViewer.cpp
+++ b/Engine/Plugins/Experimental/EditorDataStorageFeatures/Source/TedsTableViewer/Private/Widgets/STedsTableViewer.cpp
@@ -71,10 +71,12 @@ namespace UE::Editor::DataStorage
 		Model = MakeShared<FTedsTableViewerModel>(FTedsTableViewerModel::FArguments{
 			.RowQueryStack = InArgs._QueryStack,
 			.RequestedColumns = InArgs._Columns,
+			.GenericMetaData = InArgs._GenericMetaData,
 			.ColumnMetaData = InArgs._ColumnMetaData,
 			.CellWidgetPurpose = CellWidgetPurpose,
 			.HeaderWidgetPurpose = HeaderWidgetPurpose,
-			.IsItemVisibleDelegate = FTedsTableViewerModel::FIsItemVisible::CreateSP(this, &STedsTableViewer::IsItemVisible)
+			.IsItemVisibleDelegate = FTedsTableViewerModel::FIsItemVisible::CreateSP(this, &STedsTableViewer::IsItemVisible),
+			.PrimaryColumn = InArgs._PrimaryColumn
 		});
 
 		InstanceCounter++;
diff --git a/Engine/Plugins/Experimental/EditorDataStorageFeatures/Source/TedsTableViewer/Public/Widgets/STedsTableViewer.h b/Engine/Plugins/Experimental/EditorDataStorageFeatures/Source/TedsTableViewer/Public/Widgets/STedsTableViewer.h
--- a/Engine/Plugins/Experimental/EditorDataStorageFeatures/Source/TedsTableViewer/Public/Widgets/STedsTableViewer.h
+++ b/Engine/Plugins/Experimental/EditorDataStorageFeatures/Source/TedsTableViewer/Public/Widgets/STedsTableViewer.h
@@ -67,6 +67,9 @@ namespace UE::Editor::DataStorage
 		// The Columns that this table viewer will display
 		SLATE_ARGUMENT(TArray<TWeakObjectPtr<const UScriptStruct>>, Columns)
 
+		// Generic Metadata about all widgets table viewer is displaying
+		SLATE_ARGUMENT(TSharedPtr<FMetaData>, GenericMetaData)
+
 		// Metadata about the columns the table viewer is displaying
 		SLATE_ARGUMENT(TArray<TSharedPtr<FColumnMetaData>>, ColumnMetaData)
 
