Add async support for initial pull of Asset Registry into TEDs

#Synced-CL 51330672

diff --git a/Engine/Plugins/Experimental/EditorDataStorageFeatures/Source/TedsAssetData/Private/TedsAssetData.cpp b/Engine/Plugins/Experimental/EditorDataStorageFeatures/Source/TedsAssetData/Private/TedsAssetData.cpp
--- a/Engine/Plugins/Experimental/EditorDataStorageFeatures/Source/TedsAssetData/Private/TedsAssetData.cpp
+++ b/Engine/Plugins/Experimental/EditorDataStorageFeatures/Source/TedsAssetData/Private/TedsAssetData.cpp
@@ -8,6 +8,7 @@
 #include "Async/ParallelFor.h"
 #include "Containers/Array.h"
 #include "Containers/ChunkedArray.h"
+#include "Containers/Queue.h"
 #include "DataStorage/Queries/Types.h"
 #include "Elements/Columns/TypedElementFolderColumns.h"
 #include "Elements/Columns/TypedElementSlateWidgetColumns.h"
@@ -319,17 +320,20 @@ void FTedsAssetData::Initialize()
 
 	// Init with the data existing at moment in asset registry
 	UE::AssetRegistry::FFiltering::InitializeShouldSkipAsset();
-	TChunkedArray<FAssetData> CachedAssets;
+	TQueue<FAssetData, EQueueMode::Mpsc> CachedAssets;
 	AssetRegistry.EnumerateAllAssets([&CachedAssets](const FAssetData& AssetData)
 		{
 			if (!UE::AssetRegistry::FFiltering::ShouldSkipAsset(AssetData.AssetClassPath, AssetData.PackageFlags) &&
 				ContentBrowserDataUtils::IsPrimaryAsset(AssetData))
 			{
-				CachedAssets.AddElement(AssetData);
+				CachedAssets.Enqueue(AssetData);
 			}
 			return true;
-		});
-	CachedAssets.MoveToLinearArray(AddedAssets);
+		}, UE::AssetRegistry::EEnumerateAssetsFlags::Parallel);
+
+	while (CachedAssets.Dequeue(AddedAssets.AddDefaulted_GetRef()))
+	{
+	}
 
 	TChunkedArray<FName> CachedPaths;
 	AssetRegistry.EnumerateAllCachedPaths([&CachedPaths](FName Name)
