TedsOutliner: Fix hidden items not being properly selectable.

#Synced-CL 51852197

diff --git a/Engine/Plugins/Experimental/EditorDataStorageFeatures/Source/TedsOutliner/Private/TedsOutlinerImpl.cpp b/Engine/Plugins/Experimental/EditorDataStorageFeatures/Source/TedsOutliner/Private/TedsOutlinerImpl.cpp
--- a/Engine/Plugins/Experimental/EditorDataStorageFeatures/Source/TedsOutliner/Private/TedsOutlinerImpl.cpp
+++ b/Engine/Plugins/Experimental/EditorDataStorageFeatures/Source/TedsOutliner/Private/TedsOutlinerImpl.cpp
@@ -481,6 +481,7 @@ void FTedsOutlinerImpl::SetSelection(const TArray<DataStorage::RowHandle>& InSel
 
 	// Select each row using the TypedElementSelectionInterface via UTypedElementSelectionSet
 	FTypedElementSelectionOptions SelectionOptions;
+	SelectionOptions.SetAllowHidden(true); // Have to allow selecting things that we made hidden
 	for(DataStorage::RowHandle Row : InSelectedRows)
 	{
 		FTypedElementHandle ElementHandle;
