[TEDS Outliner] Move the CreateWidgetDisplayNameText to the Unsavved widget CellConstructor so the Outliner can use it in the visibility context menu (rather than "Primary Package Object"

#Synced-CL 52003071

diff --git a/Engine/Plugins/Experimental/EditorDataStorageFeatures/Source/TedsOutliner/Private/Widgets/OutlinerUnsavedWidget.cpp b/Engine/Plugins/Experimental/EditorDataStorageFeatures/Source/TedsOutliner/Private/Widgets/OutlinerUnsavedWidget.cpp
--- a/Engine/Plugins/Experimental/EditorDataStorageFeatures/Source/TedsOutliner/Private/Widgets/OutlinerUnsavedWidget.cpp
+++ b/Engine/Plugins/Experimental/EditorDataStorageFeatures/Source/TedsOutliner/Private/Widgets/OutlinerUnsavedWidget.cpp
@@ -53,12 +53,6 @@ TSharedPtr<SWidget> FOutlinerUnsavedHeaderConstructor::CreateWidget(UE::Editor::
 		.ToolTipText(LOCTEXT("SceneOutlinerUnsavedHeaderTooltip", "Unsaved"));
 }
 
-FText FOutlinerUnsavedHeaderConstructor::CreateWidgetDisplayNameText(UE::Editor::DataStorage::ICoreProvider* DataStorage,
-	UE::Editor::DataStorage::RowHandle Row) const
-{
-	return LOCTEXT("OutlinerUnsavedHeaderDisplayName", "Unsaved");
-}
-
 TSharedPtr<SWidget> FOutlinerUnsavedWidgetConstructor::CreateWidget(
 	UE::Editor::DataStorage::ICoreProvider* DataStorage,
 	UE::Editor::DataStorage::IUiProvider* DataStorageUi,
@@ -97,4 +91,10 @@ TSharedPtr<SWidget> FOutlinerUnsavedWidgetConstructor::CreateWidget(
 			];
 }
 
+FText FOutlinerUnsavedWidgetConstructor::CreateWidgetDisplayNameText(UE::Editor::DataStorage::ICoreProvider* DataStorage,
+	UE::Editor::DataStorage::RowHandle Row) const
+{
+	return LOCTEXT("OutlinerUnsavedWidgetDisplayName", "Unsaved");
+}
+
 #undef LOCTEXT_NAMESPACE
