TedsOutliner: Make level editor context menu show on right click.

#Synced-CL 51852197

diff --git a/Engine/Plugins/Experimental/EditorDataStorageFeatures/Source/TedsOutliner/Private/TedsOutlinerMode.cpp b/Engine/Plugins/Experimental/EditorDataStorageFeatures/Source/TedsOutliner/Private/TedsOutlinerMode.cpp
--- a/Engine/Plugins/Experimental/EditorDataStorageFeatures/Source/TedsOutliner/Private/TedsOutlinerMode.cpp
+++ b/Engine/Plugins/Experimental/EditorDataStorageFeatures/Source/TedsOutliner/Private/TedsOutlinerMode.cpp
@@ -280,7 +280,11 @@ TSharedPtr<SWidget> FTedsOutlinerMode::CreateContextMenu()
 	FToolMenuContext MenuContext;
 	MenuContext.AddObject(TedsOutlinerMenuContext);
 
-	return UToolMenus::Get()->GenerateWidget(Private::ContextMenuName, MenuContext);
+	// Allow the ModifyContextMenu callback to modify the context and menu name
+	FName MutableMenuName = Private::ContextMenuName;
+	SceneOutliner->GetSharedData().ModifyContextMenu.ExecuteIfBound(MutableMenuName, MenuContext);
+
+	return UToolMenus::Get()->GenerateWidget(MutableMenuName, MenuContext);
 }
 
 void FTedsOutlinerMode::Tick()
diff --git a/Engine/Plugins/Experimental/EditorDataStorageFeatures/Source/TedsOutliner/Private/TedsOutlinerModule.cpp b/Engine/Plugins/Experimental/EditorDataStorageFeatures/Source/TedsOutliner/Private/TedsOutlinerModule.cpp
--- a/Engine/Plugins/Experimental/EditorDataStorageFeatures/Source/TedsOutliner/Private/TedsOutlinerModule.cpp
+++ b/Engine/Plugins/Experimental/EditorDataStorageFeatures/Source/TedsOutliner/Private/TedsOutlinerModule.cpp
@@ -5,6 +5,7 @@
 #include "Engine/Engine.h"
 #include "Engine/World.h"
 #include "LevelEditor.h"
+#include "LevelEditorContextMenu.h"
 #include "SceneOutlinerHelpers.h"
 #include "SceneOutlinerPublicTypes.h"
 #include "TedsOutlinerMode.h"
@@ -207,6 +208,25 @@ TSharedRef<SWidget> FTedsOutlinerModule::CreateLevelEditorTedsOutliner()
 	InitOptions.bShowTransient = true;
 	InitOptions.OutlinerIdentifier = "TEDSOutliner";
 
+	// Similar to what is done in SLevelEditor::CreateSceneOutliner, bind a callback that adjusts 
+	//  the context menu to have all the level editor context menu sections.
+	TWeakPtr<ILevelEditor> WeakLevelEditor;
+	if (const FLevelEditorModule* const LevelEditorModule = FModuleManager::GetModulePtr<FLevelEditorModule>(TEXT("LevelEditor")))
+	{
+		WeakLevelEditor = LevelEditorModule->GetLevelEditorInstance();
+	}
+	InitOptions.ModifyContextMenu.BindLambda([WeakLevelEditor](FName& OutMenuName, FToolMenuContext& MenuContext)
+	{
+		static const FName ContextMenuName = "LevelEditor.TedsOutliner.ContextMenu";
+		FLevelEditorContextMenu::RegisterAsDerivedMenu(ContextMenuName, OutMenuName);
+		OutMenuName = ContextMenuName;
+
+		if (WeakLevelEditor.IsValid())
+		{
+			FLevelEditorContextMenu::InitMenuContext(MenuContext, WeakLevelEditor, ELevelEditorMenuContext::SceneOutliner);
+		}
+	});
+
 	FTedsOutlinerParams Params(nullptr);
 	Params.QueryDescription = OutlinerQueryDescription;
 	Params.ColumnDescription = GetLevelEditorTedsOutlinerColumnDescription();
diff --git a/Engine/Source/Editor/LevelEditor/Private/LevelEditorContextMenu.cpp b/Engine/Source/Editor/LevelEditor/Private/LevelEditorContextMenu.cpp
--- a/Engine/Source/Editor/LevelEditor/Private/LevelEditorContextMenu.cpp
+++ b/Engine/Source/Editor/LevelEditor/Private/LevelEditorContextMenu.cpp
@@ -853,6 +853,39 @@ void FLevelEditorContextMenu::RegisterEmptySelectionContextMenu()
 	}));
 }
 
+void FLevelEditorContextMenu::RegisterAsDerivedMenu(FName MenuNameToUse, FName BaseMenuName, const FToolMenuInsert& InsertLocation)
+{
+	UToolMenus* ToolMenus = UToolMenus::Get();
+	if (!ensure(ToolMenus))
+	{
+		return;
+	}
+
+	if (!ToolMenus->IsMenuRegistered(MenuNameToUse))
+	{
+		// Inherit from the passed-in context
+		UToolMenu* Menu = ToolMenus->RegisterMenu(MenuNameToUse, BaseMenuName);
+
+		// Add everything that we would show in the level editor context menu
+		FToolMenuSection& Section = Menu->AddDynamicSection("LevelEditorContextMenu", FNewToolMenuDelegate::CreateLambda([](UToolMenu* InMenu)
+		{
+			ULevelEditorContextMenuContext* LevelEditorContext = InMenu->FindContext<ULevelEditorContextMenuContext>();
+			if (!LevelEditorContext)
+			{
+				return;
+			}
+
+			FName LevelContextMenuName = FLevelEditorContextMenu::GetContextMenuName(ELevelEditorMenuContext::SceneOutliner, LevelEditorContext->CurrentSelection);
+			if (LevelContextMenuName != NAME_None)
+			{
+				UToolMenu* OtherMenu = UToolMenus::Get()->GenerateMenu(LevelContextMenuName, InMenu->Context);
+				InMenu->Sections.Append(OtherMenu->Sections);
+			}
+		}));
+		Section.InsertPosition = InsertLocation;
+	}
+}
+
 FName FLevelEditorContextMenu::GetContextMenuName(ELevelEditorMenuContext ContextType, const UTypedElementSelectionSet* InSelectionSet)
 {
 	if (InSelectionSet)
diff --git a/Engine/Source/Editor/LevelEditor/Private/SLevelEditor.cpp b/Engine/Source/Editor/LevelEditor/Private/SLevelEditor.cpp
--- a/Engine/Source/Editor/LevelEditor/Private/SLevelEditor.cpp
+++ b/Engine/Source/Editor/LevelEditor/Private/SLevelEditor.cpp
@@ -838,29 +838,9 @@ TSharedRef<ISceneOutliner> SLevelEditor::CreateSceneOutliner(FName TabIdentifier
 	InitOptions.CommandList = LevelEditorCommands;
 
 	{
-		UToolMenus* ToolMenus = UToolMenus::Get();
 		static const FName MenuName = "LevelEditor.LevelEditorSceneOutliner.ContextMenu";
-		if (!ToolMenus->IsMenuRegistered(MenuName))
-		{
-			UToolMenu* Menu = ToolMenus->RegisterMenu(MenuName, "SceneOutliner.DefaultContextMenuBase");
-			FToolMenuSection& Section = Menu->AddDynamicSection("LevelEditorContextMenu", FNewToolMenuDelegate::CreateLambda([](UToolMenu* InMenu)
-			{
-				ULevelEditorContextMenuContext* LevelEditorContext = InMenu->FindContext<ULevelEditorContextMenuContext>();
-				if (!LevelEditorContext)
-				{
-					return;
-				}
-
-				FName LevelContextMenuName = FLevelEditorContextMenu::GetContextMenuName(ELevelEditorMenuContext::SceneOutliner, LevelEditorContext->CurrentSelection);
-				if (LevelContextMenuName != NAME_None)
-				{
-					// Extend the menu even if no actors selected, as Edit menu should always exist for scene outliner
-					UToolMenu* OtherMenu = UToolMenus::Get()->GenerateMenu(LevelContextMenuName, InMenu->Context);
-					InMenu->Sections.Append(OtherMenu->Sections);
-				}
-			}));
-			Section.InsertPosition = FToolMenuInsert("MainSection", EToolMenuInsertType::Before);
-		}
+		FLevelEditorContextMenu::RegisterAsDerivedMenu(MenuName, "SceneOutliner.DefaultContextMenuBase", 
+			FToolMenuInsert("MainSection", EToolMenuInsertType::Before));
 
 		TWeakPtr<SLevelEditor> WeakLevelEditor = SharedThis(this);
 		InitOptions.ModifyContextMenu.BindLambda([=](FName& OutMenuName, FToolMenuContext& MenuContext)
diff --git a/Engine/Source/Editor/LevelEditor/Public/LevelEditorContextMenu.h b/Engine/Source/Editor/LevelEditor/Public/LevelEditorContextMenu.h
--- a/Engine/Source/Editor/LevelEditor/Public/LevelEditorContextMenu.h
+++ b/Engine/Source/Editor/LevelEditor/Public/LevelEditorContextMenu.h
@@ -5,6 +5,7 @@
 #include "CoreMinimal.h"
 #include "Editor/UnrealEdTypes.h"
 #include "LevelEditorMenuContext.h"
+#include "ToolMenuMisc.h" // FToolMenuInsert
 
 #define UE_API LEVELEDITOR_API
 
@@ -79,6 +80,15 @@ public:
 	static UE_API void RegisterMenuBarEmptyContextMenu();
 	static UE_API void RegisterEmptySelectionContextMenu();
 
+	/**
+	 * Registers a menu with name MenuNameToUse that is derived from BaseMenuName, but has
+	 *  the level editor context menu items added at InsertLocation.
+	 * This is used to add the level editor context menu to a scene outliner context menu.
+	 */
+	UE_EXPERIMENTAL(5.8, "This helper might be removed when a single outliner implementation is used.")
+	static UE_API void RegisterAsDerivedMenu(FName MenuNameToUse, FName BaseMenuName,
+		const FToolMenuInsert& InsertLocation = FToolMenuInsert(NAME_None, EToolMenuInsertType::Last));
+
 	static UE_API void AddPlayFromHereSubMenu(FToolMenuSection& Section);
 
 	/**
