SHierarchyViewer: Add argument to display hierarchy lines (wires)

#Synced-CL 52061643

diff --git a/Engine/Plugins/Experimental/EditorDataStorageFeatures/Source/TedsTableViewer/Private/Widgets/Composite/STedsCompositeHierarchyViewer.cpp b/Engine/Plugins/Experimental/EditorDataStorageFeatures/Source/TedsTableViewer/Private/Widgets/Composite/STedsCompositeHierarchyViewer.cpp
--- a/Engine/Plugins/Experimental/EditorDataStorageFeatures/Source/TedsTableViewer/Private/Widgets/Composite/STedsCompositeHierarchyViewer.cpp
+++ b/Engine/Plugins/Experimental/EditorDataStorageFeatures/Source/TedsTableViewer/Private/Widgets/Composite/STedsCompositeHierarchyViewer.cpp
@@ -207,7 +207,8 @@ namespace UE::Editor::DataStorage
 			.GetCustomTableViewerRowMapping(HierarchyViewerArgs._GetCustomTableViewerRowMapping)
 			.ExpandNewRows(HierarchyViewerArgs._ExpandNewRows)
 			.PrimaryColumn(HierarchyViewerArgs._PrimaryColumn)
-			.AllowInvisibleItemSelection(HierarchyViewerArgs._AllowInvisibleItemSelection);
+			.AllowInvisibleItemSelection(HierarchyViewerArgs._AllowInvisibleItemSelection)
+			.ShowWires(HierarchyViewerArgs._ShowWires);
 		
 		ContentWidget->AddSlot()
 		[
diff --git a/Engine/Plugins/Experimental/EditorDataStorageFeatures/Source/TedsTableViewer/Private/Widgets/STedsHierarchyViewer.cpp b/Engine/Plugins/Experimental/EditorDataStorageFeatures/Source/TedsTableViewer/Private/Widgets/STedsHierarchyViewer.cpp
--- a/Engine/Plugins/Experimental/EditorDataStorageFeatures/Source/TedsTableViewer/Private/Widgets/STedsHierarchyViewer.cpp
+++ b/Engine/Plugins/Experimental/EditorDataStorageFeatures/Source/TedsTableViewer/Private/Widgets/STedsHierarchyViewer.cpp
@@ -50,6 +50,7 @@ namespace UE::Editor::DataStorage
 		GetCustomTableViewerRowMapping = InArgs._GetCustomTableViewerRowMapping;
 		bExpandNewRows = InArgs._ExpandNewRows;
 		EmptyRowsMessage = InArgs._EmptyRowsMessage;
+		bShowWires = InArgs._ShowWires;
 		
 		IUiProvider::FPurposeID CellWidgetPurpose = InArgs._CellWidgetPurpose;
 
@@ -574,6 +575,7 @@ namespace UE::Editor::DataStorage
 				.Padding(FMargin(0))
 				.WidgetRowHandle(TableViewerRowUIRow == InvalidRowHandle ? GetWidgetRowHandle() : TableViewerRowUIRow)
 				.TableViewerWidgetRowHandle(GetWidgetRowHandle())
+				.ShowWires(bShowWires)
 				.Item(InItem);
 	}
 } // namespace UE::Editor::DataStorage
diff --git a/Engine/Plugins/Experimental/EditorDataStorageFeatures/Source/TedsTableViewer/Private/Widgets/STedsTableViewerRow.cpp b/Engine/Plugins/Experimental/EditorDataStorageFeatures/Source/TedsTableViewer/Private/Widgets/STedsTableViewerRow.cpp
--- a/Engine/Plugins/Experimental/EditorDataStorageFeatures/Source/TedsTableViewer/Private/Widgets/STedsTableViewerRow.cpp
+++ b/Engine/Plugins/Experimental/EditorDataStorageFeatures/Source/TedsTableViewer/Private/Widgets/STedsTableViewerRow.cpp
@@ -24,9 +24,11 @@ namespace UE::Editor::DataStorage
 		WidgetRowHandle = InArgs._WidgetRowHandle;
 		TableViewerWidgetRowHandle = InArgs._TableViewerWidgetRowHandle;
 		ItemHeight = InArgs._ItemHeight;
+		bShowWires = InArgs._ShowWires;
 
 		const auto Args = FSuperRowType::FArguments()
 		.Padding(InArgs._Padding)
+		.ShowWires(bShowWires)
 		.OnDragDetected(this, &STedsTableViewerRow::HandleOnDragDetected)
 		.OnDragEnter(this, &STedsTableViewerRow::HandleOnDragEnter)
 		.OnDragLeave(this, &STedsTableViewerRow::HandleOnDragLeave)
@@ -146,7 +148,9 @@ namespace UE::Editor::DataStorage
 					.AutoWidth()
 					.Padding(6.f, 0.f, 0.f, 0.f)
 					[
-						SNew(SExpanderArrow, SharedThis(this)).IndentAmount(12)
+						SNew(SExpanderArrow, SharedThis(this))
+						.IndentAmount(12)
+						.ShouldDrawWires(bShowWires)
 					]
 					+ SHorizontalBox::Slot()
 					.FillWidth(1.0f)
diff --git a/Engine/Plugins/Experimental/EditorDataStorageFeatures/Source/TedsTableViewer/Public/Widgets/STedsHierarchyViewer.h b/Engine/Plugins/Experimental/EditorDataStorageFeatures/Source/TedsTableViewer/Public/Widgets/STedsHierarchyViewer.h
--- a/Engine/Plugins/Experimental/EditorDataStorageFeatures/Source/TedsTableViewer/Public/Widgets/STedsHierarchyViewer.h
+++ b/Engine/Plugins/Experimental/EditorDataStorageFeatures/Source/TedsTableViewer/Public/Widgets/STedsHierarchyViewer.h
@@ -69,6 +69,7 @@ namespace UE::Editor::DataStorage
 			, _PersistUIRowsOnClose(false)
 			, _ExpandNewRows(true)
 			, _AllowInvisibleItemSelection(false)
+			, _ShowWires(false)
 		{
 			
 		}
@@ -121,6 +122,9 @@ namespace UE::Editor::DataStorage
 
 		// Allows collapsed or temporarily hidden items to still be selected
 		SLATE_ARGUMENT(bool, AllowInvisibleItemSelection);
+		
+		// Whether to show hierarchy lines visually
+        SLATE_ARGUMENT(bool, ShowWires);
 
 		SLATE_END_ARGS()
 
@@ -272,5 +276,8 @@ namespace UE::Editor::DataStorage
 
 		// Drag and Drop Handler
 		TUniquePtr<FWidgetDropHandler> DropHandler;
+		
+		// Whether to show hierarchy lines visually
+		bool bShowWires = false;
 	};
 }
\ No newline at end of file
diff --git a/Engine/Plugins/Experimental/EditorDataStorageFeatures/Source/TedsTableViewer/Public/Widgets/STedsTableViewerRow.h b/Engine/Plugins/Experimental/EditorDataStorageFeatures/Source/TedsTableViewer/Public/Widgets/STedsTableViewerRow.h
--- a/Engine/Plugins/Experimental/EditorDataStorageFeatures/Source/TedsTableViewer/Public/Widgets/STedsTableViewerRow.h
+++ b/Engine/Plugins/Experimental/EditorDataStorageFeatures/Source/TedsTableViewer/Public/Widgets/STedsTableViewerRow.h
@@ -15,7 +15,7 @@ namespace UE::Editor::DataStorage
 {
 	class STedsTableViewer;
 	
-	/** Widget that represents a row in the table viewer.  Generates widgets for each column on demand. */
+	/** Widget that represents a row in the table viewer. Generates widgets for each column on demand. */
 	class STedsTableViewerRow
 		: public SMultiColumnTableRow< TableViewerItemPtr >
 	{
@@ -27,27 +27,33 @@ namespace UE::Editor::DataStorage
 			, _TableViewerWidgetRowHandle(InvalidRowHandle)
 			, _Padding(FMargin(0.f, 4.f))
 			, _ItemHeight()
+			, _ShowWires(false)
 		{}
 
 		/** The list item for this row */
 		SLATE_ARGUMENT( TableViewerItemPtr, Item )
 
-			/** Row Handle of the parent of this */
+		/** Widget row representing this table viewer row in the UI */
 		SLATE_ARGUMENT( RowHandle, WidgetRowHandle )
 
+		/** Widget row representing the whole table viewer widget */
 		SLATE_ARGUMENT( RowHandle, TableViewerWidgetRowHandle )
 
+		/** The amount of padding to use */
 		SLATE_ATTRIBUTE( FMargin, Padding )
 			
 		/** The height of the list item */
 		SLATE_ATTRIBUTE( float, ItemHeight )
 
+		/** Whether to draw when showing hierarchies */
+		SLATE_ARGUMENT( bool, ShowWires )
+			
 		SLATE_END_ARGS()
 		
 		/** Construct function for this widget */
 		void Construct( const FArguments& InArgs, TSharedRef<STableViewBase> OwnerTableView, TSharedRef<FTedsTableViewerModel> InTableViewerModel, TWeakPtr<const ITableViewer> InTableViewer);
 
-		/** Overridden from SMultiColumnTableRow.  Generates a widget for this column of the tree row. */
+		/** Overridden from SMultiColumnTableRow. Generates a widget for this column of the tree row. */
 		virtual TSharedRef<SWidget> GenerateWidgetForColumn( const FName& ColumnName ) override;
 
 	protected:
@@ -71,9 +77,10 @@ namespace UE::Editor::DataStorage
 		TSharedPtr<FTedsTableViewerModel> TableViewerModel;
 		TWeakPtr<const ITableViewer> WeakTableViewer;
 		TableViewerItemPtr Item;
+		bool bShowWires = false;
 	};
 
-	/** Widget that represents a row in the hierarchy viewers.  Generates widgets for each column and adds an expander arrow to the very first column. */
+	/** Widget that represents a row in the hierarchy viewers. Generates widgets for each column and adds an expander arrow to the very first column. */
 	class SHierarchyViewerRow : public STedsTableViewerRow
 	{
 	public:
