[TedsTableViewer] Default column header tooltip to display name, configurable via metadata

Column header tooltips in the TEDS Table Viewer now show the human-readable display name
(from CreateWidgetDisplayNameText, falling back to ColumnName) by default, instead of the
internal "Data Storage columns:\n    TypeName..." list.

A new metadata bool "bUseDisplayNameTooltip" controls this behavior (default: true).
Set it to false to revert to showing the backing TEDS column struct names ??? used by the
TEDS Debugger (TedsQueryEditorResultsView and STedsDiscoveryDebugger) where seeing the
internal column types is desirable.

DefaultHeaderText computation was hoisted before widget construction so it is always
available for both the fallback header widget and the tooltip.

#Synced-CL 52005427

diff --git a/Engine/Plugins/Experimental/EditorDataStorageFeatures/Source/TedsDebugger/Private/Widgets/QueryEditor/TedsQueryEditorResultsView.cpp b/Engine/Plugins/Experimental/EditorDataStorageFeatures/Source/TedsDebugger/Private/Widgets/QueryEditor/TedsQueryEditorResultsView.cpp
--- a/Engine/Plugins/Experimental/EditorDataStorageFeatures/Source/TedsDebugger/Private/Widgets/QueryEditor/TedsQueryEditorResultsView.cpp
+++ b/Engine/Plugins/Experimental/EditorDataStorageFeatures/Source/TedsDebugger/Private/Widgets/QueryEditor/TedsQueryEditorResultsView.cpp
@@ -220,11 +220,15 @@ namespace UE::Editor::DataStorage::Debug::QueryEditor
 
 	TSharedRef<SWidget> SResultsView::CreateTableViewer()
 	{
+		TSharedPtr<FMetaData> DebuggerMetaData = MakeShared<FMetaData>();
+		DebuggerMetaData->AddOrSetMutableData("bUseDisplayNameTooltip", false);
+
 		switch (ResultsViewMode)
 		{
 		case ETableViewMode::List:
 			TableViewer = SNew(STedsTableViewer)
 					.QueryStack(RowQueryStack)
+					.GenericMetaData(DebuggerMetaData)
 					.EmptyRowsMessage(LOCTEXT("EmptyRowsMessage", "The provided query has no results."))
 					.OnSelectionChanged(STedsTableViewer::FOnSelectionChanged::CreateLambda(
 						[this](RowHandle SelectedRow, ESelectInfo::Type SelectInfo)
@@ -250,6 +254,7 @@ namespace UE::Editor::DataStorage::Debug::QueryEditor
 			TableViewer =
 				SNew(SHierarchyViewer, HierarchyData)
 					.AllNodeProvider(RowQueryStack)
+					.GenericMetaData(DebuggerMetaData)
 					.EmptyRowsMessage(LOCTEXT("EmptyRowsMessage", "The provided query has no results."))
 					.OnSelectionChanged(STedsTableViewer::FOnSelectionChanged::CreateLambda(
 						[this](RowHandle SelectedRow, ESelectInfo::Type SelectInfo)
diff --git a/Engine/Plugins/Experimental/EditorDataStorageFeatures/Source/TedsTableViewer/Private/TedsTableViewerColumn.cpp b/Engine/Plugins/Experimental/EditorDataStorageFeatures/Source/TedsTableViewer/Private/TedsTableViewerColumn.cpp
--- a/Engine/Plugins/Experimental/EditorDataStorageFeatures/Source/TedsTableViewer/Private/TedsTableViewerColumn.cpp
+++ b/Engine/Plugins/Experimental/EditorDataStorageFeatures/Source/TedsTableViewer/Private/TedsTableViewerColumn.cpp
@@ -97,16 +97,11 @@ namespace UE::Editor::DataStorage
 
 	SHeaderRow::FColumn::FArguments FTedsTableViewerColumn::ConstructHeaderRowColumn()
 	{
-		const TConstArrayView<TWeakObjectPtr<const UScriptStruct>> ColumnTypes = GetMatchedColumns();
-
-		FString TooltipText = TEXT("Data Storage columns:");
-		for (const TWeakObjectPtr<const UScriptStruct>& ColumnType : ColumnTypes)
+		// Compute the display name up front so it can be used for both the header widget and the tooltip
+		FText DefaultHeaderText = CellWidgetConstructor ? CellWidgetConstructor->CreateWidgetDisplayNameText(Storage) : FText::GetEmpty();
+		if (DefaultHeaderText.IsEmpty())
 		{
-			if (ColumnType.IsValid())
-			{
-				TooltipText += TEXT("\n    ");
-				ColumnType->AppendName(TooltipText);
-			}
+			DefaultHeaderText = FText::FromString(ColumnName.ToString());
 		}
 
 		TSharedPtr<SWidget> Widget;
@@ -128,19 +123,13 @@ namespace UE::Editor::DataStorage
 		}
 		if (!Widget.IsValid())
 		{
-			FText DefaultHeaderText = CellWidgetConstructor ? CellWidgetConstructor->CreateWidgetDisplayNameText(Storage) : FText::GetEmpty();
-
-			if (DefaultHeaderText.IsEmpty())
-			{
-				DefaultHeaderText = FText::FromString(ColumnName.ToString());
-			}
-
 			Widget = SNew(STextBlock)
 				.Text(DefaultHeaderText);
 		}
 
 		const bool* bAllowSorting;
 		const bool* bColumnInitialSortIsAscending;
+		const bool* bUseDisplayNameTooltipPtr;
 
 		{
 			UE::Editor::DataStorage::FMetaDataEntryView MetaDataEntryView = GetMetaData().FindGeneric("bAllowSorting");
@@ -150,6 +139,32 @@ namespace UE::Editor::DataStorage
 			UE::Editor::DataStorage::FMetaDataEntryView MetaDataEntryView = GetMetaData().FindGeneric("bColumnInitialSortIsAscending");
 			bColumnInitialSortIsAscending = MetaDataEntryView.TryGetExact<bool>();
 		}
+		{
+			UE::Editor::DataStorage::FMetaDataEntryView MetaDataEntryView = GetMetaData().FindGeneric("bUseDisplayNameTooltip");
+			bUseDisplayNameTooltipPtr = MetaDataEntryView.TryGetExact<bool>();
+		}
+		const bool bUseDisplayNameTooltip = bUseDisplayNameTooltipPtr ? *bUseDisplayNameTooltipPtr : true;
+
+		FText TooltipFText;
+		if (bUseDisplayNameTooltip)
+		{
+			TooltipFText = DefaultHeaderText;
+		}
+		else
+		{
+			const TConstArrayView<TWeakObjectPtr<const UScriptStruct>> ColumnTypes = GetMatchedColumns();
+
+			FString TooltipText = TEXT("Data Storage columns:");
+			for (const TWeakObjectPtr<const UScriptStruct>& ColumnType : ColumnTypes)
+			{
+				if (ColumnType.IsValid())
+				{
+					TooltipText += TEXT("\n    ");
+					ColumnType->AppendName(TooltipText);
+				}
+			}
+			TooltipFText = FText::FromString(MoveTemp(TooltipText));
+		}
 
 		EColumnSortMode::Type SortMode = EColumnSortMode::Ascending;
 		if (bColumnInitialSortIsAscending)
@@ -160,7 +175,7 @@ namespace UE::Editor::DataStorage
 		SHeaderRow::FColumn::FArguments Column = SHeaderRow::Column(ColumnName)
 			.FillWidth(1)
 			.HeaderComboVisibility(EHeaderComboVisibility::OnHover)
-			.DefaultTooltip(FText::FromString(MoveTemp(TooltipText)))
+			.DefaultTooltip(TooltipFText)
 			.DefaultLabel(FText::FromName(ColumnName))
 			.SortMode_Raw(this, &FTedsTableViewerColumn::GetSortMode)
 			.SortPriority_Raw(this, &FTedsTableViewerColumn::GetSortPriority)
