[TEDS Outliner] Have the item label bridge sorting to use TEDS and create a sorting function that takes in type info priorities so Maps/Levels/Folders can come before Actors

#Synced-CL 52098727

diff --git a/Engine/Plugins/ChaosVD/Source/ChaosVD/Private/ChaosVDWorldOutlinerTab.cpp b/Engine/Plugins/ChaosVD/Source/ChaosVD/Private/ChaosVDWorldOutlinerTab.cpp
--- a/Engine/Plugins/ChaosVD/Source/ChaosVD/Private/ChaosVDWorldOutlinerTab.cpp
+++ b/Engine/Plugins/ChaosVD/Source/ChaosVD/Private/ChaosVDWorldOutlinerTab.cpp
@@ -53,8 +53,13 @@ void FChaosVDWorldOutlinerTab::CreateWorldOutlinerWidget()
 	Params.QueryDescription = OutlinerQueryDescription;
 	Params.bShowRowHandleColumn = false;
 	Params.bUseDefaultObservers = false; // CVD uses custom observers to time slice addition and removal currently
+	
+	const TMap<TWeakObjectPtr<const UScriptStruct>, FTedsOutlinerColumnParams> ColumnParams
+	{
+		{FVisibleInEditorColumn::StaticStruct(), FTedsOutlinerColumnParams(FTedsOutlinerColumnParams::EColumnPriorityGroup::Left)}
+	};
 
-	Params.ColumnDescription = FTedsOutlinerColumnDescription({FVisibleInEditorColumn::StaticStruct()});
+	Params.ColumnDescription = FTedsOutlinerColumnDescription({FVisibleInEditorColumn::StaticStruct()}, ColumnParams);
 
 	const FTedsOutlinerHierarchyData::FGetParentRowHandle RowHandleGetter = FTedsOutlinerHierarchyData::FGetParentRowHandle::CreateLambda([](const void* InColumnData)
 	{
diff --git a/Engine/Plugins/Experimental/Dataflow/Source/DataflowEditor/Private/Dataflow/DataflowOutlinerView.cpp b/Engine/Plugins/Experimental/Dataflow/Source/DataflowEditor/Private/Dataflow/DataflowOutlinerView.cpp
--- a/Engine/Plugins/Experimental/Dataflow/Source/DataflowEditor/Private/Dataflow/DataflowOutlinerView.cpp
+++ b/Engine/Plugins/Experimental/Dataflow/Source/DataflowEditor/Private/Dataflow/DataflowOutlinerView.cpp
@@ -38,11 +38,16 @@ TSharedPtr<ISceneOutliner> FDataflowOutlinerView::CreateWidget()
 		Select()
 		.Where(TColumn<FDataflowSceneObjectTag>(GetEditorContent()->GetDataflowOwner().GetFName()) || TColumn<FDataflowSceneStructTag>(GetEditorContent()->GetDataflowOwner().GetFName()) )
 		.Compile();
+		
+	const TMap<TWeakObjectPtr<const UScriptStruct>, FTedsOutlinerColumnParams> ColumnParams
+	{
+		{FVisibleInEditorColumn::StaticStruct(), FTedsOutlinerColumnParams(FTedsOutlinerColumnParams::EColumnPriorityGroup::Left)}
+	};
 
 	const UE::Editor::Outliner::FTedsOutlinerColumnDescription ColumnDescription({
 		FTypedElementClassTypeInfoColumn::StaticStruct(),
 		FVisibleInEditorColumn::StaticStruct()
-	});
+	}, ColumnParams);
 
 	FSceneOutlinerInitializationOptions InitOptions;
 	InitOptions.bShowHeaderRow = true;
@@ -86,9 +91,6 @@ TSharedPtr<ISceneOutliner> FDataflowOutlinerView::CreateWidget()
 		Params.SceneOutliner = Outliner;
 		return new FDataflowOutlinerMode(Params, ConstructionScene, SimulationScene);
 	});
-	
-	InitOptions.ColumnMap.Add(FSceneOutlinerBuiltInColumnTypes::Gutter(), FSceneOutlinerColumnInfo(ESceneOutlinerColumnVisibility::Visible, 40));
-	InitOptions.ColumnMap.Add(FSceneOutlinerBuiltInColumnTypes::Label(), FSceneOutlinerColumnInfo(ESceneOutlinerColumnVisibility::Visible, 50));
 
 	OutlinerWidget = SNew(SSceneOutliner, InitOptions);
 
diff --git a/Engine/Plugins/Experimental/EditorDataStorageFeatures/Source/TedsOutliner/Private/Compatibility/TedsSceneOutlinerItemLabelColumn.cpp b/Engine/Plugins/Experimental/EditorDataStorageFeatures/Source/TedsOutliner/Private/Compatibility/TedsSceneOutlinerItemLabelColumn.cpp
new file mode 100644
--- /dev/null
+++ b/Engine/Plugins/Experimental/EditorDataStorageFeatures/Source/TedsOutliner/Private/Compatibility/TedsSceneOutlinerItemLabelColumn.cpp
@@ -0,0 +1,19 @@
+// Copyright Epic Games, Inc. All Rights Reserved.
+
+#include "Compatibility/TedsSceneOutlinerItemLabelColumn.h"
+
+#include "ISceneOutlinerMode.h"
+#include "SceneOutlinerHelpers.h"
+
+const TSharedRef<SWidget> FTedsSceneOutlinerItemLabelColumn::ConstructRowWidget(FSceneOutlinerTreeItemRef TreeItem, const STableRow<FSceneOutlinerTreeItemPtr>& Row)
+{
+	ISceneOutliner* Outliner = OwningOutliner.Pin().Get();
+	check(Outliner);
+	return TreeItem->GenerateLabelWidget(*Outliner, Row);
+}
+
+void FTedsSceneOutlinerItemLabelColumn::PopulateSearchStrings( const ISceneOutlinerTreeItem& Item, TArray< FString >& OutSearchStrings ) const
+{
+	OutSearchStrings.Add(Item.GetDisplayString());
+	SceneOutliner::FSceneOutlinerHelpers::PopulateExtraSearchStrings(Item, OutSearchStrings);
+}
\ No newline at end of file
diff --git a/Engine/Plugins/Experimental/EditorDataStorageFeatures/Source/TedsOutliner/Private/Compatibility/TedsSceneOutlinerItemLabelColumn.h b/Engine/Plugins/Experimental/EditorDataStorageFeatures/Source/TedsOutliner/Private/Compatibility/TedsSceneOutlinerItemLabelColumn.h
new file mode 100644
--- /dev/null
+++ b/Engine/Plugins/Experimental/EditorDataStorageFeatures/Source/TedsOutliner/Private/Compatibility/TedsSceneOutlinerItemLabelColumn.h
@@ -0,0 +1,25 @@
+// Copyright Epic Games, Inc. All Rights Reserved.
+
+#pragma once
+
+#include "SceneOutlinerFwd.h"
+#include "Widgets/Views/SHeaderRow.h"
+#include "Widgets/Views/STableRow.h"
+#include "TedsOutlinerColumn.h"
+
+/**
+ * This is an override column for the TEDS Outliner/Table Viewer to display item labels for items. The label column requires that we populate 
+ * additional search strings for components and construct the Label Column widget via the Tree Item so it can get the appropriate metadata
+ * (still constructs a TEDS UI Widget)
+ */
+class FTedsSceneOutlinerItemLabelColumn : public UE::Editor::Outliner::FTedsOutlinerUiColumn
+{
+public:
+	FTedsSceneOutlinerItemLabelColumn(UE::Editor::Outliner::FTedsOutlinerUiColumnInitParams& InitParams) : FTedsOutlinerUiColumn(InitParams) {}
+	
+	virtual ~FTedsSceneOutlinerItemLabelColumn() = default;
+	
+	virtual const TSharedRef<SWidget> ConstructRowWidget( FSceneOutlinerTreeItemRef TreeItem, const STableRow<FSceneOutlinerTreeItemPtr>& Row ) override;
+
+	virtual void PopulateSearchStrings( const ISceneOutlinerTreeItem& Item, TArray< FString >& OutSearchStrings ) const override;
+};
diff --git a/Engine/Plugins/Experimental/EditorDataStorageFeatures/Source/TedsOutliner/Private/TedsOutlinerImpl.cpp b/Engine/Plugins/Experimental/EditorDataStorageFeatures/Source/TedsOutliner/Private/TedsOutlinerImpl.cpp
--- a/Engine/Plugins/Experimental/EditorDataStorageFeatures/Source/TedsOutliner/Private/TedsOutlinerImpl.cpp
+++ b/Engine/Plugins/Experimental/EditorDataStorageFeatures/Source/TedsOutliner/Private/TedsOutlinerImpl.cpp
@@ -35,6 +35,7 @@
 #include "Columns/SceneOutlinerColumns.h"
 #include "Columns/TedsOutlinerColumns.h"
 #include "Compatibility/SceneOutlinerTedsBridge.h"
+#include "Compatibility/TedsSceneOutlinerItemLabelColumn.h"
 #include "Elements/Columns/TypedElementCompatibilityColumns.h"
 #include "Elements/Columns/TypedElementHandleColumn.h"
 #include "Elements/Columns/TypedElementOverrideColumns.h"
@@ -400,6 +401,8 @@ void FTedsOutlinerImpl::Init()
 	TickerHandle = FTSTicker::GetCoreTicker().AddTicker(FTickerDelegate::CreateLambda([this](float)
 	{
 		RegenerateColumns(CreationParams.ColumnDescription);
+		// Since Sort runs before the Teds Columns are added, we want to resort since we override the Item Label column
+		SceneOutliner->RequestSort();
 		return false;
 	}));
 }
@@ -685,6 +688,71 @@ void FTedsOutlinerImpl::RegenerateColumns(const Outliner::FTedsOutlinerColumnDes
 	const RowHandle CellPurposeRow = StorageUi->FindPurpose(CellWidgetPurpose);
 	const RowHandle HeaderPurposeRow = StorageUi->FindPurpose(HeaderWidgetPurpose);
 
+	// Label Column Construction
+	if (!bHybridMode)
+	{
+		auto LabelColumnConstructor = [this, ColumnDescription, TedsOutlinerColumn, MetaDataView, HeaderPurposeRow, Dealiaser](
+			IUiProvider::FWidgetConstructorPtr CellConstructor, TConstArrayView<TWeakObjectPtr<const UScriptStruct>> ColumnTypes)
+			{
+				const FName LabelColumnName = FSceneOutlinerBuiltInColumnTypes::Label();
+				// Explicitly set the priority Index to greater than high-priority icon columns since the Label Column
+				// is manually added to the ColumnMap right now and cannot use dynamic priority grouping
+				constexpr uint8 LabelColumnPriority = 10;
+	
+				// Grab a reference to the label column and remove it from the actual outliner
+				const TMap<FName, TSharedPtr<ISceneOutlinerColumn>>& Columns = SceneOutliner->GetColumns();
+				const TSharedPtr<ISceneOutlinerColumn> FallbackColumnPtr = Columns.FindRef(LabelColumnName);
+				SceneOutliner->RemoveColumn(LabelColumnName);
+	
+				const FText DisplayName = CellConstructor.Get()->CreateWidgetDisplayNameText(Storage);
+				AddedColumns.Add(LabelColumnName);
+				SceneOutliner->AddColumn(LabelColumnName,
+					FSceneOutlinerColumnInfo(
+						ESceneOutlinerColumnVisibility::Visible, 
+						LabelColumnPriority,
+						FCreateSceneOutlinerColumn::CreateLambda(
+							[this, ColumnDescription, TedsOutlinerColumn, MetaDataView, LabelColumnName, ColumnTypes, CellConstructor, FallbackColumnPtr, HeaderPurposeRow, Dealiaser](ISceneOutliner&)
+							{
+								IUiProvider::FWidgetConstructorPtr HeaderConstructor = 
+									TableViewerUtils::CreateHeaderWidgetConstructor(*StorageUi, MetaDataView, ColumnTypes, HeaderPurposeRow);
+	
+								FTedsTableViewerColumn::FSortHandler SortDelegates
+								{
+									.OnSortHandler = FTedsTableViewerColumn::FOnSort::CreateSP(this, &FTedsOutlinerImpl::OnSort),
+									.IsSortingHandler = FTedsTableViewerColumn::FIsSorting::CreateSP(this, &FTedsOutlinerImpl::IsSorting)
+								};
+	
+								FTedsOutlinerUiColumnInitParams Params(ColumnDescription.GetGenericMetaData(), *Storage, *StorageUi, *StorageCompatibility);
+								Params.NameId = LabelColumnName;
+								Params.ColumnTypes = TArray<TWeakObjectPtr<const UScriptStruct>>(ColumnTypes.GetData(), ColumnTypes.Num());
+								Params.HeaderWidgetConstructor = MoveTemp(HeaderConstructor);
+								Params.CellWidgetConstructor = CellConstructor;
+								Params.FallbackColumn = FallbackColumnPtr;
+								Params.OwningOutliner = TedsOutlinerColumn->Outliner;
+								Params.TedsOutlinerImpl = AsWeak();
+								Params.SortDelegates = MoveTemp(SortDelegates);
+								Params.Dealiaser = Dealiaser;
+								Params.bHybridMode = false;
+								
+								return MakeShared<FTedsSceneOutlinerItemLabelColumn>(Params);
+	
+							}),
+						false,
+						TOptional<float>(),
+						DisplayName
+					)
+				);
+				return true;
+			};
+
+		// The Outliner Label matches to the label column && the type info column, so we pass in both
+		TArray<TWeakObjectPtr<const UScriptStruct>> MatchedLabelColumns = {
+			FTypedElementLabelColumn::StaticStruct(),
+			FTypedElementClassTypeInfoColumn::StaticStruct() };
+		StorageUi->CreateWidgetConstructors(CellPurposeRow, MatchApproach::LongestMatch, MatchedLabelColumns, 
+			MetaDataView, LabelColumnConstructor);
+	}
+
 	int32 IndexOffset = 0;
 	auto ColumnConstructor = [this, ColumnDescription, TedsOutlinerColumn, MetaDataView, &IndexOffset, &OrderedColumnMap, HeaderPurposeRow, Dealiaser](
 		IUiProvider::FWidgetConstructorPtr CellConstructor, TConstArrayView<TWeakObjectPtr<const UScriptStruct>> ColumnTypes)
diff --git a/Engine/Plugins/Experimental/EditorDataStorageFeatures/Source/TedsOutliner/Private/TedsOutlinerModule.cpp b/Engine/Plugins/Experimental/EditorDataStorageFeatures/Source/TedsOutliner/Private/TedsOutlinerModule.cpp
--- a/Engine/Plugins/Experimental/EditorDataStorageFeatures/Source/TedsOutliner/Private/TedsOutlinerModule.cpp
+++ b/Engine/Plugins/Experimental/EditorDataStorageFeatures/Source/TedsOutliner/Private/TedsOutlinerModule.cpp
@@ -101,17 +101,14 @@ TSharedRef<ISceneOutliner> FTedsOutlinerModule::CreateTedsOutliner(const FSceneO
 		
 		return new FTedsOutlinerMode(InitTedsOptions);
 	});
-
-	// Explicitly set the priority Index to greater than high-priority icon columns since the Label Column
-	// is manually added to the ColumnMap right now and cannot use dynamic priority grouping
-	constexpr int LabelColumnPriorityIndex = 10;
-	// Explicitly set the priority Index to one less than the Label Priority since we cannot use priority groups (not
-	// passed in as a TEDS Column), and we want it to appear after high-priority icon columns but before the Label Column.
-	constexpr int RowHandleColumnPriorityIndex = LabelColumnPriorityIndex - 1;
 	
 	// Add the custom column that displays row handles
 	if (InInitTedsOptions.bShowRowHandleColumn)
 	{
+		// Explicitly set the priority Index to one less than the Label Priority (10) since we cannot use priority groups (not
+		// passed in as a TEDS Column), and we want it to appear after high-priority icon columns but before the Label Column.
+		constexpr int RowHandleColumnPriorityIndex = 9;
+		
 		InitOptions.ColumnMap.Add(FSceneOutlinerRowHandleColumn::GetID(),
 			FSceneOutlinerColumnInfo(ESceneOutlinerColumnVisibility::Invisible, RowHandleColumnPriorityIndex,
 				FCreateSceneOutlinerColumn::CreateLambda([](ISceneOutliner& InSceneOutliner)
@@ -120,13 +117,6 @@ TSharedRef<ISceneOutliner> FTedsOutlinerModule::CreateTedsOutliner(const FSceneO
 				})));
 	}
 	
-	InitOptions.ColumnMap.Add(FSceneOutlinerBuiltInColumnTypes::Label(),
-		FSceneOutlinerColumnInfo(
-			ESceneOutlinerColumnVisibility::Visible,
-			LabelColumnPriorityIndex,
-			FCreateSceneOutlinerColumn(),
-			false));
-	
 	TSharedRef<ISceneOutliner> TedsOutlinerShared = SNew(SSceneOutliner, InitOptions);
 	
 	return TedsOutlinerShared;
diff --git a/Engine/Plugins/Experimental/EditorDataStorageFeatures/Source/TedsOutliner/Private/Widgets/OutlinerLabelWidget.cpp b/Engine/Plugins/Experimental/EditorDataStorageFeatures/Source/TedsOutliner/Private/Widgets/OutlinerLabelWidget.cpp
--- a/Engine/Plugins/Experimental/EditorDataStorageFeatures/Source/TedsOutliner/Private/Widgets/OutlinerLabelWidget.cpp
+++ b/Engine/Plugins/Experimental/EditorDataStorageFeatures/Source/TedsOutliner/Private/Widgets/OutlinerLabelWidget.cpp
@@ -2,21 +2,87 @@
 
 #include "Widgets/OutlinerLabelWidget.h"
 
+#include "Columns/SlateHeaderColumns.h"
+#include "Elements/Columns/TypedElementCompatibilityColumns.h"
 #include "Elements/Columns/TypedElementTypeInfoColumns.h"
 #include "Elements/Columns/TypedElementLabelColumns.h"
 #include "Layout/Margin.h"
 #include "Widgets/SNullWidget.h"
 #include "Widgets/SBoxPanel.h"
 #include "Widgets/Layout/SSpacer.h"
+#include "Widgets/Text/STextBlock.h"
 
 #include UE_INLINE_GENERATED_CPP_BY_NAME(OutlinerLabelWidget)
 
 #define LOCTEXT_NAMESPACE "OutlinerLabelWidget"
 
+namespace UE::Editor::DataStorage
+{
+	FColumnSorterInterface::ESortType FOutlinerLabelWidgetSorter::GetSortType() const
+	{
+		return ESortType::HybridSort;
+	}
+	
+	FText FOutlinerLabelWidgetSorter::GetShortName() const
+	{
+		return LOCTEXT("OutlinerLabelWidgetSorter", "Outliner Label Sorter");
+	}
+
+	int32 FOutlinerLabelWidgetSorter::Compare(const ICoreProvider& Storage, RowHandle Left, RowHandle Right) const
+	{
+		const int32 LeftTypePriority = GetTypeInfoPriorityForRow(Storage, Left);
+		const int32 RightTypePriority = GetTypeInfoPriorityForRow(Storage, Right);
+
+		if (LeftTypePriority != RightTypePriority)
+		{
+			return LeftTypePriority - RightTypePriority;
+		}
+
+		const FTypedElementLabelColumn* LeftLabelColumn = Storage.GetColumn<FTypedElementLabelColumn>(Left);
+		const FTypedElementLabelColumn* RightLabelColumn = Storage.GetColumn<FTypedElementLabelColumn>(Right);
+		if (LeftLabelColumn && RightLabelColumn)
+		{
+			return LeftLabelColumn->Label.Compare(RightLabelColumn->Label);
+		}
+		
+		return 0;
+	}
+
+	FPrefixInfo FOutlinerLabelWidgetSorter::CalculatePrefix(const ICoreProvider& Storage, RowHandle Row, uint32 ByteIndex) const
+	{
+		if (const FTypedElementLabelColumn* LabelColumn = Storage.GetColumn<FTypedElementLabelColumn>(Row))
+		{
+			return CreateSortPrefix(ByteIndex, GetTypeInfoPriorityForRow(Storage, Row), TSortStringView(FSortCaseSensitive{}, LabelColumn->Label));
+		}
+		return CreateSortPrefix(ByteIndex, -1);
+	}
+
+	int32 FOutlinerLabelWidgetSorter::GetTypeInfoPriorityForRow(const ICoreProvider& Storage, RowHandle Row) const
+	{
+		if (Storage.HasColumns<FEditorDataStorageWorldTag>(Row))
+		{
+			return 0;
+		}
+		
+		if (Storage.HasColumns<FEditorDataStorageLevelTag>(Row))
+		{
+			return 1;
+		}
+		
+		// TODO: When Folders are added in TEDS, they should return 2
+
+		return 3;
+	}
+}
+
 void UOutlinerLabelWidgetFactory::RegisterWidgetConstructors(UE::Editor::DataStorage::ICoreProvider& DataStorage,
 	UE::Editor::DataStorage::IUiProvider& DataStorageUi) const
 {
 	using namespace UE::Editor::DataStorage::Queries;
+	DataStorageUi.RegisterWidgetFactory<FOutlinerLabelHeaderConstructor>(
+		DataStorageUi.FindPurpose(UE::Editor::DataStorage::IUiProvider::FPurposeInfo("SceneOutliner", "Header", NAME_None).GeneratePurposeID()),
+		TColumn<FTypedElementLabelColumn>() && (TColumn<FTypedElementClassTypeInfoColumn>() || TColumn<FTypedElementScriptStructTypeInfoColumn>()));
+    	
 	DataStorageUi.RegisterWidgetFactory<FOutlinerLabelWidgetConstructor>(
 		DataStorageUi.FindPurpose(UE::Editor::DataStorage::IUiProvider::FPurposeInfo("SceneOutliner", "RowLabel", NAME_None).GeneratePurposeID()),
 		TColumn<FTypedElementLabelColumn>() && (TColumn<FTypedElementClassTypeInfoColumn>() || TColumn<FTypedElementScriptStructTypeInfoColumn>()));
@@ -50,6 +116,37 @@ static void GetAllColumns(TArray<TWeakObjectPtr<const UScriptStruct>>& OutColumn
 		});
 }
 
+FOutlinerLabelHeaderConstructor::FOutlinerLabelHeaderConstructor()
+	: FSimpleWidgetConstructor(StaticStruct())
+{
+}
+
+TSharedPtr<SWidget> FOutlinerLabelHeaderConstructor::CreateWidget(UE::Editor::DataStorage::ICoreProvider* DataStorage,
+	UE::Editor::DataStorage::IUiProvider* DataStorageUi, UE::Editor::DataStorage::RowHandle TargetRow, UE::Editor::DataStorage::RowHandle WidgetRow,
+	const UE::Editor::DataStorage::FMetaDataView& Arguments)
+{
+	DataStorage->AddColumn(WidgetRow, FHeaderWidgetSizeColumn
+	{
+		.ColumnSizeMode = EColumnSizeMode::Fill,
+		.Width = 5.0f
+	});
+	
+	return SNew(STextBlock)
+		.Text(LOCTEXT("OutlinerLabelHeaderText", "Item Label"));
+}
+
+TArray<TSharedPtr<const UE::Editor::DataStorage::FColumnSorterInterface>> FOutlinerLabelHeaderConstructor::ConstructColumnSorters(
+	UE::Editor::DataStorage::ICoreProvider* DataStorage,
+	UE::Editor::DataStorage::IUiProvider* DataStorageUi,
+	const UE::Editor::DataStorage::FMetaDataView& Arguments)
+{
+	using namespace UE::Editor::DataStorage;
+	return TArray<TSharedPtr<const FColumnSorterInterface>>(
+		{
+			MakeShared<FOutlinerLabelWidgetSorter>()
+		});
+}
+
 FOutlinerLabelWidgetConstructor::FOutlinerLabelWidgetConstructor()
 	: FSimpleWidgetConstructor(StaticStruct())
 {
diff --git a/Engine/Plugins/Experimental/EditorDataStorageFeatures/Source/TedsOutliner/Private/Widgets/OutlinerLabelWidget.h b/Engine/Plugins/Experimental/EditorDataStorageFeatures/Source/TedsOutliner/Private/Widgets/OutlinerLabelWidget.h
--- a/Engine/Plugins/Experimental/EditorDataStorageFeatures/Source/TedsOutliner/Private/Widgets/OutlinerLabelWidget.h
+++ b/Engine/Plugins/Experimental/EditorDataStorageFeatures/Source/TedsOutliner/Private/Widgets/OutlinerLabelWidget.h
@@ -8,6 +8,26 @@
 
 #include "OutlinerLabelWidget.generated.h"
 
+struct FTypedElementLabelColumn;
+
+namespace UE::Editor::DataStorage
+{
+	class FOutlinerLabelWidgetSorter final : public FColumnSorterInterface
+	{
+	public:
+		virtual ESortType GetSortType() const override;
+		
+		virtual FText GetShortName() const override;
+
+		virtual int32 Compare(const ICoreProvider& Storage, RowHandle Left, RowHandle Right) const override;
+		
+		virtual FPrefixInfo CalculatePrefix(const ICoreProvider& Storage, RowHandle Row, uint32 ByteIndex) const override;
+
+	private:
+		int32 GetTypeInfoPriorityForRow(const ICoreProvider& Storage, RowHandle Row) const;
+	};
+	}
+
 UCLASS()
 class UOutlinerLabelWidgetFactory : public UEditorDataStorageFactory
 {
@@ -22,6 +42,29 @@ public:
 	TEDSOUTLINER_API void RegisterWidgetPurposes(UE::Editor::DataStorage::IUiProvider& DataStorageUi) const override;
 };
 
+// Label header widget for the Scene Outliner
+USTRUCT()
+struct FOutlinerLabelHeaderConstructor : public FSimpleWidgetConstructor
+{
+	GENERATED_BODY()
+
+public:
+	TEDSOUTLINER_API FOutlinerLabelHeaderConstructor();
+	~FOutlinerLabelHeaderConstructor() override = default;
+
+	TEDSOUTLINER_API virtual TSharedPtr<SWidget> CreateWidget(
+		UE::Editor::DataStorage::ICoreProvider* DataStorage,
+		UE::Editor::DataStorage::IUiProvider* DataStorageUi,
+		UE::Editor::DataStorage::RowHandle TargetRow,
+		UE::Editor::DataStorage::RowHandle WidgetRow, 
+		const UE::Editor::DataStorage::FMetaDataView& Arguments) override;
+	
+	TEDSOUTLINER_API virtual TArray<TSharedPtr<const UE::Editor::DataStorage::FColumnSorterInterface>> ConstructColumnSorters(
+		UE::Editor::DataStorage::ICoreProvider* DataStorage,
+		UE::Editor::DataStorage::IUiProvider* DataStorageUi,
+		const UE::Editor::DataStorage::FMetaDataView& Arguments) override;
+};
+
 // Label widget for the Scene Outliner that shows an icon (with optional override information) + a text label
 USTRUCT()
 struct FOutlinerLabelWidgetConstructor : public FSimpleWidgetConstructor
diff --git a/Engine/Plugins/Experimental/EditorDataStorageFeatures/Source/TedsOutliner/Public/TedsOutlinerImpl.h b/Engine/Plugins/Experimental/EditorDataStorageFeatures/Source/TedsOutliner/Public/TedsOutlinerImpl.h
--- a/Engine/Plugins/Experimental/EditorDataStorageFeatures/Source/TedsOutliner/Public/TedsOutlinerImpl.h
+++ b/Engine/Plugins/Experimental/EditorDataStorageFeatures/Source/TedsOutliner/Public/TedsOutlinerImpl.h
@@ -52,7 +52,7 @@ struct FTedsOutlinerColumnParams
 {
 	// Defines the priority groups, where Left priority is ordered first and Right priority is ordered last
 	// ||  Left  |  Middle  |  Right  ||
-	enum class EColumnPriorityGroup
+	enum class TEDSOUTLINER_API EColumnPriorityGroup
 	{
 		Middle = 0, // Initial Priority value of 80
 		Left = 1 << 0, // Initial Priority value of 0
@@ -62,14 +62,14 @@ struct FTedsOutlinerColumnParams
 	// Defines possible relations between columns so they can specify the desired ordering
 	struct FColumnPriorityRelation
 	{
-		enum class EColumnPriorityRelation
+		enum class TEDSOUTLINER_API EColumnPriorityRelation
 		{
 			None = 0,
 			Before = 1 << 0, // This column comes before the specified RelatedColumn
 			After = 1 << 1 // This column comes after the specified RelatedColumn
 		};
 		
-		explicit FColumnPriorityRelation(const EColumnPriorityRelation InPriorityRelation = EColumnPriorityRelation::None,
+		explicit TEDSOUTLINER_API FColumnPriorityRelation(const EColumnPriorityRelation InPriorityRelation = EColumnPriorityRelation::None,
 			const TWeakObjectPtr<const UScriptStruct>& InRelatedColumn = nullptr);
 
 		// When determining priority, all columns with Before relations are sorted before columns with After Columns
