TedsOutliner: Fix "Add To Favorites" not showing up at the top in the entity context menu due to it being default insertion order while the entity menu has a section with first insertion order.

#Synced-CL 52191381

diff --git a/Engine/Plugins/Experimental/EditorDataStorageFeatures/Source/TedsOutliner/Private/TedsOutlinerItem.cpp b/Engine/Plugins/Experimental/EditorDataStorageFeatures/Source/TedsOutliner/Private/TedsOutlinerItem.cpp
--- a/Engine/Plugins/Experimental/EditorDataStorageFeatures/Source/TedsOutliner/Private/TedsOutlinerItem.cpp
+++ b/Engine/Plugins/Experimental/EditorDataStorageFeatures/Source/TedsOutliner/Private/TedsOutlinerItem.cpp
@@ -130,6 +130,7 @@ void FTedsOutlinerTreeItem::GenerateContextMenu(UToolMenu* Menu, SSceneOutliner&
 				)
 			);
 		}
+		Section.InsertPosition = FToolMenuInsert(NAME_None, EToolMenuInsertType::First);
 	}
 
 	if (FTedsOutlinerContextMenuColumn* Column = DataStorageInterface->GetColumn<FTedsOutlinerContextMenuColumn>(RowHandle))
