TEDS: Add an undo change to use for visibility column changes.

#Synced-CL 52191381

diff --git a/Engine/Plugins/Experimental/EditorDataStorageFeatures/Source/TedsOutliner/Private/Widgets/OutlinerVisibilityWidget.cpp b/Engine/Plugins/Experimental/EditorDataStorageFeatures/Source/TedsOutliner/Private/Widgets/OutlinerVisibilityWidget.cpp
--- a/Engine/Plugins/Experimental/EditorDataStorageFeatures/Source/TedsOutliner/Private/Widgets/OutlinerVisibilityWidget.cpp
+++ b/Engine/Plugins/Experimental/EditorDataStorageFeatures/Source/TedsOutliner/Private/Widgets/OutlinerVisibilityWidget.cpp
@@ -6,6 +6,7 @@
 #include "Columns/SlateHeaderColumns.h"
 #include "Columns/TedsOutlinerColumns.h"
 #include "Compatibility/SceneOutlinerTedsBridge.h"
+#include "CoreGlobals.h" // GUndo
 #include "DataStorage/Features.h"
 #include "Elements/Columns/TypedElementCompatibilityColumns.h"
 #include "Elements/Columns/TypedElementHiearchyColumns.h"
@@ -15,6 +16,7 @@
 #include "Elements/Framework/TypedElementQueryBuilder.h"
 #include "Elements/Interfaces/TypedElementDataStorageCompatibilityInterface.h"
 #include "Elements/Interfaces/TypedElementDataStorageInterface.h"
+#include "Misc/ITransaction.h"
 #include "Widgets/SBoxPanel.h"
 #include "ISceneOutliner.h"
 #include "ActorTreeItem.h"
@@ -178,6 +180,8 @@ void STedsVisibilityWidget::OnDragEnter(const FGeometry& MyGeometry, const FDrag
 
 FReply STedsVisibilityWidget::HandleClick()
 {
+	using namespace UE::Editor::DataStorage;
+
 	if (!IsEnabled())
 	{
 		return FReply::Unhandled();
@@ -191,16 +195,44 @@ FReply STedsVisibilityWidget::HandleClick()
 	TArray<RowHandle> SelectedRows;
 	GetSelectedRows(SelectedRows);
 
+	FRowHandleArray ChangedRows;
 	if (SelectedRows.Contains(TargetRow))
 	{
 		for (RowHandle Row : SelectedRows)
 		{
-			SetIsVisible(Row, bVisible);
+			SetIsVisible(Row, bVisible, &ChangedRows);
 		}
 	}
 	else
 	{
-		SetIsVisible(TargetRow, bVisible);
+		SetIsVisible(TargetRow, bVisible, &ChangedRows);
+	}
+
+	// Store undo change for the changed rows
+	if (GUndo && !ChangedRows.IsEmpty())
+	{
+		// We need some UObject to use for the transaction. We'll grab the world from one of the rows
+		UWorld* WorldToUse = nullptr;
+		ICoreProvider* DataStorage = GetDataStorage();
+		if (ensure(DataStorage))
+		{
+			for (RowHandle Row : ChangedRows.GetRows())
+			{
+				if (const FTypedElementWorldColumn* WorldColumn = DataStorage->GetColumn<FTypedElementWorldColumn>(Row))
+				{
+					if (UWorld* World = WorldColumn->World.Get())
+					{
+						WorldToUse = World;
+						break;
+					}
+				}
+			}
+		}
+		
+		if (ensure(WorldToUse))
+		{
+			GUndo->StoreUndo(WorldToUse, MakeUnique<FVisibilityColumnChange>(MoveTemp(ChangedRows), bVisible));
+		}
 	}
 
 	return FReply::Handled().DetectDrag(SharedThis(this), EKeys::LeftMouseButton);
@@ -297,7 +329,7 @@ bool STedsVisibilityWidget::IsSelected() const
 	return false;
 }
 
-void STedsVisibilityWidget::SetIsVisible(RowHandle InRow, const bool bVisible)
+void STedsVisibilityWidget::SetIsVisible(RowHandle InRow, const bool bVisible, UE::Editor::DataStorage::FRowHandleArray* ChangedRowsToAppendTo)
 {
 	using namespace UE::Editor::DataStorage;
 
@@ -310,31 +342,37 @@ void STedsVisibilityWidget::SetIsVisible(RowHandle InRow, const bool bVisible)
 		{
 			if (TSharedPtr<IHierarchyViewerDataInterface> DataInterface = DataColumn->HierarchyData.Pin())
 			{
-				SetVisibility_Recursive(*DataStorage, InRow, DataInterface, bVisible);
+				SetVisibility_Recursive(*DataStorage, InRow, DataInterface, bVisible, ChangedRowsToAppendTo);
 			}
 		}
 		// The Teds Outliner doesn't use the new hierarchy data yet so we have a fallback to the tree item API
 		else if (TSharedPtr<ISceneOutlinerTreeItem> RowTreeItem = GetTreeItem(InRow))
 		{
-			SetVisibility_Recursive(*DataStorage, RowTreeItem, bVisible);
+			SetVisibility_Recursive(*DataStorage, RowTreeItem, bVisible, ChangedRowsToAppendTo);
 		}
 		
 	}
 }
 
-void STedsVisibilityWidget::CommitVisibility(UE::Editor::DataStorage::ICoreProvider& DataStorage, RowHandle Row, bool bVisible)
+void STedsVisibilityWidget::CommitVisibility(UE::Editor::DataStorage::ICoreProvider& DataStorage, RowHandle Row, bool bVisible,
+	UE::Editor::DataStorage::FRowHandleArray* ChangedRowsOut)
 {
 	// Set new visibility on the TargetRow
 	FVisibleInEditorColumn* VisibileInEditorColumn = DataStorage.GetColumn<FVisibleInEditorColumn>(Row);
-	if (VisibileInEditorColumn)
+	if (VisibileInEditorColumn && VisibileInEditorColumn->bIsVisibleInEditor != bVisible)
 	{
 		VisibileInEditorColumn->bIsVisibleInEditor = bVisible;
+		DataStorage.AddColumn<FVisibilityChangedTag>(Row);
+		if (ChangedRowsOut)
+		{
+			ChangedRowsOut->Add(Row);
+		}
 	}
-	DataStorage.AddColumn<FVisibilityChangedTag>(Row);
 }
 
 void STedsVisibilityWidget::SetVisibility_Recursive(UE::Editor::DataStorage::ICoreProvider& DataStorage, RowHandle Row,
-	TSharedPtr<UE::Editor::DataStorage::IHierarchyViewerDataInterface> HierarchyData, bool bVisible)
+	TSharedPtr<UE::Editor::DataStorage::IHierarchyViewerDataInterface> HierarchyData, bool bVisible, 
+	UE::Editor::DataStorage::FRowHandleArray* ChangedRowsOut)
 {
 	using namespace UE::Editor::DataStorage;
 
@@ -346,11 +384,12 @@ void STedsVisibilityWidget::SetVisibility_Recursive(UE::Editor::DataStorage::ICo
 
 	for (RowHandle ChildRow : Children.GetRows())
 	{
-		CommitVisibility(DataStorage, ChildRow, bVisible);
+		CommitVisibility(DataStorage, ChildRow, bVisible, ChangedRowsOut);
 	}
 }
 
-void STedsVisibilityWidget::SetVisibility_Recursive(UE::Editor::DataStorage::ICoreProvider& DataStorage, FSceneOutlinerTreeItemPtr TreeItem, bool bVisible)
+void STedsVisibilityWidget::SetVisibility_Recursive(UE::Editor::DataStorage::ICoreProvider& DataStorage, FSceneOutlinerTreeItemPtr TreeItem, bool bVisible, 
+	UE::Editor::DataStorage::FRowHandleArray* ChangedRowsOut)
 {
 	using namespace UE::Editor::DataStorage;
 	using namespace UE::Editor::Outliner;
@@ -360,11 +399,11 @@ void STedsVisibilityWidget::SetVisibility_Recursive(UE::Editor::DataStorage::ICo
 		if (FActorTreeItem* ActorTreeItem = TreeItem->CastTo<FActorTreeItem>())
 		{
 			RowHandle ActorRow = DataStorageCompat->FindRowWithCompatibleObject(ActorTreeItem->Actor);
-			CommitVisibility(DataStorage, ActorRow, bVisible);
+			CommitVisibility(DataStorage, ActorRow, bVisible, ChangedRowsOut);
 		}
 		else if (FTedsOutlinerTreeItem* TedsItem = TreeItem->CastTo<FTedsOutlinerTreeItem>())
 		{
-			CommitVisibility(DataStorage, TedsItem->GetRowHandle(), bVisible);
+			CommitVisibility(DataStorage, TedsItem->GetRowHandle(), bVisible, ChangedRowsOut);
 		}
 	}
 
@@ -373,7 +412,7 @@ void STedsVisibilityWidget::SetVisibility_Recursive(UE::Editor::DataStorage::ICo
 		auto ChildTreeItem = ChildTreeItemPtr.Pin();
 		if (ChildTreeItem.IsValid())
 		{
-			SetVisibility_Recursive(DataStorage, ChildTreeItem, bVisible);
+			SetVisibility_Recursive(DataStorage, ChildTreeItem, bVisible, ChangedRowsOut);
 		}
 	}
 }
diff --git a/Engine/Plugins/Experimental/EditorDataStorageFeatures/Source/TedsOutliner/Public/Widgets/OutlinerVisibilityWidget.h b/Engine/Plugins/Experimental/EditorDataStorageFeatures/Source/TedsOutliner/Public/Widgets/OutlinerVisibilityWidget.h
--- a/Engine/Plugins/Experimental/EditorDataStorageFeatures/Source/TedsOutliner/Public/Widgets/OutlinerVisibilityWidget.h
+++ b/Engine/Plugins/Experimental/EditorDataStorageFeatures/Source/TedsOutliner/Public/Widgets/OutlinerVisibilityWidget.h
@@ -135,7 +135,8 @@ protected:
 	UE_API bool IsSelected() const;
 
 	/** Set the item this widget is responsible for to be hidden or shown */
-	UE_API void SetIsVisible(RowHandle InRow, const bool bVisible);
+	UE_API void SetIsVisible(RowHandle InRow, const bool bVisible, 
+		UE::Editor::DataStorage::FRowHandleArray* ChangedRowsToAppendTo = nullptr);
 
 	RowHandle TargetRow = UE::Editor::DataStorage::InvalidRowHandle;
 	RowHandle WidgetRow = UE::Editor::DataStorage::InvalidRowHandle;
@@ -152,14 +153,17 @@ protected:
 private:
 
 	// Commit the visibility state into the Rows DataStorage. Will trigger a SyncBackToWorld manually but the object is never dirtied.
-	static UE_API void CommitVisibility(UE::Editor::DataStorage::ICoreProvider& DataStorage, RowHandle Row, bool bVisible);
+	static UE_API void CommitVisibility(UE::Editor::DataStorage::ICoreProvider& DataStorage, RowHandle Row, bool bVisible, 
+		UE::Editor::DataStorage::FRowHandleArray* ChangedRowsToAppendTo);
 
 	// Recursive internal method to set visibility on a row and all its children.
 	static UE_API void SetVisibility_Recursive(UE::Editor::DataStorage::ICoreProvider& DataStorage, RowHandle Row,
-		TSharedPtr<UE::Editor::DataStorage::IHierarchyViewerDataInterface> HierarchyData, bool bVisible);
+		TSharedPtr<UE::Editor::DataStorage::IHierarchyViewerDataInterface> HierarchyData, bool bVisible,
+		UE::Editor::DataStorage::FRowHandleArray* ChangedRowsToAppendTo);
 
 	// Legacy method to work on Outliner tree items since they don't use the new hierarchy interface yet.
-	static UE_API void SetVisibility_Recursive(UE::Editor::DataStorage::ICoreProvider& DataStorage, FSceneOutlinerTreeItemPtr TreeItem, bool bVisible);
+	static UE_API void SetVisibility_Recursive(UE::Editor::DataStorage::ICoreProvider& DataStorage, FSceneOutlinerTreeItemPtr TreeItem, bool bVisible,
+		UE::Editor::DataStorage::FRowHandleArray* ChangedRowsToAppendTo);
 
 	// Returns the tree item associated with this widget.
 	UE_API FSceneOutlinerTreeItemPtr GetTreeItem(RowHandle InRow) const;
diff --git a/Engine/Source/Runtime/TypedElementFramework/Private/Elements/Columns/TypedElementVisibilityColumns.cpp b/Engine/Source/Runtime/TypedElementFramework/Private/Elements/Columns/TypedElementVisibilityColumns.cpp
new file mode 100644
--- /dev/null
+++ b/Engine/Source/Runtime/TypedElementFramework/Private/Elements/Columns/TypedElementVisibilityColumns.cpp
@@ -0,0 +1,57 @@
+// Copyright Epic Games, Inc. All Rights Reserved.
+
+#include "Elements/Columns/TypedElementVisibilityColumns.h"
+
+#if WITH_EDITOR
+#include "DataStorage/Features.h"
+#include "Elements/Interfaces/TypedElementDataStorageInterface.h"
+#endif //WITH_EDITOR
+
+#if WITH_EDITOR
+namespace UE::Editor::DataStorage
+{
+void FVisibilityColumnChange::Apply(UObject* Object)
+{
+	ICoreProvider* DataStorage = GetMutableDataStorageFeature<ICoreProvider>(StorageFeatureName);
+	if (!ensure(DataStorage))
+	{
+		return;
+	}
+
+	for (RowHandle Row : Rows.GetRows())
+	{
+		if (FVisibleInEditorColumn* VisibilityColumn = DataStorage->GetColumn<FVisibleInEditorColumn>(Row))
+		{
+			if (VisibilityColumn->bIsVisibleInEditor != bVisibleOnApply)
+			{
+				VisibilityColumn->bIsVisibleInEditor = bVisibleOnApply;
+				DataStorage->AddColumn<FVisibilityChangedTag>(Row);
+			}
+		}
+	}
+}
+
+void FVisibilityColumnChange::Revert(UObject* Object)
+{
+	ICoreProvider* DataStorage = GetMutableDataStorageFeature<ICoreProvider>(StorageFeatureName);
+	if (!ensure(DataStorage))
+	{
+		return;
+	}
+
+	const bool bVisibleOnRevert = !bVisibleOnApply;
+	for (RowHandle Row : Rows.GetRows())
+	{
+		if (FVisibleInEditorColumn* VisibilityColumn = DataStorage->GetColumn<FVisibleInEditorColumn>(Row))
+		{
+			if (VisibilityColumn->bIsVisibleInEditor != bVisibleOnRevert)
+			{
+				VisibilityColumn->bIsVisibleInEditor = bVisibleOnRevert;
+				DataStorage->AddColumn<FVisibilityChangedTag>(Row);
+			}
+		}
+	}
+}
+}//end UE::Editor::DataStorage
+
+#endif // WITH_EDITOR
diff --git a/Engine/Source/Runtime/TypedElementFramework/Public/Elements/Columns/TypedElementVisibilityColumns.h b/Engine/Source/Runtime/TypedElementFramework/Public/Elements/Columns/TypedElementVisibilityColumns.h
--- a/Engine/Source/Runtime/TypedElementFramework/Public/Elements/Columns/TypedElementVisibilityColumns.h
+++ b/Engine/Source/Runtime/TypedElementFramework/Public/Elements/Columns/TypedElementVisibilityColumns.h
@@ -4,6 +4,11 @@
 
 #include "Elements/Interfaces/TypedElementDataStorageInterface.h"
 
+#if WITH_EDITOR
+#include "Elements/Framework/TypedElementRowHandleArray.h"
+#include "Misc/Change.h"
+#endif
+
 #include "TypedElementVisibilityColumns.generated.h"
 
 /**
@@ -24,8 +29,41 @@ public:
  * This tag is required to be added for any changes you make to the TEDS column to be propagated back to the world, and is then removed at the end
  * of the frame.
  */
+UE_EXPERIMENTAL(5.8, "This tag might be changed into a column if extra context is needed.")
 USTRUCT(meta = (DisplayName = "Visibility Changed"))
 struct FVisibilityChangedTag final : public FEditorDataStorageTag
 {
 	GENERATED_BODY()
-};
\ No newline at end of file
+};
+
+#if WITH_EDITOR
+namespace UE::Editor::DataStorage
+{
+	/** Change object for undo/redo of TEDS visibility column modifications. */
+	UE_EXPERIMENTAL(5.8, "Undo handling of rows may change.")
+	class FVisibilityColumnChange : public FCommandChange
+	{
+	public:
+		/**
+		 * Note that ChangedRowsIn should only include rows whose visibilty state actually changed. Including
+		 *  rows whose visibility was already bVisibleOnApply before the transaction will cause undo to place
+		 *  the rows in unmatching state on undo. 
+		 */
+		FVisibilityColumnChange(FRowHandleArray&& ChangedRowsIn, bool bVisibleOnApplyIn)
+			: Rows(MoveTemp(ChangedRowsIn))
+			, bVisibleOnApply(bVisibleOnApplyIn)
+		{ }
+
+		TYPEDELEMENTFRAMEWORK_API void Apply(UObject* Object) override;
+		TYPEDELEMENTFRAMEWORK_API void Revert(UObject* Object) override;
+		FString ToString() const override 
+		{
+			return TEXT("Change In-Editor Visibility");
+		}
+
+	private:
+		FRowHandleArray Rows;
+		bool bVisibleOnApply;
+	};
+}//end UE::Editor::DataStorage
+#endif // WITH_EDITOR
\ No newline at end of file
