[TEDS Outliner] Add purposes for the full TEDS Outliner, the ability to show the Settings cog, and Eneity/Prefab filters in the Scene Graph TEDS Outliner

#Synced-CL 52254823

diff --git a/Engine/Plugins/Experimental/EditorDataStorageFeatures/Source/TedsOutliner/Private/TedsOutlinerMode.cpp b/Engine/Plugins/Experimental/EditorDataStorageFeatures/Source/TedsOutliner/Private/TedsOutlinerMode.cpp
--- a/Engine/Plugins/Experimental/EditorDataStorageFeatures/Source/TedsOutliner/Private/TedsOutlinerMode.cpp
+++ b/Engine/Plugins/Experimental/EditorDataStorageFeatures/Source/TedsOutliner/Private/TedsOutlinerMode.cpp
@@ -47,6 +47,8 @@ FTedsOutlinerMode::FTedsOutlinerMode(const FTedsOutlinerParams& InParams)
 	TedsOutlinerImpl = MakeShared<FTedsOutlinerImpl>(InParams, this, false /* bHybridMode */);
 	TedsOutlinerImpl->Init();
 	TedsOutlinerImpl->OnSelectionChanged().AddRaw(this, &FTedsOutlinerMode::OnSelectionChanged);
+
+	bShowViewButton = InParams.bShowViewButton;
 	
 	TedsOutlinerImpl->IsItemCompatible().BindLambda([](const ISceneOutlinerTreeItem& Item)
 	{
@@ -296,6 +298,11 @@ TUniquePtr<ISceneOutlinerHierarchy> FTedsOutlinerMode::CreateHierarchy()
 {
 	return MakeUnique<FTedsOutlinerHierarchy>(this, TedsOutlinerImpl.ToSharedRef());
 }
+
+bool FTedsOutlinerMode::ShowViewButton() const
+{
+	return bShowViewButton;
+}
 } // namespace UE::Editor::Outliner
 
 #undef LOCTEXT_NAMESPACE
diff --git a/Engine/Plugins/Experimental/EditorDataStorageFeatures/Source/TedsOutliner/Private/TedsOutlinerModule.cpp b/Engine/Plugins/Experimental/EditorDataStorageFeatures/Source/TedsOutliner/Private/TedsOutlinerModule.cpp
--- a/Engine/Plugins/Experimental/EditorDataStorageFeatures/Source/TedsOutliner/Private/TedsOutlinerModule.cpp
+++ b/Engine/Plugins/Experimental/EditorDataStorageFeatures/Source/TedsOutliner/Private/TedsOutlinerModule.cpp
@@ -5,7 +5,6 @@
 #include "Engine/Engine.h"
 #include "Engine/World.h"
 #include "LevelEditor.h"
-#include "LevelEditorContextMenu.h"
 #include "SceneOutlinerHelpers.h"
 #include "SceneOutlinerPublicTypes.h"
 #include "TedsOutlinerMode.h"
@@ -13,13 +12,10 @@
 #include "WorkspaceMenuStructureModule.h"
 #include "WorkspaceMenuStructure.h"
 #include "Elements/Columns/TypedElementCompatibilityColumns.h"
-#include "Elements/Columns/TypedElementLabelColumns.h"
 #include "Elements/Columns/TypedElementMiscColumns.h"
 #include "Elements/Columns/TypedElementPackageColumns.h"
 #include "Elements/Columns/TypedElementTypeInfoColumns.h"
-#include "Elements/Columns/TypedElementUIColumns.h"
 #include "Elements/Columns/TypedElementVisibilityColumns.h"
-#include "Columns/TedsActorComponentCompatibilityColumns.h"
 #include "Columns/TedsActorMobilityColumns.h"
 #include "Columns/TedsActorSocketColumns.h"
 #include "TedsAlertColumns.h"
@@ -180,105 +176,43 @@ TSharedRef<SWidget> FTedsOutlinerModule::CreateLevelEditorTedsOutliner()
 		.Text(LOCTEXT("TEDSPluginNotEnabledText", "You need to enable the Typed Element Data Storage plugin to see the table viewer!"));
 	}
 
-	using namespace DataStorage::Queries;
-
 	Storage = GetMutableDataStorageFeature<ICoreProvider>(StorageFeatureName);
+	StorageUi = GetMutableDataStorageFeature<IUiProvider>(UiFeatureName);
 
-	// The Outliner is populated with Actors and Entities
-	const FQueryDescription OutlinerQueryDescription =
-		Select()
-		.Where()
-			.All<FTypedElementUObjectColumn, FTypedElementLabelColumn, FTypedElementWorldColumn>()
-			.None<FHideRowFromUITag>()
-		.Compile();
-
-	FSceneOutlinerInitializationOptions InitOptions;
-	InitOptions.bShowHeaderRow = true;
-	InitOptions.FilterBarOptions.bHasFilterBar = true;
-	InitOptions.bShowTransient = true;
-	InitOptions.OutlinerIdentifier = "TEDSOutliner";
-
-	// Similar to what is done in SLevelEditor::CreateSceneOutliner, bind a callback that adjusts 
-	//  the context menu to have all the level editor context menu sections.
-	TWeakPtr<ILevelEditor> WeakLevelEditor;
-	if (const FLevelEditorModule* const LevelEditorModule = FModuleManager::GetModulePtr<FLevelEditorModule>(TEXT("LevelEditor")))
-	{
-		WeakLevelEditor = LevelEditorModule->GetLevelEditorInstance();
-	}
-	InitOptions.ModifyContextMenu.BindLambda([WeakLevelEditor](FName& OutMenuName, FToolMenuContext& MenuContext)
+	if (Storage && StorageUi)
 	{
-		static const FName ContextMenuName = "LevelEditor.TedsOutliner.ContextMenu";
-		FLevelEditorContextMenu::RegisterAsDerivedMenu(ContextMenuName, OutMenuName);
-		OutMenuName = ContextMenuName;
+		const RowHandle WidgetPurposeRow = StorageUi->FindPurpose(
+				IUiProvider::FPurposeInfo("LevelEditor", "SceneOutliner", "TedsOutliner").GeneratePurposeID());
 
-		if (WeakLevelEditor.IsValid())
-		{
-			FLevelEditorContextMenu::InitMenuContext(MenuContext, WeakLevelEditor, ELevelEditorMenuContext::SceneOutliner);
-		}
-	});
+		IUiProvider::FWidgetConstructorPtr OutSceneOutlinerConstructorPtr;
 
-	FTedsOutlinerParams Params(nullptr);
-	Params.QueryDescription = OutlinerQueryDescription;
-	Params.ColumnDescription = GetLevelEditorTedsOutlinerColumnDescription();
-	
-	const TArray<FHierarchyHandle> Hierarchies{
-		Storage->FindHierarchyByName(UTedsEditorHierarchyFactory::EditorObjectHierarchyName)
-		};
-	
-	Params.HierarchyData = MakeShared<FTedsOutlinerMultiHierarchyInterface>(MakeShared<FHierarchyViewerMultiData>(*Storage, Hierarchies));
-	
-	Params.Filters.Add(MakeShared<FTedsOutlinerFilter>(AActor::StaticClass()));
-	Params.Filters.Add(MakeShared<FTedsSceneOutlinerWorldFilter>(TEXT("WorldFilter")));
-	
-	/** Example ways to create filters using query descriptions and functions
-	 * Note QueryDescriptions should always be used over QueryFunctions when possible as they are more performant
-	 * Params.Filters.Emplace(
-	 * 	...
-	 *	Storage->RegisterQuery(
-	 *		Select()
-	 *		.Where(TColumn<FActorComponentTypeTag>())
-	 *		.Compile()));
-	 *
-	 * Params.Filters.Emplace(
-	 *	...
-	 *	[](TQueryContext<SingleRowInfo> Context, const FTedsActorMobilityColumn& MobilityColumn)
-	 *	{
-	 *		return MobilityColumn.Mobility == EComponentMobility::Static;
-	 *	});
-	 */
-	AddFavoriteFilters(Params);
-	
-	const FLevelEditorModule* const LevelEditorModule = FModuleManager::GetModulePtr<FLevelEditorModule>(TEXT("LevelEditor"));
-	if (LevelEditorModule)
-	{
-		if (auto LevelEditor = LevelEditorModule->GetLevelEditorInstance().Pin())
+		auto AssignWidgetToColumn = [&OutSceneOutlinerConstructorPtr] (IUiProvider::FWidgetConstructorPtr WidgetConstructor, TConstArrayView<TWeakObjectPtr<const UScriptStruct>>)
+			{
+				OutSceneOutlinerConstructorPtr = MoveTemp(WidgetConstructor);
+				return false;
+			};
+
+		StorageUi->CreateWidgetConstructors(WidgetPurposeRow, FMetaDataView(), AssignWidgetToColumn);
+		if (OutSceneOutlinerConstructorPtr)
 		{
-			Params.SelectionSet = LevelEditor->GetMutableElementSelectionSet();
+			const FName NamePurposeTableName = TEXT("Editor_WidgetTable");
+			
+			if (RowHandle RowHandle = Storage->AddRow(Storage->FindTable(NamePurposeTableName));
+				Storage->IsRowAvailable(RowHandle))
+			{
+				DataStorage::FMetaData OutlinerMetaData;
+				OutlinerMetaData.AddImmutableData(TEXT("OutlinerIdentifier"), TedsOutlinerTabName.ToString());
+
+				if (const TSharedPtr<SWidget> OutlinerWidget =
+					StorageUi->ConstructWidget(RowHandle, *OutSceneOutlinerConstructorPtr, FGenericMetaDataView(OutlinerMetaData)))
+				{
+					return OutlinerWidget.ToSharedRef();
+				}
+			}
 		}
 	}
 	
-	TSharedRef<ISceneOutliner> TEDSOutlinerShared = CreateTedsOutliner(InitOptions, Params);
-	
-	return TEDSOutlinerShared;
-}
-
-void FTedsOutlinerModule::AddFavoriteFilters(FTedsOutlinerParams& Params) const
-{
-	const TSharedRef<FFilterCategory> FavoriteFiltersCategory = MakeShared<FFilterCategory>(LOCTEXT("FavoriteFilters", "Favorites"), LOCTEXT("FavoriteFiltersTooltip", "Filter by Favorite state"));
-
-	// Create the filter
-	Params.Filters.Add(MakeShared<FTedsOutlinerFilter>(
-		"FilterByFavorites",
-		LOCTEXT("FavoriteFilterDisplayName", "Favorite"),
-		LOCTEXT("FavoriteFilterDisplayNameToolTip", "Filter by Favorites"),
-		FName("Icons.Star"),
-		FavoriteFiltersCategory,
-		Storage->RegisterQuery(
-			Select()
-			.Where()
-			.All<FFavoriteTag>()
-			.Compile())
-		));
+	return SNullWidget::NullWidget;
 }
 
 TSharedRef<SDockTab> FTedsOutlinerModule::OpenLevelEditorTedsOutliner(const FSpawnTabArgs& SpawnTabArgs)
diff --git a/Engine/Plugins/Experimental/EditorDataStorageFeatures/Source/TedsOutliner/Private/Widgets/SceneOutlinerWidget.cpp b/Engine/Plugins/Experimental/EditorDataStorageFeatures/Source/TedsOutliner/Private/Widgets/SceneOutlinerWidget.cpp
new file mode 100644
--- /dev/null
+++ b/Engine/Plugins/Experimental/EditorDataStorageFeatures/Source/TedsOutliner/Private/Widgets/SceneOutlinerWidget.cpp
@@ -0,0 +1,229 @@
+// Copyright Epic Games, Inc. All Rights Reserved.
+
+#include "Widgets/SceneOutlinerWidget.h"
+
+#include "Columns/SceneOutlinerColumns.h"
+#include "Elements/Columns/TypedElementUIColumns.h"
+#include "ILevelEditor.h"
+#include "LevelEditor.h"
+#include "LevelEditorContextMenu.h"
+#include "LevelEditorMenuContext.h"
+#include "Modules/ModuleManager.h"
+#include "TedsOutlinerHelpers.h"
+#include "TedsOutlinerModule.h"
+#include "TedsOutlinerImpl.h"
+#include "Elements/Columns/TypedElementCompatibilityColumns.h"
+#include "Elements/Columns/TypedElementLabelColumns.h"
+#include "Elements/Columns/TypedElementMiscColumns.h"
+#include "Elements/Framework/TypedElementSelectionSet.h"
+#include "Factories/TedsEditorHierarchyFactory.h"
+
+#define LOCTEXT_NAMESPACE "SceneOutlinerWidget"
+
+static FName LevelEditorModuleName = "LevelEditor";
+static FName SceneOutlinerModuleName = "SceneOutliner";
+static FName TedsOutlinerModuleName = "TEDSOutliner";
+
+static bool bShowTedsOutlinerRowHandles = false;
+// Cvar to add a column to show row handles for items in Teds Outliners
+static FAutoConsoleVariableRef CVarShowRowHandlesOutliner(
+TEXT("TEDS.UI.ShowTedsOutlinerRowHandles"),
+bShowTedsOutlinerRowHandles,
+TEXT("Add a column to show row handles in the TEDS Outliner (requires TEDS.UI.SetOutlinerPurpose!=BaseOutliner)"),
+FConsoleVariableDelegate::CreateLambda([](IConsoleVariable* /*CVar*/)
+{
+	UE::Editor::Outliner::Helpers::RefreshLevelEditorOutliners();
+}));
+
+void UTedsOutlinerWidgetFactory::RegisterWidgetPurposes(UE::Editor::DataStorage::IUiProvider& DataStorageUi) const
+{
+	using namespace UE::Editor::DataStorage;
+
+	const IUiProvider::FPurposeInfo TEDSOutlinerPurpose = IUiProvider::FPurposeInfo(
+		LevelEditorModuleName, SceneOutlinerModuleName, "TedsOutliner",
+		IUiProvider::EPurposeType::UniqueByName,
+		LOCTEXT("TEDSOutlinerPurpose", "Widget used to generate the TEDS ISceneOutliner widget."),
+		IUiProvider::FPurposeInfo(LevelEditorModuleName, SceneOutlinerModuleName, "ActorOutliner").GeneratePurposeID());
+
+	DataStorageUi.RegisterWidgetPurpose(TEDSOutlinerPurpose);
+}
+
+void UTedsOutlinerWidgetFactory::RegisterWidgetConstructors(UE::Editor::DataStorage::ICoreProvider& DataStorage,
+	UE::Editor::DataStorage::IUiProvider& DataStorageUi) const
+{
+	using namespace UE::Editor::DataStorage::Queries;
+	
+	DataStorageUi.RegisterWidgetFactory<FTedsSceneOutlinerWidgetConstructor>(
+		DataStorageUi.FindPurpose(IUiProvider::FPurposeInfo(LevelEditorModuleName, SceneOutlinerModuleName, "TedsOutliner").GeneratePurposeID()));
+}
+
+FTedsSceneOutlinerWidgetConstructor::FTedsSceneOutlinerWidgetConstructor()
+	: FSimpleWidgetConstructor(StaticStruct())
+{
+}
+
+TSharedPtr<SWidget> FTedsSceneOutlinerWidgetConstructor::CreateWidget(UE::Editor::DataStorage::ICoreProvider* DataStorage,
+	UE::Editor::DataStorage::IUiProvider* DataStorageUi,
+	UE::Editor::DataStorage::RowHandle TargetRow, UE::Editor::DataStorage::RowHandle WidgetRow, const UE::Editor::DataStorage::FMetaDataView& Arguments)
+{
+	using namespace UE::Editor::DataStorage;
+	using namespace UE::Editor::Outliner;
+	
+	// Fallback to the main SceneOutliner Tab Id
+	FName TabIdentifier = LevelEditorTabIds::LevelEditorSceneOutliner;
+	const FMetaDataEntryView OutlinerIdMeta = Arguments.FindGeneric("OutlinerIdentifier");
+	if(const FString* const* IdMetaData = OutlinerIdMeta.TryGetExact<const FString*>())
+	{
+		TabIdentifier = FName(**IdMetaData);
+	}
+
+	FSceneOutlinerInitializationOptions InitOptions;
+	GetSceneOutlinerInitializationOptions(TabIdentifier, InitOptions);
+
+	FTedsOutlinerParams Params(nullptr);
+	Params.QueryDescription = GetQueryDescription();
+	Params.ColumnDescription = GetColumnDescription();
+	GetFilters(DataStorage, Params.Filters);
+	Params.bShowRowHandleColumn = ShowRowHandleColumn();
+	Params.bForceShowParents = ForceShowParents();
+	Params.bUseDefaultObservers = UseDefaultObservers();
+	Params.bShowViewButton = ShowViewButton();
+	Params.HierarchyData = GetHierarchyData(DataStorage);
+	Params.SelectionSet = GetSelectionSet();
+
+	const FTedsOutlinerModule& TedsOutlinerModule = FModuleManager::Get().LoadModuleChecked<FTedsOutlinerModule>(TedsOutlinerModuleName);
+	TSharedRef<ISceneOutliner> Outliner = TedsOutlinerModule.CreateTedsOutliner(InitOptions, Params);
+	DataStorage->AddColumn(WidgetRow, FSceneOutlinerColumn { .Outliner = Outliner });
+	
+	return Outliner;
+}
+
+void FTedsSceneOutlinerWidgetConstructor::GetSceneOutlinerInitializationOptions(FName TabIdentifier, FSceneOutlinerInitializationOptions& InitOptions)
+{
+	GetLevelEditorSceneOutlinerInitOptions(TabIdentifier, InitOptions);
+	
+	InitOptions.FilterBarOptions.bUseSharedSettings = true;
+	InitOptions.FilterBarOptions.bHasFilterBar = true;
+	InitOptions.bShowTransient = true;
+
+	// Similar to what is done in SLevelEditor::CreateSceneOutliner, bind a callback that adjusts 
+	//  the context menu to have all the level editor context menu sections.
+	TWeakPtr<ILevelEditor> WeakLevelEditor;
+	if (const FLevelEditorModule* const LevelEditorModule = FModuleManager::GetModulePtr<FLevelEditorModule>(LevelEditorModuleName))
+	{
+		WeakLevelEditor = LevelEditorModule->GetLevelEditorInstance();
+	}
+	InitOptions.ModifyContextMenu.BindLambda([WeakLevelEditor](FName& OutMenuName, FToolMenuContext& MenuContext)
+	{
+		static const FName ContextMenuName = "LevelEditor.TedsOutliner.ContextMenu";
+		FLevelEditorContextMenu::RegisterAsDerivedMenu(ContextMenuName, OutMenuName);
+		OutMenuName = ContextMenuName;
+
+		if (WeakLevelEditor.IsValid())
+		{
+			FLevelEditorContextMenu::InitMenuContext(MenuContext, WeakLevelEditor, ELevelEditorMenuContext::SceneOutliner);
+		}
+	});
+}
+
+UE::Editor::DataStorage::FQueryDescription FTedsSceneOutlinerWidgetConstructor::GetQueryDescription() const
+{
+	using namespace UE::Editor::DataStorage::Queries;
+	
+	return Select()
+		.Where()
+			.All<FTypedElementUObjectColumn, FTypedElementLabelColumn, FTypedElementWorldColumn>()
+			.None<FHideRowFromUITag>()
+		.Compile();
+}
+
+UE::Editor::Outliner::FTedsOutlinerColumnDescription FTedsSceneOutlinerWidgetConstructor::GetColumnDescription() const
+{
+	using namespace UE::Editor::Outliner;
+	
+	FTedsOutlinerModule& TedsOutlinerModule = FModuleManager::Get().LoadModuleChecked<FTedsOutlinerModule>(TedsOutlinerModuleName);
+	return TedsOutlinerModule.GetLevelEditorTedsOutlinerColumnDescription();
+}
+
+void FTedsSceneOutlinerWidgetConstructor::GetFilters(UE::Editor::DataStorage::ICoreProvider* DataStorage, TArray<TSharedPtr<UE::Editor::Outliner::FTedsOutlinerFilter>>& Filters) const
+{
+	using namespace UE::Editor::Outliner;
+	
+	Filters.Add(MakeShared<FTedsOutlinerFilter>(AActor::StaticClass()));
+	Filters.Add(MakeShared<FTedsSceneOutlinerWorldFilter>(TEXT("WorldFilter")));
+	
+	AddFavoriteFilters(DataStorage, Filters);
+}
+
+bool FTedsSceneOutlinerWidgetConstructor::ShowRowHandleColumn() const
+{
+	return bShowTedsOutlinerRowHandles;
+}
+
+bool FTedsSceneOutlinerWidgetConstructor::ForceShowParents() const
+{
+	return true;
+}
+
+bool FTedsSceneOutlinerWidgetConstructor::UseDefaultObservers() const
+{
+	return true;
+}
+
+bool FTedsSceneOutlinerWidgetConstructor::ShowViewButton() const
+{
+	return true;
+}
+
+TSharedPtr<UE::Editor::Outliner::ITedsOutlinerHierarchyDataInterface> FTedsSceneOutlinerWidgetConstructor::GetHierarchyData(UE::Editor::DataStorage::ICoreProvider* DataStorage) const
+{
+	using namespace UE::Editor::Outliner;
+	
+	const TArray<FHierarchyHandle> Hierarchies{
+		DataStorage->FindHierarchyByName(UTedsEditorHierarchyFactory::EditorObjectHierarchyName)
+		};
+	
+	return MakeShared<FTedsOutlinerMultiHierarchyInterface>(MakeShared<FHierarchyViewerMultiData>(*DataStorage, Hierarchies));
+}
+
+UTypedElementSelectionSet* FTedsSceneOutlinerWidgetConstructor::GetSelectionSet() const
+{
+	TWeakPtr<ILevelEditor> LevelEditor = FModuleManager::GetModuleChecked<FLevelEditorModule>(LevelEditorModuleName).GetLevelEditorInstance();
+	if (TSharedPtr<ILevelEditor> LevelEditorPinned = LevelEditor.Pin())
+	{
+		return LevelEditorPinned->GetMutableElementSelectionSet();
+	}
+	return nullptr;
+}
+
+void FTedsSceneOutlinerWidgetConstructor::GetLevelEditorSceneOutlinerInitOptions(FName TabIdentifier, FSceneOutlinerInitializationOptions& OutInitOptions) const
+{
+	TWeakPtr<ILevelEditor> LevelEditor = FModuleManager::GetModuleChecked<FLevelEditorModule>(LevelEditorModuleName).GetLevelEditorInstance();
+	if (TSharedPtr<ILevelEditor> LevelEditorPinned = LevelEditor.Pin())
+	{
+		LevelEditorPinned->GetSceneOutlinerInitializationOptions(TabIdentifier, OutInitOptions);
+	}
+}
+
+void FTedsSceneOutlinerWidgetConstructor::AddFavoriteFilters(UE::Editor::DataStorage::ICoreProvider* DataStorage, TArray<TSharedPtr<UE::Editor::Outliner::FTedsOutlinerFilter>>& Filters) const
+{
+	using namespace UE::Editor::DataStorage::Queries;
+	
+	const TSharedRef<FFilterCategory> FavoriteFiltersCategory = MakeShared<FFilterCategory>(LOCTEXT("FavoriteFilters", "Favorites"), LOCTEXT("FavoriteFiltersTooltip", "Filter by Favorite state"));
+
+	// Create the filter
+	Filters.Add(MakeShared<UE::Editor::Outliner::FTedsOutlinerFilter>(
+		"FilterByFavorites",
+		LOCTEXT("FavoriteFilterDisplayName", "Favorite"),
+		LOCTEXT("FavoriteFilterDisplayNameToolTip", "Filter by Favorites"),
+		FName("Icons.Star"),
+		FavoriteFiltersCategory,
+		DataStorage->RegisterQuery(
+			Select()
+			.Where()
+			.All<FFavoriteTag>()
+			.Compile())
+		));
+}
+
+#undef LOCTEXT_NAMESPACE
\ No newline at end of file
diff --git a/Engine/Plugins/Experimental/EditorDataStorageFeatures/Source/TedsOutliner/Public/TedsOutlinerImpl.h b/Engine/Plugins/Experimental/EditorDataStorageFeatures/Source/TedsOutliner/Public/TedsOutlinerImpl.h
--- a/Engine/Plugins/Experimental/EditorDataStorageFeatures/Source/TedsOutliner/Public/TedsOutlinerImpl.h
+++ b/Engine/Plugins/Experimental/EditorDataStorageFeatures/Source/TedsOutliner/Public/TedsOutlinerImpl.h
@@ -145,6 +145,9 @@ struct FTedsOutlinerParams
 
 	// If true, the Teds Outliner will create observers to track addition/removal of rows and update the Outliner
 	bool bUseDefaultObservers;
+	
+	// If true, the Teds Outliner will create a view button (settings cog) pre-populated with hierarchy actions
+	bool bShowViewButton = false;
 
 	// If specified, this is how the TEDS Outliner will handle hierarchies. If not specified - there will be no hierarchies shown as a
 	// parent-child relation in the tree view
diff --git a/Engine/Plugins/Experimental/EditorDataStorageFeatures/Source/TedsOutliner/Public/TedsOutlinerMode.h b/Engine/Plugins/Experimental/EditorDataStorageFeatures/Source/TedsOutliner/Public/TedsOutlinerMode.h
--- a/Engine/Plugins/Experimental/EditorDataStorageFeatures/Source/TedsOutliner/Public/TedsOutlinerMode.h
+++ b/Engine/Plugins/Experimental/EditorDataStorageFeatures/Source/TedsOutliner/Public/TedsOutlinerMode.h
@@ -36,6 +36,7 @@ public:
 	UE_API virtual void OnDrop(ISceneOutlinerTreeItem& DropTarget, const FSceneOutlinerDragDropPayload& Payload, const FSceneOutlinerDragValidationInfo& ValidationInfo) const override;
 	UE_API virtual TSharedPtr<SWidget> CreateContextMenu() override;
 	UE_API virtual void Tick() override;
+	UE_API virtual bool ShowViewButton() const override;
 	/* end ISceneOutlinerMode interface */
 
 protected:
@@ -45,7 +46,7 @@ protected:
 	UE_API void OnSelectionChanged();
 
 protected:
-
+	bool bShowViewButton = false;
 	// The actual model for the TEDS Outliner
 	TSharedPtr<FTedsOutlinerImpl> TedsOutlinerImpl;
 };
diff --git a/Engine/Plugins/Experimental/EditorDataStorageFeatures/Source/TedsOutliner/Public/TedsOutlinerModule.h b/Engine/Plugins/Experimental/EditorDataStorageFeatures/Source/TedsOutliner/Public/TedsOutlinerModule.h
--- a/Engine/Plugins/Experimental/EditorDataStorageFeatures/Source/TedsOutliner/Public/TedsOutlinerModule.h
+++ b/Engine/Plugins/Experimental/EditorDataStorageFeatures/Source/TedsOutliner/Public/TedsOutlinerModule.h
@@ -10,6 +10,11 @@
 #include "TedsOutlinerFilter.h"
 #include "Elements/Framework/TypedElementQueryContext.h"
 
+namespace UE::Editor::DataStorage
+{
+	class IUiProvider;
+}
+
 struct FTypedElementWorldColumn;
 class ISceneOutliner;
 struct FSceneOutlinerInitializationOptions;
@@ -59,12 +64,12 @@ private:
 	void UnregisterLevelEditorTedsOutlinerTab();
 	TSharedRef<SDockTab> OpenLevelEditorTedsOutliner(const FSpawnTabArgs& SpawnTabArgs);
 	TSharedRef<SWidget> CreateLevelEditorTedsOutliner();
-	void AddFavoriteFilters(FTedsOutlinerParams& Params) const;
 	
 private:
 	FName TedsOutlinerTabName;
 	FDelegateHandle LevelEditorTabManagerChangedHandle;
 	DataStorage::ICoreProvider* Storage = nullptr;
+	DataStorage::IUiProvider* StorageUi = nullptr;
 };
 
 class FTedsSceneOutlinerWorldFilter : public FTedsOutlinerFilter
diff --git a/Engine/Plugins/Experimental/EditorDataStorageFeatures/Source/TedsOutliner/Public/Widgets/SceneOutlinerWidget.h b/Engine/Plugins/Experimental/EditorDataStorageFeatures/Source/TedsOutliner/Public/Widgets/SceneOutlinerWidget.h
new file mode 100644
--- /dev/null
+++ b/Engine/Plugins/Experimental/EditorDataStorageFeatures/Source/TedsOutliner/Public/Widgets/SceneOutlinerWidget.h
@@ -0,0 +1,94 @@
+// Copyright Epic Games, Inc. All Rights Reserved.
+
+#pragma once
+
+#include "Elements/Interfaces/TypedElementDataStorageUiInterface.h"
+#include "Elements/Interfaces/TypedElementDataStorageFactory.h"
+#include "Elements/Interfaces/TypedElementDataStorageInterface.h"
+
+#include "SceneOutlinerWidget.generated.h"
+
+#define UE_API TEDSOUTLINER_API
+
+struct FSceneOutlinerInitializationOptions;
+
+namespace UE::Editor::Outliner
+{
+	class ITedsOutlinerHierarchyDataInterface;
+	class FTedsOutlinerFilter;
+	struct FTedsOutlinerColumnDescription;
+	struct FTedsOutlinerParams;
+}
+
+namespace UE::Editor::DataStorage
+{
+	class ICoreProvider;
+}
+
+class UTypedElementSelectionSet;
+
+UCLASS()
+class UTedsOutlinerWidgetFactory : public UEditorDataStorageFactory
+{
+	GENERATED_BODY()
+
+public:
+	~UTedsOutlinerWidgetFactory() override = default;
+	
+	UE_API void RegisterWidgetPurposes(UE::Editor::DataStorage::IUiProvider& DataStorageUi) const override;
+
+	UE_API void RegisterWidgetConstructors(UE::Editor::DataStorage::ICoreProvider& DataStorage,
+		UE::Editor::DataStorage::IUiProvider& DataStorageUi) const override;
+};
+
+USTRUCT()
+struct FTedsSceneOutlinerWidgetConstructor : public FSimpleWidgetConstructor
+{
+	GENERATED_BODY()
+
+public:
+	TEDSOUTLINER_API FTedsSceneOutlinerWidgetConstructor();
+	~FTedsSceneOutlinerWidgetConstructor() override = default;
+
+	/**
+	 * Constructs the Outliner Widget. If overridden, ensure that the resulting SSceneOutliner is
+	 * referenced in an added FSceneOutlinerColumn on the WidgetRow and the TabIdentifier is gotten via metadata
+	 */
+	TEDSOUTLINER_API virtual TSharedPtr<SWidget> CreateWidget(
+		UE::Editor::DataStorage::ICoreProvider* DataStorage,
+		UE::Editor::DataStorage::IUiProvider* DataStorageUi,
+		UE::Editor::DataStorage::RowHandle TargetRow,
+		UE::Editor::DataStorage::RowHandle WidgetRow,
+		const UE::Editor::DataStorage::FMetaDataView& Arguments) override;
+
+	/** Adds parameters to the default (non-TEDS) SceneOutlinerInitializationOptions used to create the SSceneOutliner */
+	TEDSOUTLINER_API virtual void GetSceneOutlinerInitializationOptions(FName TabIdentifier, FSceneOutlinerInitializationOptions& InitOptions);
+	/** Gets the query description that will be used to populate rows in this Teds Outliner */
+	TEDSOUTLINER_API virtual UE::Editor::DataStorage::FQueryDescription GetQueryDescription() const;
+	/** Gets the columns to display in this Teds Outliner and any additional metadata */
+	TEDSOUTLINER_API virtual UE::Editor::Outliner::FTedsOutlinerColumnDescription GetColumnDescription() const;
+	/** Adds FTedsOutlinerFilters to an array used to populate the Outliner FilterBar */
+	TEDSOUTLINER_API virtual void GetFilters(UE::Editor::DataStorage::ICoreProvider* DataStorage,
+		TArray<TSharedPtr<UE::Editor::Outliner::FTedsOutlinerFilter>>& Filters) const;
+	/** Whether we should add the Row Handle column to the Outliner, by default this is controlled by a CVar */
+	TEDSOUTLINER_API virtual bool ShowRowHandleColumn() const;
+	/** Whether Parent rows should be force added back when filtered or searched out */
+	TEDSOUTLINER_API virtual bool ForceShowParents() const;
+	/** Whether the Teds Outliner will create observers to track addition/removal of rows and update */
+	TEDSOUTLINER_API virtual bool UseDefaultObservers() const;
+	/** Whether the Teds Outliner will create a view button (settings cog) pre-populated with hierarchy actions */
+	TEDSOUTLINER_API virtual bool ShowViewButton() const;
+	/** Gets the defined hierarchy to show, if nullptr, a flat list will be presented */
+	TEDSOUTLINER_API virtual TSharedPtr<UE::Editor::Outliner::ITedsOutlinerHierarchyDataInterface> GetHierarchyData(UE::Editor::DataStorage::ICoreProvider* DataStorage) const;
+	/** Gets the selection set to use for this Outliner, if nullptr it doesn't propagate tree selection through TypedElementSelectionInterface */
+	TEDSOUTLINER_API virtual UTypedElementSelectionSet* GetSelectionSet() const;
+	
+protected:
+	/** Gets the FSceneOutlinerInitializationOptions defined in the Level Editor for the Scene Outliner to be used as the default */
+	TEDSOUTLINER_API void GetLevelEditorSceneOutlinerInitOptions(FName TabIdentifier, FSceneOutlinerInitializationOptions& OutInitOptions) const;
+	
+private:
+	void AddFavoriteFilters(UE::Editor::DataStorage::ICoreProvider* DataStorage, TArray<TSharedPtr<UE::Editor::Outliner::FTedsOutlinerFilter>>& Filters) const;
+};
+
+#undef UE_API
