[TedsOutliner] Migrate World Partition Pinned column to TEDS Outliner

Adds the WP Pinned actor column to the TEDS Outliner

- FWorldPartitionPinnedColumn (bool bIsPinned) in TedsActorCompatibility
- Sync queries: world to TEDS (IsActorPinned) and TEDS to world (PinActors/UnpinActors via FTypedElementSyncBackToWorldTag)
- STedsPinnedWidget + cell/header factories registered on TColumn<FWorldPartitionPinnedColumn>()
- Registered in TedsOutlinerModule with EColumnPriorityGroup::Left
- Bridge mapping: FWorldPartitionPinnedColumn to FSceneOutlinerBuiltInColumnTypes::Pinned()

#Synced-CL 52344676

diff --git a/Engine/Plugins/Experimental/EditorDataStorageFeatures/Source/TedsActorCompatibility/Private/Compatibility/TedsActorWorldPartitionPinnedQueries.cpp b/Engine/Plugins/Experimental/EditorDataStorageFeatures/Source/TedsActorCompatibility/Private/Compatibility/TedsActorWorldPartitionPinnedQueries.cpp
new file mode 100644
--- /dev/null
+++ b/Engine/Plugins/Experimental/EditorDataStorageFeatures/Source/TedsActorCompatibility/Private/Compatibility/TedsActorWorldPartitionPinnedQueries.cpp
@@ -0,0 +1,134 @@
+// Copyright Epic Games, Inc. All Rights Reserved.
+
+#include "Compatibility/TedsActorWorldPartitionPinnedQueries.h"
+
+#include "Columns/TedsActorWorldPartitionColumns.h"
+#include "DataStorage/Queries/Types.h"
+#include "Elements/Columns/TypedElementCompatibilityColumns.h"
+#include "Elements/Columns/TypedElementMiscColumns.h"
+#include "Elements/Framework/TypedElementQueryBuilder.h"
+#include "GameFramework/Actor.h"
+#include "WorldPartition/LoaderAdapter/LoaderAdapterPinnedActors.h"
+#include "WorldPartition/WorldPartition.h"
+
+#include UE_INLINE_GENERATED_CPP_BY_NAME(TedsActorWorldPartitionPinnedQueries)
+
+void UActorWorldPartitionPinnedDataStorageFactory::RegisterQueries(UE::Editor::DataStorage::ICoreProvider& DataStorage)
+{
+	RegisterActorAddWorldPartitionPinnedColumn(DataStorage);
+	RegisterActorWorldPartitionPinnedToColumnQuery(DataStorage);
+	RegisterWorldPartitionPinnedColumnToActorQuery(DataStorage);
+}
+
+void UActorWorldPartitionPinnedDataStorageFactory::RegisterActorAddWorldPartitionPinnedColumn(UE::Editor::DataStorage::ICoreProvider& DataStorage) const
+{
+	using namespace UE::Editor::DataStorage::Queries;
+
+	DataStorage.RegisterQuery(
+		Select(
+			TEXT("Actor World Partition Pinned Object to New Column"),
+			FProcessor(EQueryTickPhase::PrePhysics, DataStorage.GetQueryTickGroupName(EQueryTickGroups::SyncExternalToDataStorage)).SetExecutionMode(EExecutionMode::GameThread),
+			[](IQueryContext& Context, RowHandle Row, const FTypedElementUObjectColumn& Object)
+			{
+				if (AActor* Actor = Cast<AActor>(Object.Object))
+				{
+					if (Actor->GetLevel() && FLoaderAdapterPinnedActors::SupportsPinning(Actor))
+					{
+						bool bIsPinned = false;
+						if (const UWorld* World = Actor->GetWorld())
+						{
+							if (const UWorldPartition* WorldPartition = World->GetWorldPartition())
+							{
+								bIsPinned = WorldPartition->IsActorPinned(Actor->GetActorGuid());
+							}
+						}
+						Context.AddColumn(Row, FWorldPartitionPinnedColumn{ .bIsPinned = bIsPinned });
+					}
+				}
+			})
+		.Where()
+			.All<FTypedElementActorTag, FTypedElementSyncFromWorldTag>()
+			.None<FWorldPartitionPinnedColumn>()
+		.Compile()
+	);
+}
+
+void UActorWorldPartitionPinnedDataStorageFactory::RegisterActorWorldPartitionPinnedToColumnQuery(UE::Editor::DataStorage::ICoreProvider& DataStorage) const
+{
+	using namespace UE::Editor::DataStorage::Queries;
+
+	DataStorage.RegisterQuery(
+		Select(
+			TEXT("Actor World Partition Pinned Object to Column"),
+			FProcessor(EQueryTickPhase::PrePhysics, DataStorage.GetQueryTickGroupName(EQueryTickGroups::SyncExternalToDataStorage)).SetExecutionMode(EExecutionMode::GameThread),
+			[](IQueryContext&, RowHandle, const FTypedElementUObjectColumn& Object, FWorldPartitionPinnedColumn& PinnedColumn)
+			{
+				if (const AActor* Actor = Cast<AActor>(Object.Object))
+				{
+					if (const UWorld* World = Actor->GetWorld())
+					{
+						if (const UWorldPartition* WorldPartition = World->GetWorldPartition())
+						{
+							PinnedColumn.bIsPinned = WorldPartition->IsActorPinned(Actor->GetActorGuid());
+						}
+					}
+				}
+			})
+		.Where()
+			.All<FTypedElementActorTag, FTypedElementSyncFromWorldTag>()
+		.Compile()
+	);
+}
+
+void UActorWorldPartitionPinnedDataStorageFactory::RegisterWorldPartitionPinnedColumnToActorQuery(UE::Editor::DataStorage::ICoreProvider& DataStorage) const
+{
+	using namespace UE::Editor::DataStorage::Queries;
+
+	struct FSetActorPinnedCommand
+	{
+		void operator()()
+		{
+			TStrongObjectPtr<AActor> PinnedActor = Actor.Pin();
+			if (PinnedActor)
+			{
+				if (UWorld* World = PinnedActor->GetWorld())
+				{
+					if (UWorldPartition* WorldPartition = World->GetWorldPartition())
+					{
+						const TArray<FGuid> ActorGuids{ PinnedActor->GetActorGuid() };
+						if (bIsPinned)
+						{
+							WorldPartition->PinActors(ActorGuids);
+						}
+						else
+						{
+							WorldPartition->UnpinActors(ActorGuids);
+						}
+					}
+				}
+			}
+		}
+		TWeakObjectPtr<AActor> Actor;
+		bool bIsPinned;
+	};
+
+	DataStorage.RegisterQuery(
+		Select(
+			TEXT("Actor World Partition Pinned Column to Object"),
+			FProcessor(EQueryTickPhase::PrePhysics, DataStorage.GetQueryTickGroupName(EQueryTickGroups::SyncDataStorageToExternal)).SetExecutionMode(EExecutionMode::GameThread),
+			[](IQueryContext& Context, RowHandle Row, FTypedElementUObjectColumn& Object, const FWorldPartitionPinnedColumn& PinnedColumn)
+			{
+				if (AActor* Actor = Cast<AActor>(Object.Object))
+				{
+					Context.PushCommand(FSetActorPinnedCommand
+					{
+						.Actor = Actor,
+						.bIsPinned = PinnedColumn.bIsPinned
+					});
+				}
+			})
+		.Where()
+			.All<FTypedElementActorTag, FTypedElementSyncBackToWorldTag>()
+		.Compile()
+	);
+}
diff --git a/Engine/Plugins/Experimental/EditorDataStorageFeatures/Source/TedsActorCompatibility/Private/Compatibility/TedsActorWorldPartitionPinnedQueries.h b/Engine/Plugins/Experimental/EditorDataStorageFeatures/Source/TedsActorCompatibility/Private/Compatibility/TedsActorWorldPartitionPinnedQueries.h
new file mode 100644
--- /dev/null
+++ b/Engine/Plugins/Experimental/EditorDataStorageFeatures/Source/TedsActorCompatibility/Private/Compatibility/TedsActorWorldPartitionPinnedQueries.h
@@ -0,0 +1,41 @@
+// Copyright Epic Games, Inc. All Rights Reserved.
+
+#pragma once
+
+#include "Elements/Interfaces/TypedElementDataStorageFactory.h"
+#include "UObject/ObjectMacros.h"
+
+#include "TedsActorWorldPartitionPinnedQueries.generated.h"
+
+/**
+ * Registers queries that keep FWorldPartitionPinnedColumn in sync with the World Partition
+ * pinned actor state.
+ */
+UCLASS()
+class UActorWorldPartitionPinnedDataStorageFactory : public UEditorDataStorageFactory
+{
+	GENERATED_BODY()
+
+public:
+	~UActorWorldPartitionPinnedDataStorageFactory() override = default;
+
+	void RegisterQueries(UE::Editor::DataStorage::ICoreProvider& DataStorage) override;
+
+private:
+	/**
+	 * Adds FWorldPartitionPinnedColumn to actor rows that support pinning
+	 * and do not already have the column.
+	 */
+	void RegisterActorAddWorldPartitionPinnedColumn(UE::Editor::DataStorage::ICoreProvider& DataStorage) const;
+
+	/**
+	 * Syncs the World Partition pinned state to FWorldPartitionPinnedColumn if they differ.
+	 */
+	void RegisterActorWorldPartitionPinnedToColumnQuery(UE::Editor::DataStorage::ICoreProvider& DataStorage) const;
+
+	/**
+	 * Applies the column's bIsPinned value back to the World Partition via PinActors / UnpinActors if the FTypedElementSyncBackToWorldTag
+	 * has been set and the bIsPinned differs.
+	 */
+	void RegisterWorldPartitionPinnedColumnToActorQuery(UE::Editor::DataStorage::ICoreProvider& DataStorage) const;
+};
diff --git a/Engine/Plugins/Experimental/EditorDataStorageFeatures/Source/TedsActorCompatibility/Public/Columns/TedsActorWorldPartitionColumns.h b/Engine/Plugins/Experimental/EditorDataStorageFeatures/Source/TedsActorCompatibility/Public/Columns/TedsActorWorldPartitionColumns.h
new file mode 100644
--- /dev/null
+++ b/Engine/Plugins/Experimental/EditorDataStorageFeatures/Source/TedsActorCompatibility/Public/Columns/TedsActorWorldPartitionColumns.h
@@ -0,0 +1,47 @@
+// Copyright Epic Games, Inc. All Rights Reserved.
+
+#pragma once
+
+#include "Elements/Interfaces/TypedElementDataStorageInterface.h"
+
+#include "TedsActorWorldPartitionColumns.generated.h"
+
+/**
+ * Column that stores the World Partition pinned state of an actor.
+ * Only added to rows for actors that support pinning (spatially loaded WP actors).
+ */
+USTRUCT(meta = (DisplayName = "World Partition Pinned"))
+struct FWorldPartitionPinnedColumn final : public FEditorDataStorageColumn
+{
+	GENERATED_BODY()
+
+	UPROPERTY()
+	bool bIsPinned = false;
+};
+
+/**
+ * Column that stores the World Partition data layer names for an actor.
+ * Only added to rows for actors in worlds with World Partition enabled.
+ */
+USTRUCT(meta = (DisplayName = "World Partition Data Layer"))
+struct FWorldPartitionDataLayerColumn final : public FEditorDataStorageColumn
+{
+	GENERATED_BODY()
+
+	UPROPERTY(meta = (Sortable, Searchable))
+	FString DataLayers;
+};
+
+/**
+ * Column that stores the World Asset Package name for Level Instance actors.
+ * Only added to rows for actors that implement ILevelInstanceInterface.
+ */
+USTRUCT(meta = (DisplayName = "World Partition Sub Package"))
+struct FWorldPartitionSubPackageColumn final : public FEditorDataStorageColumn
+{
+	GENERATED_BODY()
+
+	UPROPERTY(meta = (Sortable, Searchable))
+	FString SubPackage;
+};
+
diff --git a/Engine/Plugins/Experimental/EditorDataStorageFeatures/Source/TedsOutliner/Private/Compatibility/SceneOutlinerTedsBridge.cpp b/Engine/Plugins/Experimental/EditorDataStorageFeatures/Source/TedsOutliner/Private/Compatibility/SceneOutlinerTedsBridge.cpp
--- a/Engine/Plugins/Experimental/EditorDataStorageFeatures/Source/TedsOutliner/Private/Compatibility/SceneOutlinerTedsBridge.cpp
+++ b/Engine/Plugins/Experimental/EditorDataStorageFeatures/Source/TedsOutliner/Private/Compatibility/SceneOutlinerTedsBridge.cpp
@@ -12,6 +12,7 @@
 #include "Elements/Columns/TypedElementTypeInfoColumns.h"
 #include "Elements/Columns/TypedElementVisibilityColumns.h"
 #include "Columns/TedsActorSocketColumns.h"
+#include "Columns/TedsActorWorldPartitionColumns.h"
 #include "Elements/Interfaces/TypedElementDataStorageInterface.h"
 #include "Elements/Interfaces/TypedElementDataStorageUiInterface.h"
 #include "Layers/Columns/LayersColumns.h"
@@ -42,6 +43,7 @@ void USceneOutlinerTedsBridgeFactory::PreRegister(UE::Editor::DataStorage::ICore
 	TEDSToOutlinerDefaultColumnMapping.Add(FTedsPrimaryPackageObjectTag::StaticStruct(), FSceneOutlinerBuiltInColumnTypes::Unsaved());
 	TEDSToOutlinerDefaultColumnMapping.Add(FTedsActorSocketColumn::StaticStruct(), FSceneOutlinerBuiltInColumnTypes::Socket());
 	TEDSToOutlinerDefaultColumnMapping.Add(FEditorDataStorageActorLayersColumn::StaticStruct(), FSceneOutlinerBuiltInColumnTypes::Layer());
+	TEDSToOutlinerDefaultColumnMapping.Add(FWorldPartitionPinnedColumn::StaticStruct(), FSceneOutlinerBuiltInColumnTypes::Pinned());
 }
 
 void USceneOutlinerTedsBridgeFactory::RegisterTables(UE::Editor::DataStorage::ICoreProvider& DataStorage)
diff --git a/Engine/Plugins/Experimental/EditorDataStorageFeatures/Source/TedsOutliner/Private/TedsOutlinerModule.cpp b/Engine/Plugins/Experimental/EditorDataStorageFeatures/Source/TedsOutliner/Private/TedsOutlinerModule.cpp
--- a/Engine/Plugins/Experimental/EditorDataStorageFeatures/Source/TedsOutliner/Private/TedsOutlinerModule.cpp
+++ b/Engine/Plugins/Experimental/EditorDataStorageFeatures/Source/TedsOutliner/Private/TedsOutlinerModule.cpp
@@ -16,6 +16,7 @@
 #include "Elements/Columns/TypedElementPackageColumns.h"
 #include "Elements/Columns/TypedElementTypeInfoColumns.h"
 #include "Elements/Columns/TypedElementVisibilityColumns.h"
+#include "Columns/TedsActorWorldPartitionColumns.h"
 #include "Columns/TedsActorMobilityColumns.h"
 #include "Columns/TedsActorSocketColumns.h"
 #include "TedsAlertColumns.h"
@@ -140,6 +141,7 @@ FTedsOutlinerColumnDescription FTedsOutlinerModule::GetLevelEditorTedsOutlinerCo
 	TArray<TWeakObjectPtr<const UScriptStruct>> OutColumns({
 		FVisibleInEditorColumn::StaticStruct(),
 		FTedsPrimaryPackageObjectTag::StaticStruct(),
+		FWorldPartitionPinnedColumn::StaticStruct(),
 		FTypedElementClassTypeInfoColumn::StaticStruct(),
 		FLevelColumn::StaticStruct(), // PackageShortName
 		FTedsActorMobilityColumn::StaticStruct(),
@@ -163,6 +165,7 @@ FTedsOutlinerColumnDescription FTedsOutlinerModule::GetLevelEditorTedsOutlinerCo
 	{
 		{FVisibleInEditorColumn::StaticStruct(), FTedsOutlinerColumnParams(EColumnPriorityGroup::Left)},
 		{FTedsPrimaryPackageObjectTag::StaticStruct(), FTedsOutlinerColumnParams(EColumnPriorityGroup::Left)},
+		{FWorldPartitionPinnedColumn::StaticStruct(), FTedsOutlinerColumnParams(EColumnPriorityGroup::Left)},
 	};
 	
 	return FTedsOutlinerColumnDescription(OutColumns, OutColumnParams);
diff --git a/Engine/Plugins/Experimental/EditorDataStorageFeatures/Source/TedsOutliner/Private/Widgets/OutlinerPinnedWidget.cpp b/Engine/Plugins/Experimental/EditorDataStorageFeatures/Source/TedsOutliner/Private/Widgets/OutlinerPinnedWidget.cpp
new file mode 100644
--- /dev/null
+++ b/Engine/Plugins/Experimental/EditorDataStorageFeatures/Source/TedsOutliner/Private/Widgets/OutlinerPinnedWidget.cpp
@@ -0,0 +1,217 @@
+// Copyright Epic Games, Inc. All Rights Reserved.
+
+#include "Widgets/OutlinerPinnedWidget.h"
+
+#include "Columns/SlateHeaderColumns.h"
+#include "Columns/TedsActorWorldPartitionColumns.h"
+#include "Columns/TedsOutlinerColumns.h"
+#include "Elements/Columns/TypedElementMiscColumns.h"
+#include "Compatibility/SceneOutlinerTedsBridge.h"
+#include "DataStorage/Features.h"
+#include "Elements/Columns/TypedElementCompatibilityColumns.h"
+#include "Elements/Columns/TypedElementSelectionColumns.h"
+#include "Elements/Framework/TypedElementAttributeBinding.h"
+#include "Elements/Framework/TypedElementQueryBuilder.h"
+#include "Elements/Interfaces/TypedElementDataStorageInterface.h"
+#include "Widgets/SBoxPanel.h"
+
+#include UE_INLINE_GENERATED_CPP_BY_NAME(OutlinerPinnedWidget)
+
+#define LOCTEXT_NAMESPACE "OutlinerPinnedWidget"
+
+//
+// Factory
+//
+
+FOutlinerPinnedWidgetConstructor::FOutlinerPinnedWidgetConstructor()
+	: FSimpleWidgetConstructor(StaticStruct())
+{
+}
+
+void UOutlinerPinnedFactory::RegisterWidgetConstructors(UE::Editor::DataStorage::ICoreProvider& DataStorage,
+	UE::Editor::DataStorage::IUiProvider& DataStorageUi) const
+{
+	using namespace UE::Editor::DataStorage::Queries;
+	DataStorageUi.RegisterWidgetFactory<FOutlinerPinnedWidgetConstructor>(
+		DataStorageUi.FindPurpose(IUiProvider::FPurposeInfo("SceneOutliner", "Cell", NAME_None).GeneratePurposeID()),
+		TColumn<FWorldPartitionPinnedColumn>());
+	DataStorageUi.RegisterWidgetFactory<FOutlinerPinnedHeaderConstructor>(
+		DataStorageUi.FindPurpose(IUiProvider::FPurposeInfo("SceneOutliner", "Header", NAME_None).GeneratePurposeID()),
+		TColumn<FWorldPartitionPinnedColumn>());
+}
+
+TSharedPtr<SWidget> FOutlinerPinnedWidgetConstructor::CreateWidget(UE::Editor::DataStorage::ICoreProvider* DataStorage,
+	UE::Editor::DataStorage::IUiProvider* DataStorageUi,
+	UE::Editor::DataStorage::RowHandle TargetRow, UE::Editor::DataStorage::RowHandle WidgetRow,
+	const UE::Editor::DataStorage::FMetaDataView& Arguments)
+{
+	return SNew(SHorizontalBox)
+		+ SHorizontalBox::Slot()
+		.AutoWidth()
+		.VAlign(VAlign_Center)
+		[
+			SNew(STedsPinnedWidget, TargetRow)
+				.ToolTipText(LOCTEXT("SceneOutlinerPinnedToggleTooltip",
+					"Toggles Force Load: Keep the selected items loaded in the editor even when they don't overlap a loaded World Partition region."))
+		];
+}
+
+//
+// Header Constructor
+//
+
+FOutlinerPinnedHeaderConstructor::FOutlinerPinnedHeaderConstructor()
+	: FSimpleWidgetConstructor(StaticStruct())
+{
+}
+
+TSharedPtr<SWidget> FOutlinerPinnedHeaderConstructor::CreateWidget(UE::Editor::DataStorage::ICoreProvider* DataStorage,
+	UE::Editor::DataStorage::IUiProvider* DataStorageUi,
+	UE::Editor::DataStorage::RowHandle TargetRow, UE::Editor::DataStorage::RowHandle WidgetRow,
+	const UE::Editor::DataStorage::FMetaDataView& Arguments)
+{
+	DataStorage->AddColumn(WidgetRow, FHeaderWidgetSizeColumn
+	{
+		.ColumnSizeMode = EColumnSizeMode::Fixed,
+		.Width = 24.0f
+	});
+
+	return SNew(SImage)
+		.DesiredSizeOverride(FVector2D(16.f, 16.f))
+		.Image(FAppStyle::Get().GetBrush(TEXT("Icons.Unpinned")))
+		.ToolTipText(LOCTEXT("SceneOutlinerPinnedHeaderTooltip", "Force Load"));
+}
+
+//
+// STedsPinnedWidget
+//
+
+void STedsPinnedWidget::Construct(const FArguments&, const RowHandle& InTargetRow)
+{
+	using namespace UE::Editor::DataStorage;
+
+	TargetRow = InTargetRow;
+
+	static const FName NAME_PinnedBrush = TEXT("Icons.Pinned");
+	static const FName NAME_UnpinnedBrush = TEXT("Icons.Unpinned");
+
+	FAttributeBinder Binder(InTargetRow, GetDataStorage());
+
+	bIsPinnedAttr = Binder.BindData(&FWorldPartitionPinnedColumn::bIsPinned);
+	bIsSelectedAttr = Binder.BindColumnsPresence<FTypedElementSelectionColumn>([](bool bHasColumn) { return bHasColumn; });
+
+	SImage::Construct(
+		SImage::FArguments()
+		.ColorAndOpacity(this, &STedsPinnedWidget::GetForegroundColor)
+		.Image(Binder.BindData(&FWorldPartitionPinnedColumn::bIsPinned,
+			[Pinned = FAppStyle::Get().GetBrush(NAME_PinnedBrush),
+			Unpinned = FAppStyle::Get().GetBrush(NAME_UnpinnedBrush)]
+			(const bool& bIsPinned) -> const FSlateBrush*
+			{
+				return bIsPinned ? Pinned : Unpinned;
+			}))
+	);
+}
+
+FReply STedsPinnedWidget::HandleClick()
+{
+	const bool bNewPinnedState = !bIsPinnedAttr.Get();
+
+	TArray<RowHandle> SelectedRows;
+	GetSelectedRows(SelectedRows);
+
+	if (SelectedRows.Contains(TargetRow))
+	{
+		for (const RowHandle Row : SelectedRows)
+		{
+			SetIsPinned(Row, bNewPinnedState);
+		}
+	}
+	else
+	{
+		SetIsPinned(TargetRow, bNewPinnedState);
+	}
+
+	return FReply::Handled();
+}
+
+FReply STedsPinnedWidget::OnMouseButtonDoubleClick(const FGeometry& InMyGeometry, const FPointerEvent& InMouseEvent)
+{
+	return HandleClick();
+}
+
+FReply STedsPinnedWidget::OnMouseButtonDown(const FGeometry& MyGeometry, const FPointerEvent& MouseEvent)
+{
+	if (MouseEvent.GetEffectingButton() == EKeys::LeftMouseButton)
+	{
+		return HandleClick().PreventThrottling();
+	}
+	return FReply::Unhandled();
+}
+
+FSlateColor STedsPinnedWidget::GetForegroundColor() const
+{
+	if (!bIsPinnedAttr.Get() && !IsHovered() && !bIsSelectedAttr.Get())
+	{
+		return FLinearColor::Transparent;
+	}
+
+	return IsHovered() ? FSlateColor::UseForeground() : FSlateColor::UseSubduedForeground();
+}
+
+void STedsPinnedWidget::SetIsPinned(RowHandle InRow, bool bPinned)
+{
+	using namespace UE::Editor::DataStorage;
+
+	if (ICoreProvider* DataStorage = GetDataStorage())
+	{
+		CommitPinnedState(*DataStorage, InRow, bPinned);
+	}
+}
+
+void STedsPinnedWidget::CommitPinnedState(UE::Editor::DataStorage::ICoreProvider& DataStorage, RowHandle Row, bool bPinned)
+{
+	if (FWorldPartitionPinnedColumn* PinnedColumn = DataStorage.GetColumn<FWorldPartitionPinnedColumn>(Row))
+	{
+		PinnedColumn->bIsPinned = bPinned;
+	}
+	DataStorage.AddColumn<FTypedElementSyncBackToWorldTag>(Row);
+}
+
+void STedsPinnedWidget::GetSelectedRows(TArray<RowHandle>& OutSelectedRows) const
+{
+	using namespace UE::Editor::DataStorage;
+	using namespace UE::Editor::DataStorage::Queries;
+
+	if (ICoreProvider* DataStorage = GetDataStorage())
+	{
+		if (const FTypedElementSelectionColumn* SelectionColumn = DataStorage->GetColumn<FTypedElementSelectionColumn>(TargetRow))
+		{
+			FName SelectionSet = SelectionColumn->SelectionSet;
+
+			static QueryHandle AllSelectedItemsQuery =
+				DataStorage->RegisterQuery(
+					Select()
+						.ReadOnly<FTypedElementSelectionColumn>()
+					.Where()
+					.Compile());
+
+			DataStorage->RunQuery(AllSelectedItemsQuery,
+				CreateDirectQueryCallbackBinding([&OutSelectedRows, SelectionSet](const IDirectQueryContext&, RowHandle Row, const FTypedElementSelectionColumn& InSelectionColumn)
+				{
+					if (InSelectionColumn.SelectionSet == SelectionSet)
+					{
+						OutSelectedRows.Add(Row);
+					}
+				}));
+		}
+	}
+}
+
+UE::Editor::DataStorage::ICoreProvider* STedsPinnedWidget::GetDataStorage()
+{
+	using namespace UE::Editor::DataStorage;
+	return GetMutableDataStorageFeature<ICoreProvider>(StorageFeatureName);
+}
+
+#undef LOCTEXT_NAMESPACE
diff --git a/Engine/Plugins/Experimental/EditorDataStorageFeatures/Source/TedsOutliner/Public/Widgets/OutlinerPinnedWidget.h b/Engine/Plugins/Experimental/EditorDataStorageFeatures/Source/TedsOutliner/Public/Widgets/OutlinerPinnedWidget.h
new file mode 100644
--- /dev/null
+++ b/Engine/Plugins/Experimental/EditorDataStorageFeatures/Source/TedsOutliner/Public/Widgets/OutlinerPinnedWidget.h
@@ -0,0 +1,107 @@
+// Copyright Epic Games, Inc. All Rights Reserved.
+
+#pragma once
+
+#include "Elements/Interfaces/TypedElementDataStorageFactory.h"
+#include "Elements/Interfaces/TypedElementDataStorageUiInterface.h"
+#include "Internationalization/Text.h"
+#include "UObject/ObjectMacros.h"
+#include "Widgets/Images/SImage.h"
+#include "ScopedTransaction.h"
+#include "SceneOutlinerFwd.h"
+
+#include "OutlinerPinnedWidget.generated.h"
+
+#define UE_API TEDSOUTLINER_API
+
+namespace UE::Editor::DataStorage
+{
+	class ICoreProvider;
+}
+
+UCLASS()
+class UOutlinerPinnedFactory : public UEditorDataStorageFactory
+{
+	GENERATED_BODY()
+
+public:
+	~UOutlinerPinnedFactory() override = default;
+
+	TEDSOUTLINER_API void RegisterWidgetConstructors(UE::Editor::DataStorage::ICoreProvider& DataStorage, UE::Editor::DataStorage::IUiProvider& DataStorageUi) const override;
+};
+
+USTRUCT()
+struct FOutlinerPinnedWidgetConstructor : public FSimpleWidgetConstructor
+{
+	GENERATED_BODY()
+
+public:
+	TEDSOUTLINER_API FOutlinerPinnedWidgetConstructor();
+	~FOutlinerPinnedWidgetConstructor() override = default;
+
+	TEDSOUTLINER_API virtual TSharedPtr<SWidget> CreateWidget(
+		UE::Editor::DataStorage::ICoreProvider* DataStorage,
+		UE::Editor::DataStorage::IUiProvider* DataStorageUi,
+		UE::Editor::DataStorage::RowHandle TargetRow,
+		UE::Editor::DataStorage::RowHandle WidgetRow,
+		const UE::Editor::DataStorage::FMetaDataView& Arguments) override;
+};
+
+USTRUCT()
+struct FOutlinerPinnedHeaderConstructor : public FSimpleWidgetConstructor
+{
+	GENERATED_BODY()
+
+public:
+	TEDSOUTLINER_API FOutlinerPinnedHeaderConstructor();
+	~FOutlinerPinnedHeaderConstructor() override = default;
+
+	TEDSOUTLINER_API virtual TSharedPtr<SWidget> CreateWidget(
+		UE::Editor::DataStorage::ICoreProvider* DataStorage,
+		UE::Editor::DataStorage::IUiProvider* DataStorageUi,
+		UE::Editor::DataStorage::RowHandle TargetRow,
+		UE::Editor::DataStorage::RowHandle WidgetRow,
+		const UE::Editor::DataStorage::FMetaDataView& Arguments) override;
+};
+
+template<>
+struct TWidgetTypeTraits<class STedsPinnedWidget>
+{
+	static constexpr bool SupportsInvalidation() { return true; }
+};
+
+/** Widget responsible for displaying and toggling the World Partition pinned state for a single item */
+class STedsPinnedWidget : public SImage
+{
+	using RowHandle = UE::Editor::DataStorage::RowHandle;
+
+public:
+	SLATE_BEGIN_ARGS(STedsPinnedWidget) {}
+	SLATE_END_ARGS()
+
+	UE_API void Construct(const FArguments& InArgs, const RowHandle& InTargetRow);
+
+protected:
+	UE_API virtual FReply OnMouseButtonDoubleClick(const FGeometry& InMyGeometry, const FPointerEvent& InMouseEvent) override;
+	UE_API virtual FReply OnMouseButtonDown(const FGeometry& MyGeometry, const FPointerEvent& MouseEvent) override;
+
+	UE_API FReply HandleClick();
+	UE_API virtual FSlateColor GetForegroundColor() const override;
+
+	/** Toggle the pinned state of InRow in the data storage */
+	UE_API void SetIsPinned(RowHandle InRow, bool bPinned);
+
+	RowHandle TargetRow = UE::Editor::DataStorage::InvalidRowHandle;
+
+	TAttribute<bool> bIsPinnedAttr;
+	TAttribute<bool> bIsSelectedAttr;
+
+private:
+	static UE_API void CommitPinnedState(UE::Editor::DataStorage::ICoreProvider& DataStorage, RowHandle Row, bool bPinned);
+
+	UE_API void GetSelectedRows(TArray<RowHandle>& OutSelectedRows) const;
+
+	static UE_API UE::Editor::DataStorage::ICoreProvider* GetDataStorage();
+};
+
+#undef UE_API
