[TedsAlerts] Fix hierarchy traversal and leaf-to-root evaluation order in child alert propagation
Fix:                                                                                                                                                                                                                                                         Replace the single global AlertConditionName and ancestor-walking loop (IncrementParents / MoveToNextParent / FTableRowParentColumn) with a per-hierarchy multi-frame bottom-up propagation wave.

Each registered hierarchy now gets three dedicated condition names (Alerts_, AlertsPropagation_, AlertsPropagationFinished_) and a PendingPropagations atomic counter. The propagation sequence is:

Preamble "Add missing child alerts" - add FTedsChildAlertColumn to leaf rows and their direct parent (one level only)
Preamble "Reset pending propagations" - zero the atomic at wave start
PreUpdate "Clear child alerts" - zero Counts[] / ProcessedChildCount
Update "Count alert children" - increment PendingPropagations per node, parent.NumAlertChildren++, arm wave
PostUpdate "Propagate child alerts" - skip until all children reported in, then push accumulated counts + own alert to parent; decrement PendingPropagations
Postamble "Check alert propagation complete" - re-arm next frame if PendingPropagations>0, else activate Finished condition
PostUpdate "Remove unused child alerts" - remove FTedsChildAlertColumn from zero-count rows once propagation is complete

Also fixes AssignAlert / SwapAlerts which were silently dropping FTedsAlertHierarchyColumn when an alert was sorted into the chain.

Parent walking now uses Context.GetParentRow() directly; FTableRowParentColumn and the CachedParent fields have been removed.

#Synced-CL 52344676

diff --git a/Engine/Plugins/Experimental/EditorDataStorageFeatures/Source/TedsAlerts/Private/TedsAlertsFactory.cpp b/Engine/Plugins/Experimental/EditorDataStorageFeatures/Source/TedsAlerts/Private/TedsAlertsFactory.cpp
--- a/Engine/Plugins/Experimental/EditorDataStorageFeatures/Source/TedsAlerts/Private/TedsAlertsFactory.cpp
+++ b/Engine/Plugins/Experimental/EditorDataStorageFeatures/Source/TedsAlerts/Private/TedsAlertsFactory.cpp
@@ -12,7 +12,6 @@
 
 const FName UTedsAlertsFactory::AlertChainTableName = "Alerts chain";
 const FName UTedsAlertsFactory::UnsortedAlertChainTableName = "Alerts chain (unsorted)";
-const FName UTedsAlertsFactory::AlertConditionName = FName(TEXT("Alerts"));
 
 void UTedsAlertsFactory::RegisterTables(UE::Editor::DataStorage::ICoreProvider& DataStorage)
 {
@@ -25,7 +24,6 @@ void UTedsAlertsFactory::RegisterQueries(UE::Editor::DataStorage::ICoreProvider&
 	RegisterSubQueries(DataStorage);
 	RegisterSortUnsortedAlertsQuery(DataStorage);
 	RegisterOnRemoveQueries(DataStorage);
-	RegisterChildAlertUpdatesQueries(DataStorage);
 }
 
 UE::Editor::DataStorage::TableHandle UTedsAlertsFactory::GetAlertChainTable() const
@@ -50,19 +48,26 @@ UE::Editor::DataStorage::QueryHandle UTedsAlertsFactory::GetUnsortedAlertsQuery(
 
 void UTedsAlertsFactory::RegisterAlertHierarchy(UE::Editor::DataStorage::ICoreProvider& DataStorage, const FName& InHierarchyName)
 {
-	using namespace UE::Editor::DataStorage::Queries;
-
+	using namespace UE::Editor::DataStorage;
 	FHierarchyHandle HierarchyHandle = DataStorage.FindHierarchyByName(InHierarchyName);
-	
-	if (DataStorage.IsValidHierarchyHandle(HierarchyHandle) && !RegisteredHierarchies.Contains(InHierarchyName))
-	{
-		RegisterOnAddHierarchyQueries(DataStorage, HierarchyHandle);
-		RegisterOnRemoveHierarchyQueries(DataStorage, HierarchyHandle);
-		RegisterParentUpdatesHierarchyQueries(DataStorage, HierarchyHandle);
-		RegisterChildAlertHierarchyUpdatesQueries(DataStorage, InHierarchyName, HierarchyHandle);
 
-		RegisteredHierarchies.Add(InHierarchyName);
+	if (!DataStorage.IsValidHierarchyHandle(HierarchyHandle) || HierarchyConditions.Contains(InHierarchyName))
+	{
+		return;
 	}
+
+	TUniquePtr<FAlertHierarchyConditions> Conditions = MakeUnique<FAlertHierarchyConditions>();
+	Conditions->AlertPropagationStartConditionName = FName(FString::Printf(TEXT("Alerts_%s"), *InHierarchyName.ToString()));
+	Conditions->AlertPropagationConditionName = FName(FString::Printf(TEXT("AlertsPropagation_%s"), *InHierarchyName.ToString()));
+	Conditions->AlertPropagationFinishedConditionName = FName(FString::Printf(TEXT("AlertsPropagationFinished_%s"), *InHierarchyName.ToString()));
+	AllAlertPropagationStartConditionNames.Add(Conditions->AlertPropagationStartConditionName);
+
+	RegisterOnAddHierarchyQueries(DataStorage, HierarchyHandle, Conditions->AlertPropagationStartConditionName);
+	RegisterOnRemoveHierarchyQueries(DataStorage, HierarchyHandle, Conditions->AlertPropagationStartConditionName);
+	RegisterParentUpdatesHierarchyQueries(DataStorage, HierarchyHandle, Conditions->AlertPropagationStartConditionName);
+	RegisterChildAlertHierarchyUpdatesQueries(DataStorage, InHierarchyName, HierarchyHandle, Conditions.Get());
+
+	HierarchyConditions.Add(InHierarchyName, MoveTemp(Conditions));
 }
 
 void UTedsAlertsFactory::RegisterSubQueries(UE::Editor::DataStorage::ICoreProvider& DataStorage)
@@ -108,13 +113,13 @@ void UTedsAlertsFactory::RegisterSortUnsortedAlertsQuery(UE::Editor::DataStorage
 
 	DataStorage.RegisterQuery(
 		Select("Sort unsorted alerts", FProcessor(EQueryTickPhase::PrePhysics, DataStorage.GetQueryTickGroupName(EQueryTickGroups::Update)),
-			[](IQueryContext& NewAlertContext, RowHandle NewAlertRow, FTedsAlertColumn& NewAlert)
+			[AllConditions = &AllAlertPropagationStartConditionNames](IQueryContext& NewAlertContext, RowHandle NewAlertRow, FTedsAlertColumn& NewAlert)
 			{
 				RowHandle NextRow = NewAlert.NextAlert;
 				while (NewAlertContext.IsRowAssigned(NextRow))
 				{
 					// Walk the chain and add link at the appropriate spot. This can include the row with the active alert.
-					NewAlertContext.RunSubquery(0, NextRow, CreateSubqueryCallbackBinding([NewAlertRow, &NewAlert, &NewAlertContext, &NextRow]
+					NewAlertContext.RunSubquery(0, NextRow, CreateSubqueryCallbackBinding([NewAlertRow, &NewAlert, &NewAlertContext, &NextRow, AllConditions]
 						(ISubqueryContext& Context, RowHandle TargetRow, FTedsAlertColumn& TargetAlert)
 						{
 							if (NewAlert.AlertType > TargetAlert.AlertType ||
@@ -134,12 +139,15 @@ void UTedsAlertsFactory::RegisterSortUnsortedAlertsQuery(UE::Editor::DataStorage
 										// The active alert needs to be replaced so swap the active alert with the new alert.
 										SwapAlerts(Context, TargetRow, TargetAlert, NewAlertContext, NewAlertRow, NewAlert, AlertActionQueryIndex, AlertHierarchyQueryIndex);
 									}
-									// (Re)calculate child alerts.
-									Context.ActivateQueries(AlertConditionName);
+									// (Re)calculate child alerts for all registered hierarchies.
+									for (const FName& ConditionName : *AllConditions)
+									{
+										Context.ActivateQueries(ConditionName);
+									}
 								}
 								else
 								{
-									// Found a spot in the chain. To avoid having to track the previous entry, just swap the 
+									// Found a spot in the chain. To avoid having to track the previous entry, just swap the
 									// alert at that spot with the new one and chain them up.
 									SwapAlerts(Context, TargetRow, TargetAlert, NewAlertContext, NewAlertRow, NewAlert, AlertActionQueryIndex, AlertHierarchyQueryIndex);
 								}
@@ -174,10 +182,11 @@ void UTedsAlertsFactory::RegisterSortUnsortedAlertsQuery(UE::Editor::DataStorage
 void UTedsAlertsFactory::RegisterOnRemoveQueries(UE::Editor::DataStorage::ICoreProvider& DataStorage)
 {
 	using namespace UE::Editor::DataStorage::Queries;
-	
+
+
 	DataStorage.RegisterQuery(
 		Select("Remove active alert", FObserver::OnRemove<FTedsAlertColumn>(),
-			[](IQueryContext& Context, RowHandle Row, const FTedsAlertColumn& Alert)
+			[AllConditions = &AllAlertPropagationStartConditionNames](IQueryContext& Context, RowHandle Row, const FTedsAlertColumn& Alert)
 			{
 				// Delete all entries in the alert chain.
 				RowHandle NextRow = Alert.NextAlert;
@@ -190,19 +199,26 @@ void UTedsAlertsFactory::RegisterOnRemoveQueries(UE::Editor::DataStorage::ICoreP
 
 
 				// Remove any alerts that haven't been processed yet.
+				// Also check NextAlert == Row to avoid removing unsorted entries that share the same
+				// name but target a different row (e.g. when multiple rows use the same alert name).
 				Context.RunSubquery(1, CreateSubqueryCallbackBinding(
-					[&Alert](ISubqueryContext& Context, RowHandle Row, FTedsAlertColumn& PendingAlert)
+					[&Alert, Row](ISubqueryContext& Context, RowHandle PendingAlertRow, FTedsAlertColumn& PendingAlert)
 					{
-						if (Alert.Name == PendingAlert.Name)
+						if (Alert.Name == PendingAlert.Name && PendingAlert.NextAlert == Row)
 						{
-							Context.RemoveRow(Row);
+							Context.RemoveRow(PendingAlertRow);
 						}
 					}));
 
-				Context.RemoveColumns<FTedsAlertHierarchyColumn, FTedsAlertActionColumn>(Row);
+				// FTedsChildAlertColumn is added to leaf alert rows by "Add missing child alerts";
+				// clean it up here alongside the other associated columns.
+				Context.RemoveColumns<FTedsAlertHierarchyColumn, FTedsAlertActionColumn, FTedsChildAlertColumn>(Row);
 
-				// Update any alert parents.
-				Context.ActivateQueries(AlertConditionName);
+				// Update any alert parents in all registered hierarchies.
+				for (const FName& ConditionName : *AllConditions)
+				{
+					Context.ActivateQueries(ConditionName);
+				}
 			})
 		.Where()
 			.None<FTedsAlertChainTag>()
@@ -212,55 +228,19 @@ void UTedsAlertsFactory::RegisterOnRemoveQueries(UE::Editor::DataStorage::ICoreP
 		.Compile());
 }
 
-void UTedsAlertsFactory::RegisterChildAlertUpdatesQueries(UE::Editor::DataStorage::ICoreProvider& DataStorage)
-{
-	using namespace UE::Editor::DataStorage::Queries;
-
-	DataStorage.RegisterQuery(
-		Select(
-			"Clear child alerts",
-			FProcessor(EQueryTickPhase::PostPhysics, DataStorage.GetQueryTickGroupName(EQueryTickGroups::PreUpdate))
-				.MakeActivatable(AlertConditionName),
-			[](IQueryContext& Context, RowHandle Row, FTedsChildAlertColumn& ChildAlert)
-			{
-				ResetChildAlertCounters(ChildAlert);
-			})
-		.Where()
-			.None<FTedsAlertChainTag>()
-		.Compile());
-
-	DataStorage.RegisterQuery(
-		Select(
-			"Remove unused child alerts",
-			FProcessor(EQueryTickPhase::PostPhysics, DataStorage.GetQueryTickGroupName(EQueryTickGroups::PostUpdate))
-				.MakeActivatable(AlertConditionName),
-			[](IQueryContext& Context, RowHandle Row, FTedsChildAlertColumn& ChildAlert)
-			{
-				for (size_t It = 0; It < static_cast<size_t>(FTedsAlertColumnType::MAX); ++It)
-				{
-					if (ChildAlert.Counts[It] != 0)
-					{
-						return;
-					}
-				}
-				Context.RemoveColumns<FTedsChildAlertColumn, FTedsAlertHierarchyColumn>(Row);
-				Context.AddColumns<FTypedElementSyncBackToWorldTag>(Row);
-			})
-		.Compile());
-}
 
 
-void UTedsAlertsFactory::RegisterOnAddHierarchyQueries(UE::Editor::DataStorage::ICoreProvider& DataStorage, UE::Editor::DataStorage::FHierarchyHandle InHierarchyHandle)
+void UTedsAlertsFactory::RegisterOnAddHierarchyQueries(UE::Editor::DataStorage::ICoreProvider& DataStorage, UE::Editor::DataStorage::FHierarchyHandle InHierarchyHandle, const FName& InAlertPropagationStartConditionName)
 {
 	using namespace UE::Editor::DataStorage::Queries;
 
 	DataStorage.RegisterQuery(
 		Select("Register alert with parent on alert hierarchy add", FObserver::OnAdd<FTedsAlertHierarchyColumn>(),
-			[InHierarchyHandle](IQueryContext& Context, RowHandle Row, const FTedsAlertHierarchyColumn& AlertHierarchyColumn)
+			[InHierarchyHandle, InAlertPropagationStartConditionName](IQueryContext& Context, RowHandle Row, const FTedsAlertHierarchyColumn& AlertHierarchyColumn)
 			{
 				if (AlertHierarchyColumn.HierarchyHandle == InHierarchyHandle)
 				{
-					Context.ActivateQueries(AlertConditionName);
+					Context.ActivateQueries(InAlertPropagationStartConditionName);
 				}
 			})
 		.Where()
@@ -271,11 +251,11 @@ void UTedsAlertsFactory::RegisterOnAddHierarchyQueries(UE::Editor::DataStorage::
 	DataStorage.RegisterQuery(
 		Select("Register alert with parent on parent add (FHierarchyHandle)",
 			FObserver(FObserver::EEvent::Add, DataStorage.GetChildTagType(InHierarchyHandle)),
-			[InHierarchyHandle](IQueryContext& Context, RowHandle Row, const FTedsAlertHierarchyColumn& AlertHierarchyColumn)
+			[InHierarchyHandle, InAlertPropagationStartConditionName](IQueryContext& Context, RowHandle Row, const FTedsAlertHierarchyColumn& AlertHierarchyColumn)
 			{
 				if (AlertHierarchyColumn.HierarchyHandle == InHierarchyHandle)
 				{
-					Context.ActivateQueries(AlertConditionName);
+					Context.ActivateQueries(InAlertPropagationStartConditionName);
 				}
 			})
 		.Where()
@@ -285,17 +265,17 @@ void UTedsAlertsFactory::RegisterOnAddHierarchyQueries(UE::Editor::DataStorage::
 }
 
 void UTedsAlertsFactory::RegisterOnRemoveHierarchyQueries(UE::Editor::DataStorage::ICoreProvider& DataStorage,
-	UE::Editor::DataStorage::FHierarchyHandle InHierarchyHandle)
+	UE::Editor::DataStorage::FHierarchyHandle InHierarchyHandle, const FName& InAlertPropagationStartConditionName)
 {
 	using namespace UE::Editor::DataStorage::Queries;
 
 	DataStorage.RegisterQuery(
 		Select("Update alert upon parent removal", FObserver(FObserver::EEvent::Remove, DataStorage.GetChildTagType(InHierarchyHandle)),
-			[InHierarchyHandle](IQueryContext& Context, RowHandle Row, const FTedsAlertHierarchyColumn& AlertHierarchyColumn)
+			[InHierarchyHandle, InAlertPropagationStartConditionName](IQueryContext& Context, RowHandle Row, const FTedsAlertHierarchyColumn& AlertHierarchyColumn)
 			{
 				if (AlertHierarchyColumn.HierarchyHandle == InHierarchyHandle)
 				{
-					Context.ActivateQueries(AlertConditionName);
+					Context.ActivateQueries(InAlertPropagationStartConditionName);
 				}
 			})
 		.Where()
@@ -304,21 +284,21 @@ void UTedsAlertsFactory::RegisterOnRemoveHierarchyQueries(UE::Editor::DataStorag
 		.Compile());
 }
 
-void UTedsAlertsFactory::RegisterParentUpdatesHierarchyQueries(UE::Editor::DataStorage::ICoreProvider& DataStorage, UE::Editor::DataStorage::FHierarchyHandle InHierarchyHandle)
+void UTedsAlertsFactory::RegisterParentUpdatesHierarchyQueries(UE::Editor::DataStorage::ICoreProvider& DataStorage, UE::Editor::DataStorage::FHierarchyHandle InHierarchyHandle, const FName& InAlertPropagationStartConditionName)
 {
 	using namespace UE::Editor::DataStorage::Queries;
-	
+
 	if (const UScriptStruct* ParentChangedColumnType = DataStorage.GetParentChangedColumnType(InHierarchyHandle))
 	{
 		DataStorage.RegisterQuery(
 			Select(
-			"Trigger alert update if alert or child alert's parent changed", 
+			"Trigger alert update if alert or child alert's parent changed",
 			FProcessor(EQueryTickPhase::PostPhysics, DataStorage.GetQueryTickGroupName(EQueryTickGroups::Default)),
-			[InHierarchyHandle](IQueryContext& Context, RowHandle Row, const FTedsAlertHierarchyColumn& AlertHierarchyColumn)
+			[InHierarchyHandle, InAlertPropagationStartConditionName](IQueryContext& Context, RowHandle Row, const FTedsAlertHierarchyColumn& AlertHierarchyColumn)
 			{
 				if (AlertHierarchyColumn.HierarchyHandle == InHierarchyHandle)
 				{
-					Context.ActivateQueries(AlertConditionName);
+					Context.ActivateQueries(InAlertPropagationStartConditionName);
 				}
 			})
 		.Where()
@@ -329,19 +309,32 @@ void UTedsAlertsFactory::RegisterParentUpdatesHierarchyQueries(UE::Editor::DataS
 	}
 }
 
-void UTedsAlertsFactory::RegisterChildAlertHierarchyUpdatesQueries(UE::Editor::DataStorage::ICoreProvider& DataStorage, const FName& InHierarchyName, UE::Editor::DataStorage::FHierarchyHandle InHierarchyHandle)
+void UTedsAlertsFactory::RegisterChildAlertHierarchyUpdatesQueries(UE::Editor::DataStorage::ICoreProvider& DataStorage, const FName& InHierarchyName, UE::Editor::DataStorage::FHierarchyHandle InHierarchyHandle, FAlertHierarchyConditions* InConditions)
 {
 	using namespace UE::Editor::DataStorage::Queries;
-    
+
+	const FName AlertPropagationStartConditionName = InConditions->AlertPropagationStartConditionName;
+	const FName AlertPropagationConditionName = InConditions->AlertPropagationConditionName;
+	const FName AlertPropagationFinishedConditionName = InConditions->AlertPropagationFinishedConditionName;
+
+	// Ensure all ancestor rows and the leaf alert row itself have the FTedsChildAlertColumn before counting begins.
 	DataStorage.RegisterQuery(
 		Select(
 			"Add missing child alerts",
 			FPhaseAmble(FPhaseAmble::ELocation::Preamble, EQueryTickPhase::PostPhysics)
-				.MakeActivatable(AlertConditionName),
+				.MakeActivatable(AlertPropagationStartConditionName),
 			[InHierarchyHandle](IQueryContext& Context, RowHandle Row, const FTedsAlertHierarchyColumn& AlertHierarchyColumn)
 			{
 				if (AlertHierarchyColumn.HierarchyHandle == InHierarchyHandle)
 				{
+					// Add FTedsChildAlertColumn to this leaf alert row so it participates in the per-frame propagation pass alongside intermediate nodes.
+					if (!Context.HasColumn<FTedsChildAlertColumn>(Row))
+					{
+						FTedsChildAlertColumn LeafChildAlert;
+						ResetChildAlertCounters(LeafChildAlert);
+						Context.AddColumn(Row, MoveTemp(LeafChildAlert));
+					}
+
 					RowHandle ParentRow = Context.GetParentRow(Row);
 					if (Context.IsRowAssigned(ParentRow))
 					{
@@ -351,27 +344,181 @@ void UTedsAlertsFactory::RegisterChildAlertHierarchyUpdatesQueries(UE::Editor::D
 			})
 		.AccessesHierarchy(InHierarchyName)
 		.Where()
+			.All<FTedsAlertColumn>()
 			.None<FTedsAlertChainTag>()
 		.Compile());
 
+	// Count every node that will propagate.
+	// For each node, record it in PendingPropagations and increment its parent's NumAlertChildren so the parent knows how many to wait for.
+	// Activate the per-frame propagation condition for the next frame.
 	DataStorage.RegisterQuery(
 		Select(
-			"Increment child alerts",
+			"Count alert children",
 			FProcessor(EQueryTickPhase::PostPhysics, DataStorage.GetQueryTickGroupName(EQueryTickGroups::Update))
-				.MakeActivatable(AlertConditionName),
-			[InHierarchyHandle](IQueryContext& Context, RowHandle Row, FTedsAlertColumn& Alert, const FTedsAlertHierarchyColumn& AlertHierarchyColumn)
+				.MakeActivatable(AlertPropagationStartConditionName),
+			[InHierarchyHandle, PendingPropagations = &InConditions->PendingPropagations, AlertPropagationConditionName](
+				IQueryContext& Context, RowHandle Row, const FTedsAlertHierarchyColumn& AlertHierarchyColumn)
+			{
+				if (AlertHierarchyColumn.HierarchyHandle != InHierarchyHandle)
+				{
+					return;
+				}
+
+				PendingPropagations->fetch_add(1);
+
+				RowHandle ParentRow = Context.GetParentRow(Row);
+				if (Context.IsRowAssigned(ParentRow))
+				{
+					Context.RunSubquery(0, ParentRow, CreateSubqueryCallbackBinding(
+						[](ISubqueryContext&, RowHandle, FTedsChildAlertColumn& ParentChildAlert)
+						{
+							++ParentChildAlert.NumAlertChildren;
+						}));
+				}
+
+				// Schedule the propagation wave for next frame.
+				Context.ActivateQueries(AlertPropagationConditionName);
+			})
+		.AccessesHierarchy(InHierarchyName)
+		.Where()
+			.Any<FTedsAlertColumn, FTedsChildAlertColumn>()
+			.None<FTedsAlertChainTag>()
+		.DependsOn()
+			.SubQuery(ChildAlertColumnReadWriteQuery)
+		.Compile());
+
+	// Reset the per-hierarchy propagation counter once at the start of each AlertPropagationStartConditionName wave,
+	// before "Count alert children" starts incrementing it.
+	DataStorage.RegisterQuery(
+		Select(
+			"Reset pending propagations",
+			FPhaseAmble(FPhaseAmble::ELocation::Preamble, EQueryTickPhase::PostPhysics)
+				.MakeActivatable(AlertPropagationStartConditionName),
+			[PendingPropagations = &InConditions->PendingPropagations](IQueryContext&, const RowHandle*)
+			{
+				PendingPropagations->store(0);
+			})
+		.Where()
+			.All<FTedsChildAlertColumn>()
+			.None<FTedsAlertChainTag>()
+		.Compile());
+
+	DataStorage.RegisterQuery(
+		Select(
+			"Clear child alerts",
+			FProcessor(EQueryTickPhase::PostPhysics, DataStorage.GetQueryTickGroupName(EQueryTickGroups::PreUpdate))
+				.MakeActivatable(AlertPropagationStartConditionName),
+			[InHierarchyHandle](IQueryContext&, RowHandle, FTedsChildAlertColumn& ChildAlert, const FTedsAlertHierarchyColumn& AlertHierarchyColumn)
 			{
 				if (AlertHierarchyColumn.HierarchyHandle == InHierarchyHandle)
 				{
-					IncrementParents(Context, Context.GetParentRow(Row), Alert.AlertType, 0);
+					ResetChildAlertCounters(ChildAlert);
+				}
+			})
+		.Where()
+			.None<FTedsAlertChainTag>()
+		.Compile());
+
+	DataStorage.RegisterQuery(
+		Select(
+			"Remove unused child alerts",
+			FProcessor(EQueryTickPhase::PostPhysics, DataStorage.GetQueryTickGroupName(EQueryTickGroups::PostUpdate))
+				.MakeActivatable(AlertPropagationFinishedConditionName),
+			[InHierarchyHandle](IQueryContext& Context, RowHandle Row, FTedsChildAlertColumn& ChildAlert, const FTedsAlertHierarchyColumn& AlertHierarchyColumn)
+			{
+				if (AlertHierarchyColumn.HierarchyHandle != InHierarchyHandle)
+				{
+					return;
 				}
+				for (int32 It = 0; It < static_cast<int32>(FTedsAlertColumnType::MAX); ++It)
+				{
+					if (ChildAlert.Counts[It] != 0)
+					{
+						return;
+					}
+				}
+				Context.RemoveColumns<FTedsChildAlertColumn, FTedsAlertHierarchyColumn>(Row);
+				Context.AddColumns<FTypedElementSyncBackToWorldTag>(Row);
+			})
+		.Where()
+			.None<FTedsAlertChainTag, FTedsAlertColumn>()
+		.Compile());
+
+	// For each node whose direct contributing children have all reported in, push its accumulated counts plus its own alert, if any to its parent,
+	// then decrement the per-hierarchy counter.
+	DataStorage.RegisterQuery(
+		Select(
+			"Propagate child alerts",
+			FProcessor(EQueryTickPhase::PostPhysics, DataStorage.GetQueryTickGroupName(EQueryTickGroups::PostUpdate))
+				.MakeActivatable(AlertPropagationConditionName),
+			[InHierarchyHandle, PendingPropagations = &InConditions->PendingPropagations](
+				IQueryContext& Context, RowHandle Row, FTedsChildAlertColumn& ChildAlert,
+				const FTedsAlertHierarchyColumn& AlertHierarchyColumn)
+			{
+				if (AlertHierarchyColumn.HierarchyHandle != InHierarchyHandle)
+				{
+					return;
+				}
+
+				// Not ready yet: still waiting for some children to report in.
+				if (ChildAlert.ProcessedChildCount != ChildAlert.NumAlertChildren)
+				{
+					return;
+				}
+
+				if (const RowHandle ParentRow = Context.GetParentRow(Row); Context.IsRowAssigned(ParentRow))
+				{
+					Context.RunSubquery(0, ParentRow, CreateSubqueryCallbackBinding(
+						[&ChildAlert, ParentRow, &Context](ISubqueryContext& SubCtx, RowHandle, FTedsChildAlertColumn& ParentChildAlert)
+						{
+							// Only propagate child count to parent if row has an Alert.
+							if (const FTedsAlertColumn* OwnAlert = Context.GetColumn<FTedsAlertColumn>())
+							{	
+								for (int32 AlertType = 0; AlertType < static_cast<int32>(FTedsAlertColumnType::MAX); ++AlertType)
+								{
+									ParentChildAlert.Counts[AlertType] += ChildAlert.Counts[AlertType];
+								}
+								++ParentChildAlert.Counts[static_cast<int32>(OwnAlert->AlertType)];
+							}
+							++ParentChildAlert.ProcessedChildCount;
+							SubCtx.AddColumns<FTypedElementSyncBackToWorldTag>(ParentRow);
+						}));
+				}
+				Context.AddColumns<FTypedElementSyncBackToWorldTag>(Row);
+
+				PendingPropagations->fetch_sub(1);
+
+				// Sentinel: move ProcessedChildCount past NumAlertChildren so this node does not match the condition again in subsequent frames.
+				++ChildAlert.ProcessedChildCount;
 			})
 		.AccessesHierarchy(InHierarchyName)
+		.ReadOnly<FTedsAlertColumn>(EOptional::Yes)
 		.Where()
 			.None<FTedsAlertChainTag>()
 		.DependsOn()
 			.SubQuery(ChildAlertColumnReadWriteQuery)
 		.Compile());
+
+	// If there are still nodes waiting to propagate, re-activate for the next frame.
+	DataStorage.RegisterQuery(
+		Select(
+			"Check alert propagation complete",
+			FPhaseAmble(FPhaseAmble::ELocation::Postamble, EQueryTickPhase::PostPhysics)
+				.MakeActivatable(AlertPropagationConditionName),
+			[PendingPropagations = &InConditions->PendingPropagations, AlertPropagationConditionName, AlertPropagationFinishedConditionName](IQueryContext& Context, const RowHandle*)
+			{
+				if (PendingPropagations->load() > 0)
+				{
+					Context.ActivateQueries(AlertPropagationConditionName);
+				}
+				else
+				{
+					Context.ActivateQueries(AlertPropagationFinishedConditionName);
+				}
+			})
+		.Where()
+			.All<FTedsChildAlertColumn>()
+		.Compile());
 }
 
 void UTedsAlertsFactory::AssignAlert(
@@ -544,53 +691,20 @@ void UTedsAlertsFactory::AddChildAlertsToHierarchy(
 {
 	using namespace UE::Editor::DataStorage;
 
-	bool bHasParent = true;
-	do
+	// Only add to the direct parent. The gray warning indicator is intentionally limited to
+	// the immediate parent of an alert. Grandparents and higher do not receive the column.
+	if (!Context.HasColumn<FTedsChildAlertColumn>(Parent))
 	{
-		RowHandle NextParent = Parent;
-		bHasParent = MoveToNextParent(NextParent, Context);
-
-		// Check if a child alert column exists and add one if not.
-		if (!Context.HasColumn<FTedsChildAlertColumn>(Parent))
-		{
-			FTedsChildAlertColumn ChildAlert;
-			ResetChildAlertCounters(ChildAlert);
-			Context.AddColumn(Parent, MoveTemp(ChildAlert));
+		FTedsChildAlertColumn ChildAlert;
+		ResetChildAlertCounters(ChildAlert);
+		Context.AddColumn(Parent, MoveTemp(ChildAlert));
 
-			// Child Alerts also need the hierarchy column to make sure they can walk up the parent chain
-			Context.AddColumn(Parent, FTedsAlertHierarchyColumn{.HierarchyHandle = HierarchyHandle});
-		}
-		Parent = NextParent;
-	} while (bHasParent);
-}
-
-void UTedsAlertsFactory::IncrementParents(
-	UE::Editor::DataStorage::IQueryContext& Context, UE::Editor::DataStorage::RowHandle Row, FTedsAlertColumnType AlertType,
-	int32 ChildAlertQueryIndex)
-{
-	using namespace UE::Editor::DataStorage::Queries;
-
-	while (Context.IsRowAvailable(Row))
-	{
-		FQueryResult Result = Context.RunSubquery(ChildAlertQueryIndex, Row, CreateSubqueryCallbackBinding(
-			[&ParentRow = Row, AlertType](
-				ISubqueryContext& Context, RowHandle Row, FTedsChildAlertColumn& ChildAlert)
-			{
-				ChildAlert.Counts[static_cast<size_t>(AlertType)]++;
-				Context.AddColumns<FTypedElementSyncBackToWorldTag>(Row);
-				ParentRow = Context.GetParentRow(ParentRow);
-			}));
-		checkf(Result.Count > 0, TEXT("Expected to be able to setup the child alert, but it was missing on the parent column."));
+		// Child alerts also need the hierarchy column to walk up during propagation.
+		Context.AddColumn(Parent, FTedsAlertHierarchyColumn{ .HierarchyHandle = HierarchyHandle });
 	}
 }
 
-bool UTedsAlertsFactory::MoveToNextParent(
-	UE::Editor::DataStorage::RowHandle& Parent, UE::Editor::DataStorage::IQueryContext& Context)
-{
-	using namespace UE::Editor::DataStorage::Queries;
-	Parent = Context.GetParentRow(Parent);
-	return Context.IsRowAvailable(Parent);
-}
+
 
 void UTedsAlertsFactory::ResetChildAlertCounters(FTedsChildAlertColumn& ChildAlert)
 {
@@ -598,4 +712,6 @@ void UTedsAlertsFactory::ResetChildAlertCounters(FTedsChildAlertColumn& ChildAle
 	{
 		ChildAlert.Counts[It] = 0;
 	}
+	ChildAlert.ProcessedChildCount = 0;
+	ChildAlert.NumAlertChildren = 0;
 }
\ No newline at end of file
diff --git a/Engine/Plugins/Experimental/EditorDataStorageFeatures/Source/TedsAlerts/Private/TedsAlertsFactory.h b/Engine/Plugins/Experimental/EditorDataStorageFeatures/Source/TedsAlerts/Private/TedsAlertsFactory.h
--- a/Engine/Plugins/Experimental/EditorDataStorageFeatures/Source/TedsAlerts/Private/TedsAlertsFactory.h
+++ b/Engine/Plugins/Experimental/EditorDataStorageFeatures/Source/TedsAlerts/Private/TedsAlertsFactory.h
@@ -2,6 +2,8 @@
 
 #pragma once
 
+#include <atomic>
+
 #include "DataStorage/Handles.h"
 #include "Elements/Interfaces/TypedElementDataStorageFactory.h"
 #include "TedsAlertColumns.h"
@@ -37,31 +39,42 @@ public:
 	void RegisterAlertHierarchy(UE::Editor::DataStorage::ICoreProvider& DataStorage, const FName& InHierarchyName);
 
 private:
-	static const FName AlertConditionName;
+	// Per-hierarchy condition names and propagation counter.
+	// Stored as a unique_ptr so the address of PendingPropagations is stable after the map grows.
+	struct FAlertHierarchyConditions
+	{
+		FName AlertPropagationStartConditionName;
+		FName AlertPropagationConditionName;
+		FName AlertPropagationFinishedConditionName;
+		// Counts how many nodes still need to propagate their alert counts to their parent.
+		// Decremented by "Propagate child alerts"; when it hits 0 propagation is complete for this wave.
+		std::atomic<int32> PendingPropagations{0};
+	};
+	TMap<FName, TUniquePtr<FAlertHierarchyConditions>> HierarchyConditions;
+
+	// All per-hierarchy AlertPropagationStartConditionNames, captured by non-hierarchy queries so they can
+	// activate every registered hierarchy when a global alert change occurs.
+	TArray<FName> AllAlertPropagationStartConditionNames;
 
 	UE::Editor::DataStorage::TableHandle ChainTable = UE::Editor::DataStorage::InvalidTableHandle;
 	UE::Editor::DataStorage::TableHandle UnsortedChainTable = UE::Editor::DataStorage::InvalidTableHandle;
-	
+
 	UE::Editor::DataStorage::QueryHandle SortedAlertsQuery = UE::Editor::DataStorage::InvalidQueryHandle;
 	UE::Editor::DataStorage::QueryHandle UnsortedAlertsQuery = UE::Editor::DataStorage::InvalidQueryHandle;
 	UE::Editor::DataStorage::QueryHandle AlertActionQuery = UE::Editor::DataStorage::InvalidQueryHandle;
 	UE::Editor::DataStorage::QueryHandle AlertHierarchyQuery = UE::Editor::DataStorage::InvalidQueryHandle;
 	UE::Editor::DataStorage::QueryHandle ChildAlertColumnReadWriteQuery = UE::Editor::DataStorage::InvalidQueryHandle;
 
-	// Keep track of hierarchies we have already registered observers for to avoid duplicates
-	TSet<FName> RegisteredHierarchies;
-
 	void RegisterSubQueries(UE::Editor::DataStorage::ICoreProvider& DataStorage);
 	void RegisterSortUnsortedAlertsQuery(UE::Editor::DataStorage::ICoreProvider& DataStorage);
 	void RegisterOnRemoveQueries(UE::Editor::DataStorage::ICoreProvider& DataStorage);
-	void RegisterChildAlertUpdatesQueries(UE::Editor::DataStorage::ICoreProvider& DataStorage);
 
 	// Hierarchy update queries, these currently have to be registered per hierarchy since a query can only access one hierarchy and it isn't possible
 	// to right a query that can dynamically pick which hierarchy to use.
-	void RegisterOnAddHierarchyQueries(UE::Editor::DataStorage::ICoreProvider& DataStorage, UE::Editor::DataStorage::FHierarchyHandle InHierarchyHandle);
-	void RegisterOnRemoveHierarchyQueries(UE::Editor::DataStorage::ICoreProvider& DataStorage, UE::Editor::DataStorage::FHierarchyHandle InHierarchyHandle);
-	void RegisterParentUpdatesHierarchyQueries(UE::Editor::DataStorage::ICoreProvider& DataStorage, UE::Editor::DataStorage::FHierarchyHandle InHierarchyHandle);
-	void RegisterChildAlertHierarchyUpdatesQueries(UE::Editor::DataStorage::ICoreProvider& DataStorage, const FName& InHierarchyName, UE::Editor::DataStorage::FHierarchyHandle InHierarchyHandle);
+	void RegisterOnAddHierarchyQueries(UE::Editor::DataStorage::ICoreProvider& DataStorage, UE::Editor::DataStorage::FHierarchyHandle InHierarchyHandle, const FName& InAlertPropagationStartConditionName);
+	void RegisterOnRemoveHierarchyQueries(UE::Editor::DataStorage::ICoreProvider& DataStorage, UE::Editor::DataStorage::FHierarchyHandle InHierarchyHandle, const FName& InAlertPropagationStartConditionName);
+	void RegisterParentUpdatesHierarchyQueries(UE::Editor::DataStorage::ICoreProvider& DataStorage, UE::Editor::DataStorage::FHierarchyHandle InHierarchyHandle, const FName& InAlertPropagationStartConditionName);
+	void RegisterChildAlertHierarchyUpdatesQueries(UE::Editor::DataStorage::ICoreProvider& DataStorage, const FName& InHierarchyName, UE::Editor::DataStorage::FHierarchyHandle InHierarchyHandle, FAlertHierarchyConditions* InConditions);
 
 	static void AssignAlert(
 		UE::Editor::DataStorage::ISubqueryContext& TargetContext, UE::Editor::DataStorage::RowHandle TargetRow, FTedsAlertColumn& Target,
@@ -83,10 +96,5 @@ private:
 
 	static void AddChildAlertsToHierarchy(
 		UE::Editor::DataStorage::IQueryContext& Context, UE::Editor::DataStorage::RowHandle Parent, UE::Editor::DataStorage::FHierarchyHandle HierarchyHandle);
-	static void IncrementParents(
-		UE::Editor::DataStorage::IQueryContext& Context, UE::Editor::DataStorage::RowHandle Row, FTedsAlertColumnType AlertType,
-		int32 ChildAlertQueryIndex);
-	static bool MoveToNextParent(
-		UE::Editor::DataStorage::RowHandle& Parent, UE::Editor::DataStorage::IQueryContext& Context);
 	static void ResetChildAlertCounters(FTedsChildAlertColumn& ChildAlert);
 };
diff --git a/Engine/Plugins/Experimental/EditorDataStorageFeatures/Source/TedsAlerts/Public/TedsAlertColumns.h b/Engine/Plugins/Experimental/EditorDataStorageFeatures/Source/TedsAlerts/Public/TedsAlertColumns.h
--- a/Engine/Plugins/Experimental/EditorDataStorageFeatures/Source/TedsAlerts/Public/TedsAlertColumns.h
+++ b/Engine/Plugins/Experimental/EditorDataStorageFeatures/Source/TedsAlerts/Public/TedsAlertColumns.h
@@ -60,12 +60,18 @@ USTRUCT(meta = (DisplayName = "Child alert"))
 struct FTedsChildAlertColumn final : public FEditorDataStorageColumn
 {
 	GENERATED_BODY()
-	
+
 	uint16 Counts[static_cast<size_t>(FTedsAlertColumnType::MAX)];
+
+	/** How many direct contributing children (leaf alerts or intermediate nodes) have propagated to this row. */
+	uint16 ProcessedChildCount = 0;
+
+	/** Expected number of direct contributing children. Set at reset time by the "Count alert children" processor. */
+	uint16 NumAlertChildren = 0;
 };
 
 /**
- * Column that holdes info about the hierarchy an alert is participating in
+ * Column that holds info about the hierarchy an alert is participating in
  */
 USTRUCT(meta = (DisplayName = "Alert Hierarchy Column"))
 struct FTedsAlertHierarchyColumn final : public FEditorDataStorageColumn
