[TEDS Outliner] Add Show Settings to the TEDS Outliner

#Synced-CL 52352994

diff --git a/Engine/Plugins/Experimental/EditorDataStorageFeatures/Source/TedsActorCompatibility/Private/Compatibility/TedsActorComponentFactory.cpp b/Engine/Plugins/Experimental/EditorDataStorageFeatures/Source/TedsActorCompatibility/Private/Compatibility/TedsActorComponentFactory.cpp
--- a/Engine/Plugins/Experimental/EditorDataStorageFeatures/Source/TedsActorCompatibility/Private/Compatibility/TedsActorComponentFactory.cpp
+++ b/Engine/Plugins/Experimental/EditorDataStorageFeatures/Source/TedsActorCompatibility/Private/Compatibility/TedsActorComponentFactory.cpp
@@ -10,6 +10,7 @@
 #include "Elements/Columns/TypedElementTypeInfoColumns.h"
 #include "Elements/Framework/TypedElementQueryBuilder.h"
 #include "Elements/Interfaces/TypedElementDataStorageCompatibilityInterface.h"
+#include "Engine/Level.h"
 #include "Factories/TedsEditorHierarchyFactory.h"
 #include "GameFramework/Actor.h"
 
@@ -28,7 +29,7 @@ void UTedsActorComponentFactory::RegisterTables(UE::Editor::DataStorage::ICorePr
 	
 	if (bEnableActorComponentCompatibility)
 	{
-		TableHandle ActorComponentTable = DataStorage.RegisterTable<FActorComponentTypeTag, FTypedElementLabelColumn, FTypedElementLabelHashColumn>(
+		TableHandle ActorComponentTable = DataStorage.RegisterTable<FActorComponentTypeTag, FTypedElementLabelColumn, FLevelColumn, FTypedElementLabelHashColumn>(
 											DataStorage.FindTable(FName(TEXT("Editor_StandardUObjectTable"))),
 											FName("Editor_ActorComponentTable"));
 		
@@ -147,7 +148,7 @@ void UTedsActorComponentFactory::RegisterQueries(UE::Editor::DataStorage::ICoreP
 					.All<FActorComponentTypeTag, FTypedElementSyncFromWorldTag>()
 				.Compile());
 
-		// 
+		// Sync Actor Component Label
 		DataStorage.RegisterQuery(
 			Select(
 				TEXT("Sync actorcomponent label to column"),
@@ -170,5 +171,26 @@ void UTedsActorComponentFactory::RegisterQueries(UE::Editor::DataStorage::ICoreP
 			.Where()
 				.All<FActorComponentTypeTag, FTypedElementSyncFromWorldTag>()
 			.Compile());
+
+		// Sync Actor Component ULevel
+		DataStorage.RegisterQuery(
+			Select(
+				TEXT("Sync ActorComponent level to column"),
+				FProcessor(EQueryTickPhase::PrePhysics, DataStorage.GetQueryTickGroupName(EQueryTickGroups::SyncExternalToDataStorage))
+				.SetExecutionMode(EExecutionMode::GameThread),
+				[](const FTypedElementUObjectColumn& Actor, FLevelColumn& LevelColumn)
+				{
+					if (const UActorComponent* ComponentInstance = Cast<UActorComponent>(Actor.Object); ComponentInstance != nullptr)
+					{
+						if (ULevel* Level = ComponentInstance->GetComponentLevel())
+						{
+							LevelColumn.Level = Level;
+						}
+					}
+				}
+			)
+			.Where()
+				.All<FActorComponentTypeTag, FTypedElementSyncFromWorldTag>()
+			.Compile());
 	}
 }
diff --git a/Engine/Plugins/Experimental/EditorDataStorageFeatures/Source/TedsOutliner/Private/TedsOutlinerMode.cpp b/Engine/Plugins/Experimental/EditorDataStorageFeatures/Source/TedsOutliner/Private/TedsOutlinerMode.cpp
--- a/Engine/Plugins/Experimental/EditorDataStorageFeatures/Source/TedsOutliner/Private/TedsOutlinerMode.cpp
+++ b/Engine/Plugins/Experimental/EditorDataStorageFeatures/Source/TedsOutliner/Private/TedsOutlinerMode.cpp
@@ -15,6 +15,7 @@
 #include "ToolMenus.h"
 #include "DragAndDrop/Widgets/WidgetDropHandler.h"
 #include "Elements/Columns/TypedElementHandleColumn.h"
+#include "SceneOutlinerFilters.h"
 
 #include UE_INLINE_GENERATED_CPP_BY_NAME(TedsOutlinerMode)
 
@@ -49,6 +50,25 @@ FTedsOutlinerMode::FTedsOutlinerMode(const FTedsOutlinerParams& InParams)
 	TedsOutlinerImpl->OnSelectionChanged().AddRaw(this, &FTedsOutlinerMode::OnSelectionChanged);
 
 	bShowViewButton = InParams.bShowViewButton;
+	// Convert ShowOptionsFilters to FilterInfoMap to appear in the Settings "Show" options menu
+	for (const TSharedPtr<FTedsOutlinerFilter>& ShowFilter : InParams.ShowOptionsFilters)
+	{
+		if (ShowFilter)
+		{
+			ShowFilter->SetSceneOutlinerImpl(TedsOutlinerImpl);
+			FSceneOutlinerFilterInfo SceneOutlinerFilterInfo(ShowFilter->GetDisplayName(),
+				ShowFilter->GetToolTipText(), ShowFilter->IsActive(),
+				FCreateSceneOutlinerFilter::CreateLambda([]()
+				{
+					return MakeShared<TSceneOutlinerFilter<const ISceneOutlinerTreeItem>>(FSceneOutlinerFilter::EDefaultBehaviour::Pass);
+				}));
+			SceneOutlinerFilterInfo.OnToggle().AddLambda([ShowFilter](bool bIsActive)
+				{
+					ShowFilter->ActiveStateChanged(bIsActive);
+				});
+			FilterInfoMap.Add(FName(ShowFilter->GetName()), SceneOutlinerFilterInfo);
+		}
+	}
 	
 	TedsOutlinerImpl->IsItemCompatible().BindLambda([](const ISceneOutlinerTreeItem& Item)
 	{
@@ -303,6 +323,14 @@ bool FTedsOutlinerMode::ShowViewButton() const
 {
 	return bShowViewButton;
 }
+
+bool FTedsOutlinerMode::ShowFilterOptions() const
+{
+	// Use the same bool as ShowViewButton so we can have any show filters appear in the settings menu, it is better
+	// to give an empty ShowOptionsFilters in the params if you do not want to show any settings.
+	return bShowViewButton;
+}
+	
 } // namespace UE::Editor::Outliner
 
 #undef LOCTEXT_NAMESPACE
diff --git a/Engine/Plugins/Experimental/EditorDataStorageFeatures/Source/TedsOutliner/Private/Widgets/SceneOutlinerWidget.cpp b/Engine/Plugins/Experimental/EditorDataStorageFeatures/Source/TedsOutliner/Private/Widgets/SceneOutlinerWidget.cpp
--- a/Engine/Plugins/Experimental/EditorDataStorageFeatures/Source/TedsOutliner/Private/Widgets/SceneOutlinerWidget.cpp
+++ b/Engine/Plugins/Experimental/EditorDataStorageFeatures/Source/TedsOutliner/Private/Widgets/SceneOutlinerWidget.cpp
@@ -12,11 +12,14 @@
 #include "TedsOutlinerHelpers.h"
 #include "TedsOutlinerModule.h"
 #include "TedsOutlinerImpl.h"
+#include "Columns/TedsActorComponentCompatibilityColumns.h"
 #include "Elements/Columns/TypedElementCompatibilityColumns.h"
 #include "Elements/Columns/TypedElementLabelColumns.h"
 #include "Elements/Columns/TypedElementMiscColumns.h"
+#include "Elements/Columns/TypedElementSelectionColumns.h"
 #include "Elements/Framework/TypedElementSelectionSet.h"
 #include "Factories/TedsEditorHierarchyFactory.h"
+#include "LevelInstance/LevelInstanceInterface.h"
 
 #define LOCTEXT_NAMESPACE "SceneOutlinerWidget"
 
@@ -80,10 +83,19 @@ TSharedPtr<SWidget> FTedsSceneOutlinerWidgetConstructor::CreateWidget(UE::Editor
 	FSceneOutlinerInitializationOptions InitOptions;
 	GetSceneOutlinerInitializationOptions(TabIdentifier, InitOptions);
 
+	TMap<FName, TSharedPtr<FTedsOutlinerFilter>> ShowOptionsFilters;
+	GetShowOptionsFilters(DataStorage, ShowOptionsFilters);
+
 	FTedsOutlinerParams Params(nullptr);
 	Params.QueryDescription = GetQueryDescription();
 	Params.ColumnDescription = GetColumnDescription();
 	GetFilters(DataStorage, Params.Filters);
+	
+	for (const TPair<FName, TSharedPtr<FTedsOutlinerFilter>>& Filter : ShowOptionsFilters)
+	{
+		Params.ShowOptionsFilters.Add(Filter.Value);
+	}
+	
 	Params.bShowRowHandleColumn = ShowRowHandleColumn();
 	Params.bForceShowParents = ForceShowParents();
 	Params.bUseDefaultObservers = UseDefaultObservers();
@@ -155,6 +167,153 @@ void FTedsSceneOutlinerWidgetConstructor::GetFilters(UE::Editor::DataStorage::IC
 	AddFavoriteFilters(DataStorage, Filters);
 }
 
+void FTedsSceneOutlinerWidgetConstructor::GetShowOptionsFilters(UE::Editor::DataStorage::ICoreProvider* DataStorage, TMap<FName, TSharedPtr<UE::Editor::Outliner::FTedsOutlinerFilter>>& Filters) const
+{
+	using namespace UE::Editor::Outliner;
+	
+	// Hide Temporary Actors
+	const FName HideTemporaryName = TEXT("HideTemporary");
+	Filters.Add(HideTemporaryName, MakeShared<FTedsOutlinerFilter>(
+		HideTemporaryName,
+		LOCTEXT("ToggleHideTemporaryActors", "Hide Temporary Actors"),
+			LOCTEXT("ToggleHideTemporaryActorsToolTip", "When enabled, hides temporary/run-time actors."),
+		NAME_None,
+		nullptr,
+		[](TQueryContext<SingleRowInfo>, const FTypedElementUObjectColumn& ObjectColumn)
+	 	{
+			if (const UObject* Object = ObjectColumn.Object.Get())
+			{
+				return !Object->HasAnyFlags(RF_Transient);
+			}
+			return false;
+	 	}));
+
+	// Hide Level Instance Content
+	const FName HideLevelInstanceContentName = TEXT("HideLevelInstanceContent");
+	Filters.Add(HideLevelInstanceContentName, MakeShared<FTedsOutlinerFilter>(
+		HideLevelInstanceContentName,
+		LOCTEXT("ToggleHideLevelInstanceContent", "Hide Level Instance Content"),
+		LOCTEXT("ToggleHideLevelInstancesToolTip", "When enabled, hides all level instance content."),
+		NAME_None,
+		nullptr,
+		[DataStorage](TQueryContext<SingleRowInfo> Context, const FTypedElementUObjectColumn& ObjectColumn)
+		 {
+			if (UObject* Object = ObjectColumn.Object.Get())
+			{
+				// Check if it is a level instance first since the Level Column holds its outer
+				if (const ILevelInstanceInterface* LevelInstance = Cast<ILevelInstanceInterface>(Object))
+				{
+					return false;
+				}
+				if (Context.CurrentRowHasColumns<FEditorDataStorageLevelTag>())
+				{
+					if (const ULevel* Level = Cast<ULevel>(Object))
+					{
+						return !Level->IsInstancedLevel();
+					}
+				}
+			}
+
+			if (DataStorage && Context.CurrentRowHasColumns<FTypedElementLevelColumn>())
+			{
+				if (const FTypedElementLevelColumn* LevelColumn = DataStorage->GetColumn<FTypedElementLevelColumn>(Context.GetCurrentRow()))
+				{
+					if (const ULevel* Level = LevelColumn->Level.Get())
+					{
+						return !Level->IsInstancedLevel();
+					}
+				}
+			}
+				
+			return true;
+		 }));
+
+	// Only Selected
+	const FName ShowOnlySelectedName = TEXT("ShowOnlySelected");
+	Filters.Add(ShowOnlySelectedName, MakeShared<FTedsOutlinerFilter>(
+		ShowOnlySelectedName,
+		LOCTEXT("ToggleShowOnlySelected", "Only Selected"),
+		LOCTEXT("ToggleShowOnlySelectedToolTip", "When enabled, only displays objects that are currently selected."),
+		NAME_None,
+		nullptr,
+		DataStorage->RegisterQuery(
+			Select()
+			.Where()
+			.All<FTypedElementSelectionColumn>()
+			.Compile())
+		));
+
+	// Only in Current Level
+	const FName ShowOnlyCurrentLevelName = TEXT("ShowOnlyCurrentLevel");
+	Filters.Add(ShowOnlyCurrentLevelName, MakeShared<FTedsOutlinerFilter>(
+		ShowOnlyCurrentLevelName,
+		LOCTEXT("ToggleShowOnlyCurrentLevel", "Only in Current Level"),
+		LOCTEXT("ToggleShowOnlyCurrentLevelToolTip", "When enabled, only shows objects that are in the Current Level."),
+		NAME_None,
+		nullptr,
+		[DataStorage](TQueryContext<SingleRowInfo> Context, const FTypedElementUObjectColumn& ObjectColumn)
+		 {
+			if (UObject* Object = ObjectColumn.Object.Get())
+			{
+				// Check if it is a level instance first since the Level Column holds its outer
+				if (const ILevelInstanceInterface* LevelInstance = Cast<ILevelInstanceInterface>(Object))
+				{
+					if (const ULevel* Level = LevelInstance->GetLoadedLevel())
+					{
+						return Level->IsCurrentLevel();
+					}
+				}
+				if (Context.CurrentRowHasColumns<FEditorDataStorageLevelTag>())
+				{
+					if (const ULevel* Level = Cast<ULevel>(Object))
+					{
+						return Level->IsCurrentLevel();
+					}
+				}
+			}
+
+			if (DataStorage && Context.CurrentRowHasColumns<FTypedElementLevelColumn>())
+			{
+				if (const FTypedElementLevelColumn* LevelColumn = DataStorage->GetColumn<FTypedElementLevelColumn>(Context.GetCurrentRow()))
+				{
+					if (const ULevel* Level = LevelColumn->Level.Get())
+					{
+						return Level->IsCurrentLevel();
+					}
+				}
+			}
+
+			return false;
+		 }));
+
+	// Only in any Current Data Layers
+	// TODO When in TEDS
+
+	// Only in Current Content Bundle
+	// TODO When in TEDS
+
+	// Hide Actor Components
+	const FName HideComponentsFilterName = TEXT("HideComponentsFilter");
+	Filters.Add(HideComponentsFilterName, MakeShared<FTedsOutlinerFilter>(
+		HideComponentsFilterName,
+		LOCTEXT("ToggleHideActorComponents", "Hide Actor Components"),
+		LOCTEXT("ToggleHideActorComponentsToolTip", "When enabled, hides components belonging to actors."),
+		NAME_None,
+		nullptr,
+		DataStorage->RegisterQuery(
+			Select()
+			.Where()
+			.None<FActorComponentTypeTag>()
+			.Compile())
+		));
+
+	// Hide Unloaded Actors
+	// TODO When in TEDS
+
+	// Hide Empty Folders
+	// TODO When in TEDS
+}
+
 bool FTedsSceneOutlinerWidgetConstructor::ShowRowHandleColumn() const
 {
 	return bShowTedsOutlinerRowHandles;
diff --git a/Engine/Plugins/Experimental/EditorDataStorageFeatures/Source/TedsOutliner/Public/TedsOutlinerFilter.h b/Engine/Plugins/Experimental/EditorDataStorageFeatures/Source/TedsOutliner/Public/TedsOutlinerFilter.h
--- a/Engine/Plugins/Experimental/EditorDataStorageFeatures/Source/TedsOutliner/Public/TedsOutlinerFilter.h
+++ b/Engine/Plugins/Experimental/EditorDataStorageFeatures/Source/TedsOutliner/Public/TedsOutlinerFilter.h
@@ -18,6 +18,7 @@ namespace UE::Editor::Outliner
 	class FTedsOutlinerFilter : public FTedsFilterBase<SceneOutliner::FilterBarType>
 	{
 		friend class FTedsOutlinerImpl;
+		friend class FTedsOutlinerMode;
 		
 	public:
 		/**
diff --git a/Engine/Plugins/Experimental/EditorDataStorageFeatures/Source/TedsOutliner/Public/TedsOutlinerImpl.h b/Engine/Plugins/Experimental/EditorDataStorageFeatures/Source/TedsOutliner/Public/TedsOutlinerImpl.h
--- a/Engine/Plugins/Experimental/EditorDataStorageFeatures/Source/TedsOutliner/Public/TedsOutlinerImpl.h
+++ b/Engine/Plugins/Experimental/EditorDataStorageFeatures/Source/TedsOutliner/Public/TedsOutlinerImpl.h
@@ -137,6 +137,9 @@ struct FTedsOutlinerParams
 	// TEDS Filters that utilize QueryDescriptions or QueryFunctions in this TEDS Outliner
 	TArray<TSharedPtr<FTedsOutlinerFilter>> Filters;
 
+	// TEDS Filters that populate the Show section of the settings menu: "Hide Components", etc.
+	TArray<TSharedPtr<FTedsOutlinerFilter>> ShowOptionsFilters;
+
 	// If true, this Outliner will include a column for row handle
 	bool bShowRowHandleColumn;
 
diff --git a/Engine/Plugins/Experimental/EditorDataStorageFeatures/Source/TedsOutliner/Public/TedsOutlinerMode.h b/Engine/Plugins/Experimental/EditorDataStorageFeatures/Source/TedsOutliner/Public/TedsOutlinerMode.h
--- a/Engine/Plugins/Experimental/EditorDataStorageFeatures/Source/TedsOutliner/Public/TedsOutlinerMode.h
+++ b/Engine/Plugins/Experimental/EditorDataStorageFeatures/Source/TedsOutliner/Public/TedsOutlinerMode.h
@@ -37,6 +37,7 @@ public:
 	UE_API virtual TSharedPtr<SWidget> CreateContextMenu() override;
 	UE_API virtual void Tick() override;
 	UE_API virtual bool ShowViewButton() const override;
+	UE_API virtual bool ShowFilterOptions() const override;
 	/* end ISceneOutlinerMode interface */
 
 protected:
diff --git a/Engine/Plugins/Experimental/EditorDataStorageFeatures/Source/TedsOutliner/Public/Widgets/SceneOutlinerWidget.h b/Engine/Plugins/Experimental/EditorDataStorageFeatures/Source/TedsOutliner/Public/Widgets/SceneOutlinerWidget.h
--- a/Engine/Plugins/Experimental/EditorDataStorageFeatures/Source/TedsOutliner/Public/Widgets/SceneOutlinerWidget.h
+++ b/Engine/Plugins/Experimental/EditorDataStorageFeatures/Source/TedsOutliner/Public/Widgets/SceneOutlinerWidget.h
@@ -70,6 +70,9 @@ public:
 	/** Adds FTedsOutlinerFilters to an array used to populate the Outliner FilterBar */
 	TEDSOUTLINER_API virtual void GetFilters(UE::Editor::DataStorage::ICoreProvider* DataStorage,
 		TArray<TSharedPtr<UE::Editor::Outliner::FTedsOutlinerFilter>>& Filters) const;
+	/** Adds FTedsOutlinerFilters to an array used to populate the Outliner Show Options */
+	TEDSOUTLINER_API virtual void GetShowOptionsFilters(UE::Editor::DataStorage::ICoreProvider* DataStorage,
+		TMap<FName, TSharedPtr<UE::Editor::Outliner::FTedsOutlinerFilter>>& Filters) const;
 	/** Whether we should add the Row Handle column to the Outliner, by default this is controlled by a CVar */
 	TEDSOUTLINER_API virtual bool ShowRowHandleColumn() const;
 	/** Whether Parent rows should be force added back when filtered or searched out */
