[TEDS Outliner] Extend the Outliner Settings  with a World Filter and allow overriding a ViewMenuExtender funciton to be passed to the TEDS Outliner Mode

#Synced-CL 52380253

diff --git a/Engine/Plugins/Experimental/EditorDataStorageFeatures/Source/TedsOutliner/Private/TedsOutlinerMode.cpp b/Engine/Plugins/Experimental/EditorDataStorageFeatures/Source/TedsOutliner/Private/TedsOutlinerMode.cpp
--- a/Engine/Plugins/Experimental/EditorDataStorageFeatures/Source/TedsOutliner/Private/TedsOutlinerMode.cpp
+++ b/Engine/Plugins/Experimental/EditorDataStorageFeatures/Source/TedsOutliner/Private/TedsOutlinerMode.cpp
@@ -69,6 +69,7 @@ FTedsOutlinerMode::FTedsOutlinerMode(const FTedsOutlinerParams& InParams)
 			FilterInfoMap.Add(FName(ShowFilter->GetName()), SceneOutlinerFilterInfo);
 		}
 	}
+	OnInitializeViewMenuExtender = InParams.OnInitializeViewMenuExtender;
 	
 	TedsOutlinerImpl->IsItemCompatible().BindLambda([](const ISceneOutlinerTreeItem& Item)
 	{
@@ -324,6 +325,14 @@ bool FTedsOutlinerMode::ShowViewButton() const
 	return bShowViewButton;
 }
 
+void FTedsOutlinerMode::InitializeViewMenuExtender(TSharedPtr<FExtender> Extender)
+{
+	if (OnInitializeViewMenuExtender.IsBound())
+	{
+		OnInitializeViewMenuExtender.Execute(Extender);
+	}
+}
+
 bool FTedsOutlinerMode::ShowFilterOptions() const
 {
 	// Use the same bool as ShowViewButton so we can have any show filters appear in the settings menu, it is better
diff --git a/Engine/Plugins/Experimental/EditorDataStorageFeatures/Source/TedsOutliner/Private/TedsOutlinerModule.cpp b/Engine/Plugins/Experimental/EditorDataStorageFeatures/Source/TedsOutliner/Private/TedsOutlinerModule.cpp
--- a/Engine/Plugins/Experimental/EditorDataStorageFeatures/Source/TedsOutliner/Private/TedsOutlinerModule.cpp
+++ b/Engine/Plugins/Experimental/EditorDataStorageFeatures/Source/TedsOutliner/Private/TedsOutlinerModule.cpp
@@ -292,8 +292,7 @@ void FTedsSceneOutlinerWorldFilter::BuildWorldFilter()
 		// Store a temporary World Ptr so we can check if the newly chosen world is the same as
 		// the previous, if so, we don't have to regenerate anything since nothing changed
 		TWeakObjectPtr<UWorld> TempWorld;
-		// This function can eventually accept a manually selected UserChosenWorld as well when this is added to options
-		SceneOutliner::FSceneOutlinerHelpers::ChooseRepresentingWorld(TempWorld, SpecifiedWorldToDisplay);
+		SceneOutliner::FSceneOutlinerHelpers::ChooseRepresentingWorld(TempWorld, SpecifiedWorldToDisplay, UserChosenWorld);
 		if (TempWorld.IsValid() && TempWorld != RepresentingWorld)
 		{
 			RepresentingWorld = TempWorld;
@@ -328,6 +327,18 @@ void FTedsSceneOutlinerWorldFilter::QueryFunction(TQueryContext<RowBatchInfo> Co
 			  }
 		  }, WorldColumns);
 }
+
+void FTedsSceneOutlinerWorldFilter::SetUserChosenWorld(const TWeakObjectPtr<UWorld>& InWorld)
+{
+	UserChosenWorld = InWorld;
+	BuildWorldFilter();
+}
+
+TWeakObjectPtr<UWorld> FTedsSceneOutlinerWorldFilter::GetUserChosenWorld() const
+{
+	return UserChosenWorld;
+}
+	
 } // namespace UE::Editor::Outliner
 
 IMPLEMENT_MODULE(UE::Editor::Outliner::FTedsOutlinerModule, TedsOutliner);
diff --git a/Engine/Plugins/Experimental/EditorDataStorageFeatures/Source/TedsOutliner/Private/Widgets/SceneOutlinerWidget.cpp b/Engine/Plugins/Experimental/EditorDataStorageFeatures/Source/TedsOutliner/Private/Widgets/SceneOutlinerWidget.cpp
--- a/Engine/Plugins/Experimental/EditorDataStorageFeatures/Source/TedsOutliner/Private/Widgets/SceneOutlinerWidget.cpp
+++ b/Engine/Plugins/Experimental/EditorDataStorageFeatures/Source/TedsOutliner/Private/Widgets/SceneOutlinerWidget.cpp
@@ -20,6 +20,7 @@
 #include "Elements/Framework/TypedElementSelectionSet.h"
 #include "Factories/TedsEditorHierarchyFactory.h"
 #include "LevelInstance/LevelInstanceInterface.h"
+#include "WorldTreeItem.h"
 
 #define LOCTEXT_NAMESPACE "SceneOutlinerWidget"
 
@@ -80,6 +81,8 @@ TSharedPtr<SWidget> FTedsSceneOutlinerWidgetConstructor::CreateWidget(UE::Editor
 		TabIdentifier = FName(**IdMetaData);
 	}
 
+	WorldFilter = MakeShared<FTedsSceneOutlinerWorldFilter>(TEXT("WorldFilter"));
+
 	FSceneOutlinerInitializationOptions InitOptions;
 	GetSceneOutlinerInitializationOptions(TabIdentifier, InitOptions);
 
@@ -95,7 +98,8 @@ TSharedPtr<SWidget> FTedsSceneOutlinerWidgetConstructor::CreateWidget(UE::Editor
 	{
 		Params.ShowOptionsFilters.Add(Filter.Value);
 	}
-	
+
+	GetInitializeViewMenuExtender(Params);
 	Params.bShowRowHandleColumn = ShowRowHandleColumn();
 	Params.bForceShowParents = ForceShowParents();
 	Params.bUseDefaultObservers = UseDefaultObservers();
@@ -162,7 +166,10 @@ void FTedsSceneOutlinerWidgetConstructor::GetFilters(UE::Editor::DataStorage::IC
 	using namespace UE::Editor::Outliner;
 	
 	Filters.Add(MakeShared<FTedsOutlinerFilter>(AActor::StaticClass()));
-	Filters.Add(MakeShared<FTedsSceneOutlinerWorldFilter>(TEXT("WorldFilter")));
+	if(WorldFilter)
+	{
+		Filters.Add(WorldFilter);
+	}
 	
 	AddFavoriteFilters(DataStorage, Filters);
 }
@@ -314,6 +321,94 @@ void FTedsSceneOutlinerWidgetConstructor::GetShowOptionsFilters(UE::Editor::Data
 	// TODO When in TEDS
 }
 
+void FTedsSceneOutlinerWidgetConstructor::GetInitializeViewMenuExtender(UE::Editor::Outliner::FTedsOutlinerParams& OutlinerParams)
+{
+	TSharedPtr<UE::Editor::Outliner::FTedsSceneOutlinerWorldFilter> LocalWorldFilter = WorldFilter;
+	
+	OutlinerParams.OnInitializeViewMenuExtender =
+		UE::Editor::Outliner::FOnExtendViewMenu::CreateLambda([LocalWorldFilter](TSharedPtr<FExtender> Extender)
+	{
+		Extender->AddMenuExtension(SceneOutliner::ExtensionHooks::Show, EExtensionHook::After, nullptr, FMenuExtensionDelegate::CreateLambda([LocalWorldFilter](FMenuBuilder& MenuBuilder)
+		{
+			MenuBuilder.BeginSection("World", LOCTEXT("ShowWorldHeading", "World"));
+				
+			auto SetUserChosenWorld = [LocalWorldFilter](const TWeakObjectPtr<UWorld>& InWorld)
+			{
+				if (LocalWorldFilter)
+				{
+					LocalWorldFilter->SetUserChosenWorld(InWorld);
+				}
+			};
+
+			auto IsCurrentUserChosenWorld = [LocalWorldFilter](const TWeakObjectPtr<UWorld>& InWorld)
+			{
+				if (LocalWorldFilter)
+				{
+					const TWeakObjectPtr<UWorld> UserChosenWorld = LocalWorldFilter->GetUserChosenWorld();
+					return (UserChosenWorld == InWorld) || (InWorld.IsExplicitlyNull() && !UserChosenWorld.IsValid());
+				}
+				return false;
+			};
+		
+			MenuBuilder.AddSubMenu(
+				LOCTEXT("ChooseWorldSubMenu", "Choose World"),
+				LOCTEXT("ChooseWorldSubMenuToolTip", "Choose the world to display in the outliner."),
+				FNewMenuDelegate::CreateLambda([SetUserChosenWorld, IsCurrentUserChosenWorld](FMenuBuilder& MenuBuilder)
+				{
+					MenuBuilder.BeginSection("Worlds", LOCTEXT("WorldsHeading", "Worlds"));
+					{
+						MenuBuilder.AddMenuEntry(
+							LOCTEXT("AutoWorld", "Auto"),
+							LOCTEXT("AutoWorldToolTip", "Automatically pick the world to display based on context."),
+							FSlateIcon(),
+							FUIAction(
+								FExecuteAction::CreateLambda([SetUserChosenWorld]()
+								{
+									SetUserChosenWorld(TWeakObjectPtr<UWorld>());
+								}),
+								FCanExecuteAction(),
+								FIsActionChecked::CreateLambda([IsCurrentUserChosenWorld]()
+								{
+									return IsCurrentUserChosenWorld(TWeakObjectPtr<UWorld>());
+								})
+							),
+							NAME_None,
+							EUserInterfaceActionType::RadioButton
+						);
+
+						for (const FWorldContext& Context : GEngine->GetWorldContexts())
+						{
+							UWorld* World = Context.World();
+							if (World && (World->WorldType == EWorldType::PIE || Context.WorldType == EWorldType::Editor))
+							{
+								MenuBuilder.AddMenuEntry(
+									SceneOutliner::GetWorldDescription(World),
+									LOCTEXT("ChooseWorldToolTip", "Display actors for this world."),
+									FSlateIcon(),
+									FUIAction(
+										FExecuteAction::CreateLambda([SetUserChosenWorld, World]()
+										{
+											SetUserChosenWorld(World);
+										}),
+										FCanExecuteAction(),
+										FIsActionChecked::CreateLambda([IsCurrentUserChosenWorld, World]()
+										{
+											return IsCurrentUserChosenWorld(World);
+										})),
+									NAME_None,
+									EUserInterfaceActionType::RadioButton
+								);
+							}
+						}
+					}
+					MenuBuilder.EndSection();
+				}));
+		
+			MenuBuilder.EndSection();
+		}));
+	});
+}
+
 bool FTedsSceneOutlinerWidgetConstructor::ShowRowHandleColumn() const
 {
 	return bShowTedsOutlinerRowHandles;
diff --git a/Engine/Plugins/Experimental/EditorDataStorageFeatures/Source/TedsOutliner/Public/TedsOutlinerImpl.h b/Engine/Plugins/Experimental/EditorDataStorageFeatures/Source/TedsOutliner/Public/TedsOutlinerImpl.h
--- a/Engine/Plugins/Experimental/EditorDataStorageFeatures/Source/TedsOutliner/Public/TedsOutlinerImpl.h
+++ b/Engine/Plugins/Experimental/EditorDataStorageFeatures/Source/TedsOutliner/Public/TedsOutlinerImpl.h
@@ -11,6 +11,7 @@
 #include "Elements/Interfaces/TypedElementDataStorageUiInterface.h"
 #include "Elements/Framework/TypedElementQueryBuilder.h"
 #include "Elements/Framework/TypedElementSelectionSet.h"
+#include "Framework/MultiBox/MultiBoxExtender.h"
 #include "SceneOutlinerPublicTypes.h"
 #include "TedsOutlinerFilter.h"
 #include "TedsOutlinerHierarchyInterfaces.h"
@@ -47,6 +48,7 @@ class SWidget;
 
 namespace UE::Editor::Outliner
 {
+	DECLARE_DELEGATE_OneParam(FOnExtendViewMenu, TSharedPtr<FExtender>);
 	
 struct FTedsOutlinerColumnParams
 {
@@ -140,6 +142,9 @@ struct FTedsOutlinerParams
 	// TEDS Filters that populate the Show section of the settings menu: "Hide Components", etc.
 	TArray<TSharedPtr<FTedsOutlinerFilter>> ShowOptionsFilters;
 
+	// Delegate for extending the View Menu with additional options
+	FOnExtendViewMenu OnInitializeViewMenuExtender;
+
 	// If true, this Outliner will include a column for row handle
 	bool bShowRowHandleColumn;
 
diff --git a/Engine/Plugins/Experimental/EditorDataStorageFeatures/Source/TedsOutliner/Public/TedsOutlinerMode.h b/Engine/Plugins/Experimental/EditorDataStorageFeatures/Source/TedsOutliner/Public/TedsOutlinerMode.h
--- a/Engine/Plugins/Experimental/EditorDataStorageFeatures/Source/TedsOutliner/Public/TedsOutlinerMode.h
+++ b/Engine/Plugins/Experimental/EditorDataStorageFeatures/Source/TedsOutliner/Public/TedsOutlinerMode.h
@@ -6,6 +6,7 @@
 #include "TedsOutlinerImpl.h"
 #include "Elements/Interfaces/TypedElementDataStorageInterface.h"
 #include "Elements/Columns/TypedElementHiearchyColumns.h"
+#include "Framework/MultiBox/MultiBoxExtender.h"
 
 #include "TedsOutlinerMode.generated.h"
 
@@ -37,6 +38,7 @@ public:
 	UE_API virtual TSharedPtr<SWidget> CreateContextMenu() override;
 	UE_API virtual void Tick() override;
 	UE_API virtual bool ShowViewButton() const override;
+	UE_API virtual void InitializeViewMenuExtender(TSharedPtr<FExtender> Extender) override;
 	UE_API virtual bool ShowFilterOptions() const override;
 	/* end ISceneOutlinerMode interface */
 
@@ -48,6 +50,8 @@ protected:
 
 protected:
 	bool bShowViewButton = false;
+	// Delegate for extending the View Menu with options (World picker, frame on selection, etc.)
+	FOnExtendViewMenu OnInitializeViewMenuExtender;
 	// The actual model for the TEDS Outliner
 	TSharedPtr<FTedsOutlinerImpl> TedsOutlinerImpl;
 };
diff --git a/Engine/Plugins/Experimental/EditorDataStorageFeatures/Source/TedsOutliner/Public/TedsOutlinerModule.h b/Engine/Plugins/Experimental/EditorDataStorageFeatures/Source/TedsOutliner/Public/TedsOutlinerModule.h
--- a/Engine/Plugins/Experimental/EditorDataStorageFeatures/Source/TedsOutliner/Public/TedsOutlinerModule.h
+++ b/Engine/Plugins/Experimental/EditorDataStorageFeatures/Source/TedsOutliner/Public/TedsOutlinerModule.h
@@ -82,10 +82,17 @@ public:
 	void BuildWorldFilter();
 	
 	void QueryFunction(TQueryContext<RowBatchInfo> Context, TResult<bool>& Result, TConstBatch<FTypedElementWorldColumn> WorldColumns);
+
+	void SetUserChosenWorld(const TWeakObjectPtr<UWorld>& InWorld);
+
+	TWeakObjectPtr<UWorld> GetUserChosenWorld() const;
 	
 protected:
 	/** The world which we are currently filtering for */
 	TWeakObjectPtr<UWorld> RepresentingWorld;
+	
+	/** The world that has been manually picked by the user */
+	TWeakObjectPtr<UWorld> UserChosenWorld;
 
 	/** If this mode was created to display a specific world, don't allow it to be reassigned */
 	const TWeakObjectPtr<UWorld> SpecifiedWorldToDisplay;
diff --git a/Engine/Plugins/Experimental/EditorDataStorageFeatures/Source/TedsOutliner/Public/Widgets/SceneOutlinerWidget.h b/Engine/Plugins/Experimental/EditorDataStorageFeatures/Source/TedsOutliner/Public/Widgets/SceneOutlinerWidget.h
--- a/Engine/Plugins/Experimental/EditorDataStorageFeatures/Source/TedsOutliner/Public/Widgets/SceneOutlinerWidget.h
+++ b/Engine/Plugins/Experimental/EditorDataStorageFeatures/Source/TedsOutliner/Public/Widgets/SceneOutlinerWidget.h
@@ -10,10 +10,13 @@
 
 #define UE_API TEDSOUTLINER_API
 
+class FExtender;
+class FMenuBuilder;
 struct FSceneOutlinerInitializationOptions;
 
 namespace UE::Editor::Outliner
 {
+	class FTedsSceneOutlinerWorldFilter;
 	class ITedsOutlinerHierarchyDataInterface;
 	class FTedsOutlinerFilter;
 	struct FTedsOutlinerColumnDescription;
@@ -73,6 +76,7 @@ public:
 	/** Adds FTedsOutlinerFilters to an array used to populate the Outliner Show Options */
 	TEDSOUTLINER_API virtual void GetShowOptionsFilters(UE::Editor::DataStorage::ICoreProvider* DataStorage,
 		TMap<FName, TSharedPtr<UE::Editor::Outliner::FTedsOutlinerFilter>>& Filters) const;
+	TEDSOUTLINER_API virtual void GetInitializeViewMenuExtender(UE::Editor::Outliner::FTedsOutlinerParams& OutlinerParams);
 	/** Whether we should add the Row Handle column to the Outliner, by default this is controlled by a CVar */
 	TEDSOUTLINER_API virtual bool ShowRowHandleColumn() const;
 	/** Whether Parent rows should be force added back when filtered or searched out */
@@ -92,6 +96,9 @@ protected:
 	
 private:
 	void AddFavoriteFilters(UE::Editor::DataStorage::ICoreProvider* DataStorage, TArray<TSharedPtr<UE::Editor::Outliner::FTedsOutlinerFilter>>& Filters) const;
+
+	// Store the World filter so it can be referenced by the World Picker post-creation
+	TSharedPtr<UE::Editor::Outliner::FTedsSceneOutlinerWorldFilter> WorldFilter;
 };
 
 #undef UE_API
