Fix for TedsTableViewer sorting, to allow for an initial descending sorting and to ensure the sorting is properly synchronized

#Synced-CL 52095318

diff --git a/Engine/Plugins/Experimental/EditorDataStorageFeatures/Source/TedsTableViewer/Private/TedsTableViewerColumn.cpp b/Engine/Plugins/Experimental/EditorDataStorageFeatures/Source/TedsTableViewer/Private/TedsTableViewerColumn.cpp
--- a/Engine/Plugins/Experimental/EditorDataStorageFeatures/Source/TedsTableViewer/Private/TedsTableViewerColumn.cpp
+++ b/Engine/Plugins/Experimental/EditorDataStorageFeatures/Source/TedsTableViewer/Private/TedsTableViewerColumn.cpp
@@ -97,6 +97,8 @@ namespace UE::Editor::DataStorage
 
 	SHeaderRow::FColumn::FArguments FTedsTableViewerColumn::ConstructHeaderRowColumn()
 	{
+		const TConstArrayView<TWeakObjectPtr<const UScriptStruct>> ColumnTypes = GetMatchedColumns();
+
 		// Compute the display name up front so it can be used for both the header widget and the tooltip
 		FText DefaultHeaderText = CellWidgetConstructor ? CellWidgetConstructor->CreateWidgetDisplayNameText(Storage) : FText::GetEmpty();
 		if (DefaultHeaderText.IsEmpty())
@@ -136,7 +138,7 @@ namespace UE::Editor::DataStorage
 		    bAllowSorting = MetaDataEntryView.TryGetExact<bool>();
 		}
 		{
-			UE::Editor::DataStorage::FMetaDataEntryView MetaDataEntryView = GetMetaData().FindGeneric("bColumnInitialSortIsAscending");
+			UE::Editor::DataStorage::FMetaDataEntryView MetaDataEntryView = GetMetaData().FindForColumnsOrGeneric(ColumnTypes, "bColumnInitialSortIsAscending");
 			bColumnInitialSortIsAscending = MetaDataEntryView.TryGetExact<bool>();
 		}
 		{
@@ -152,8 +154,6 @@ namespace UE::Editor::DataStorage
 		}
 		else
 		{
-			const TConstArrayView<TWeakObjectPtr<const UScriptStruct>> ColumnTypes = GetMatchedColumns();
-
 			FString TooltipText = TEXT("Data Storage columns:");
 			for (const TWeakObjectPtr<const UScriptStruct>& ColumnType : ColumnTypes)
 			{
@@ -166,10 +166,10 @@ namespace UE::Editor::DataStorage
 			TooltipFText = FText::FromString(MoveTemp(TooltipText));
 		}
 
-		EColumnSortMode::Type SortMode = EColumnSortMode::Ascending;
+		InitialSortMode = EColumnSortMode::Ascending;
 		if (bColumnInitialSortIsAscending)
 		{
-			SortMode = *bColumnInitialSortIsAscending ? EColumnSortMode::Ascending : EColumnSortMode::Descending;
+			InitialSortMode = *bColumnInitialSortIsAscending ? EColumnSortMode::Ascending : EColumnSortMode::Descending;
 		}
 
 		SHeaderRow::FColumn::FArguments Column = SHeaderRow::Column(ColumnName)
@@ -181,7 +181,7 @@ namespace UE::Editor::DataStorage
 			.SortPriority_Raw(this, &FTedsTableViewerColumn::GetSortPriority)
 			.IsSorting_Raw(this, &FTedsTableViewerColumn::IsSorting)
 			.OnSort_Raw(this, &FTedsTableViewerColumn::OnSortCallback)
-			.InitialSortMode(SortMode)
+			.InitialSortMode(InitialSortMode)
 			.CanManuallySort(bAllowSorting ? *bAllowSorting : true)
 			.HeaderContent()
 			[
@@ -239,6 +239,12 @@ namespace UE::Editor::DataStorage
 		return !ColumnSorters.IsEmpty();
 	}
 
+	void FTedsTableViewerColumn::SetInitialSortMode(EColumnSortMode::Type InSortMode)
+	{
+		InitialSortMode = InSortMode;
+		ClearSort();
+	}
+
 	void FTedsTableViewerColumn::SetSortDelegates(FSortHandler InSortDelegates)
 	{
 		SortDelegates = MoveTemp(InSortDelegates);
@@ -454,9 +460,17 @@ namespace UE::Editor::DataStorage
 
 	EColumnSortMode::Type FTedsTableViewerColumn::GetSortMode() const
 	{
-		return bIsSorted 
-			? ((ColumnSorterIndex & 1) == 0 ? EColumnSortMode::Type::Descending : EColumnSortMode::Type::Ascending)
-			: EColumnSortMode::Type::None;
+		if (!bIsSorted)
+		{
+			return EColumnSortMode::None;
+		}
+
+		const bool bMatchesInitialDirection = (ColumnSorterIndex & 1) == 1;
+		return bMatchesInitialDirection
+			? InitialSortMode
+			: (InitialSortMode == EColumnSortMode::Ascending
+				? EColumnSortMode::Descending
+				: EColumnSortMode::Ascending);
 	}
 
 	EColumnSortPriority::Type FTedsTableViewerColumn::GetSortPriority() const
@@ -489,7 +503,7 @@ namespace UE::Editor::DataStorage
 			SortPriority = (InSortPriority == EColumnSortPriority::Secondary) ? EColumnSortPriority::Secondary : EColumnSortPriority::Primary;
 
 			TSharedPtr<const FColumnSorterInterface> Sorter = ColumnSorters[ColumnSorterIndex >> 1];
-			
+
 			if (ColumnSorters.Num() > 1)
 			{
 				FText ShortName = Sorter->GetShortName();
@@ -497,11 +511,23 @@ namespace UE::Editor::DataStorage
 					? FText::GetEmpty()
 					: FText::Format(NSLOCTEXT("TedsTableViewColumn", "SorterName", " - {0}"), ShortName);
 			}
-			
-			SortDelegates.OnSortHandler.ExecuteIfBound(ColumnName, Sorter, InSortDirection, SortPriority);
 
-			// Each column sorter is used twice, once for ascending and once for descending sorts.
-			ColumnSorterIndex = (ColumnSorterIndex + 1) % (ColumnSorters.Num() * 2);
+			// To ensure consistency, we remove None cases
+			if (InSortDirection == EColumnSortMode::None)
+			{
+				InSortDirection = (InitialSortMode != EColumnSortMode::None) ? InitialSortMode : EColumnSortMode::Ascending;
+			}
+
+			// If we are trying to change sort direction, we go to the next column sorter
+			const EColumnSortMode::Type CurrentSortDirection = GetSortMode();
+			if (CurrentSortDirection != InSortDirection)
+			{
+				// Each column sorter is used twice, once for ascending and once for descending sorts.
+				ColumnSorterIndex = (ColumnSorterIndex + 1) % (ColumnSorters.Num() * 2);
+				ensure(GetSortMode() == InSortDirection);
+			}
+
+			SortDelegates.OnSortHandler.ExecuteIfBound(ColumnName, Sorter, InSortDirection, SortPriority);
 		}
 	}
 }
diff --git a/Engine/Plugins/Experimental/EditorDataStorageFeatures/Source/TedsTableViewer/Private/TedsTableViewerModel.cpp b/Engine/Plugins/Experimental/EditorDataStorageFeatures/Source/TedsTableViewer/Private/TedsTableViewerModel.cpp
--- a/Engine/Plugins/Experimental/EditorDataStorageFeatures/Source/TedsTableViewer/Private/TedsTableViewerModel.cpp
+++ b/Engine/Plugins/Experimental/EditorDataStorageFeatures/Source/TedsTableViewer/Private/TedsTableViewerModel.cpp
@@ -290,13 +290,14 @@ namespace UE::Editor::DataStorage
 
 			// Initialize default sorting with the primary column, if it exists
 			// Eventually there will be a use case to make this into a separate parameter, but for now this is sufficient for current user needs
-			const FMetaDataEntryView MetaDataEntryView = (*FoundPrimaryColumn)->GetMetaData().FindGeneric("bColumnInitialSortIsAscending");
+			const FMetaDataEntryView MetaDataEntryView = (*FoundPrimaryColumn)->GetMetaData().FindForColumnOrGeneric(PrimaryColumn, "bColumnInitialSortIsAscending");
 			const bool* bColumnInitialSortIsAscending = MetaDataEntryView.TryGetExact<bool>();
 			EColumnSortMode::Type SortMode = EColumnSortMode::Ascending;
 			if (bColumnInitialSortIsAscending)
 			{
 				SortMode = *bColumnInitialSortIsAscending ? EColumnSortMode::Ascending : EColumnSortMode::Descending;
 			}
+			(*FoundPrimaryColumn)->SetInitialSortMode(SortMode);
 			(*FoundPrimaryColumn)->OnSortCallback(EColumnSortPriority::Primary, PrimaryUIColumnName, SortMode);
 		}
 		else
@@ -331,30 +332,20 @@ namespace UE::Editor::DataStorage
 		{
 			if (Priority == EColumnSortPriority::Secondary)
 			{
-				if (Direction == EColumnSortMode::Type::Ascending)
+				if (RowSecondarySortingNode->GetColumnSorter() != ColumnSorter)
 				{
 					RowSecondarySortingNode->SetColumnSorter(ColumnSorter);
-					RowSecondaryInversionNode->Enable(false);
-				}
-				else
-				{
-					// Since sorting always happens in ascending first and then descending, use this knowledge to avoid resorting the current
-					// list to get the exact same result and just invert the current list instead.
-					RowSecondaryInversionNode->Enable(true);
 				}
+				RowSecondaryInversionNode->Enable(Direction != EColumnSortMode::Type::Ascending);
 				SecondarySortColumnName = ColumnName;
 			}
 			else
 			{
-				if (Direction == EColumnSortMode::Type::Ascending)
+				if (RowPrimarySortingNode->GetColumnSorter() != ColumnSorter)
 				{
 					RowPrimarySortingNode->SetColumnSorter(ColumnSorter);
-					RowPrimaryInversionNode->Enable(false);
-				}
-				else
-				{
-					RowPrimaryInversionNode->Enable(true);
 				}
+				RowPrimaryInversionNode->Enable(Direction != EColumnSortMode::Type::Ascending);
 				PrimarySortColumnName = ColumnName;
 				
 				// Reset the secondary sorter.
diff --git a/Engine/Plugins/Experimental/EditorDataStorageFeatures/Source/TedsTableViewer/Public/TedsTableViewerColumn.h b/Engine/Plugins/Experimental/EditorDataStorageFeatures/Source/TedsTableViewer/Public/TedsTableViewerColumn.h
--- a/Engine/Plugins/Experimental/EditorDataStorageFeatures/Source/TedsTableViewer/Public/TedsTableViewerColumn.h
+++ b/Engine/Plugins/Experimental/EditorDataStorageFeatures/Source/TedsTableViewer/Public/TedsTableViewerColumn.h
@@ -67,6 +67,7 @@ namespace UE::Editor::DataStorage
 			FIsSorting IsSortingHandler;
 		};
 
+		TEDSTABLEVIEWER_API void SetInitialSortMode(EColumnSortMode::Type InSortMode);
 		TEDSTABLEVIEWER_API void SetSortDelegates(FSortHandler InSortDelegates);
 		TEDSTABLEVIEWER_API bool HasSortInfo() const;
 		TEDSTABLEVIEWER_API void ClearSort();
@@ -120,6 +121,7 @@ namespace UE::Editor::DataStorage
 		FSortHandler SortDelegates;
 		FText SorterShortName;
 		uint32 ColumnSorterIndex = 0;
+		EColumnSortMode::Type InitialSortMode = EColumnSortMode::Ascending;
 		EColumnSortPriority::Type SortPriority = EColumnSortPriority::None;
 		bool bIsSorted = false;
 		
