[TedsOutliner] Migrate World Partition Data Layer column to TEDS Outliner

Adds the WP Data Layer actor column to the TEDS Outliner

- FWorldPartitionDataLayerColumn (FString DataLayers) in TedsActorCompatibility
- Sync queries: world to TEDS (GetDataLayerInstances with external/level-context handling)
- Read-only column: no write-back to world
- STextBlock cell widget via FAttributeBinder, registered on TColumn<FWorldPartitionDataLayerColumn>()
- Registered in TedsOutlinerModule
- Bridge mapping: FWorldPartitionDataLayerColumn to FSceneOutlinerBuiltInColumnTypes::DataLayer()

#Synced-CL 52470154

diff --git a/Engine/Plugins/Experimental/EditorDataStorageFeatures/Source/TedsActorCompatibility/Private/Compatibility/TedsActorWorldPartitionDataLayerQueries.cpp b/Engine/Plugins/Experimental/EditorDataStorageFeatures/Source/TedsActorCompatibility/Private/Compatibility/TedsActorWorldPartitionDataLayerQueries.cpp
new file mode 100644
--- /dev/null
+++ b/Engine/Plugins/Experimental/EditorDataStorageFeatures/Source/TedsActorCompatibility/Private/Compatibility/TedsActorWorldPartitionDataLayerQueries.cpp
@@ -0,0 +1,121 @@
+// Copyright Epic Games, Inc. All Rights Reserved.
+
+#include "Compatibility/TedsActorWorldPartitionDataLayerQueries.h"
+
+#include "Algo/Find.h"
+#include "Columns/TedsActorWorldPartitionColumns.h"
+#include "DataStorage/Queries/Types.h"
+#include "Elements/Columns/TypedElementCompatibilityColumns.h"
+#include "Elements/Columns/TypedElementMiscColumns.h"
+#include "Elements/Framework/TypedElementQueryBuilder.h"
+#include "GameFramework/Actor.h"
+#include "WorldPartition/DataLayer/DataLayerInstance.h"
+#include "WorldPartition/DataLayer/ExternalDataLayerInstance.h"
+#include "WorldPartition/WorldPartition.h"
+
+#include UE_INLINE_GENERATED_CPP_BY_NAME(TedsActorWorldPartitionDataLayerQueries)
+
+namespace UE::Editor::DataStorage::Private
+{
+	// Builds the comma-separated data layer display string for a loaded actor, replicating the legacy DataLayerInfoText lambda in SceneOutlinerModule.cpp for the FActorTreeItem case.
+	static FString BuildDataLayerString(const AActor* Actor)
+	{
+		TStringBuilder<128> Builder;
+		TSet<FString> DataLayerShortNames;
+
+		auto BuildDataLayers = [&Builder, &DataLayerShortNames](const TArray<const UDataLayerInstance*>& DataLayerInstances, bool bPartOfOtherLevel)
+		{
+			if (!bPartOfOtherLevel)
+			{
+				if (const UDataLayerInstance* const* ExternalDataLayerInstance = Algo::FindByPredicate(DataLayerInstances,
+					[](const UDataLayerInstance* DataLayerInstance) { return DataLayerInstance->IsA<UExternalDataLayerInstance>(); }))
+				{
+					Builder += (*ExternalDataLayerInstance)->GetDataLayerShortName();
+				}
+			}
+
+			for (const UDataLayerInstance* DataLayerInstance : DataLayerInstances)
+			{
+				if (!DataLayerInstance->IsA<UExternalDataLayerInstance>())
+				{
+					bool bIsAlreadyInSet = false;
+					DataLayerShortNames.Add(DataLayerInstance->GetDataLayerShortName(), &bIsAlreadyInSet);
+					if (!bIsAlreadyInSet)
+					{
+						if (Builder.Len())
+						{
+							Builder += TEXT(", ");
+						}
+						// Put a '*' in front of DataLayers that are not part of the main world
+						if (bPartOfOtherLevel)
+						{
+							Builder += TEXT("*");
+						}
+						Builder += DataLayerInstance->GetDataLayerShortName();
+					}
+				}
+			}
+		};
+
+		// List Actor's DataLayers part of the owning world, then those only part of the actor level
+		BuildDataLayers(Actor->GetDataLayerInstances(), false);
+		BuildDataLayers(Actor->GetDataLayerInstancesForLevel(), true);
+
+		return Builder.ToString();
+	}
+} // namespace UE::Editor::DataStorage::Private
+
+void UActorWorldPartitionDataLayerDataStorageFactory::RegisterQueries(UE::Editor::DataStorage::ICoreProvider& DataStorage)
+{
+	RegisterActorAddWorldPartitionDataLayerColumn(DataStorage);
+	RegisterActorWorldPartitionDataLayerToColumnQuery(DataStorage);
+}
+
+void UActorWorldPartitionDataLayerDataStorageFactory::RegisterActorAddWorldPartitionDataLayerColumn(UE::Editor::DataStorage::ICoreProvider& DataStorage) const
+{
+	using namespace UE::Editor::DataStorage::Queries;
+
+	DataStorage.RegisterQuery(
+		Select(
+			TEXT("Actor World Partition Data Layer Object to New Column"),
+			FProcessor(EQueryTickPhase::PrePhysics, DataStorage.GetQueryTickGroupName(EQueryTickGroups::SyncExternalToDataStorage)).SetExecutionMode(EExecutionMode::GameThread),
+			[](IQueryContext& Context, RowHandle Row, const FTypedElementUObjectColumn& Object, const FTypedElementWorldColumn& WorldColumn)
+			{
+				// Return if not a WP World
+				if (const UWorld* World = WorldColumn.World.Get(); !World || !World->GetWorldPartition())
+				{
+					return;
+				}
+
+				if (const AActor* Actor = Cast<AActor>(Object.Object))
+				{
+					Context.AddColumn(Row, FWorldPartitionDataLayerColumn{ .DataLayers = UE::Editor::DataStorage::Private::BuildDataLayerString(Actor) });
+				}
+			})
+		.Where()
+			.All<FTypedElementActorTag, FTypedElementSyncFromWorldTag, FLevelColumn>()
+			.None<FWorldPartitionDataLayerColumn>()
+		.Compile()
+	);
+}
+
+void UActorWorldPartitionDataLayerDataStorageFactory::RegisterActorWorldPartitionDataLayerToColumnQuery(UE::Editor::DataStorage::ICoreProvider& DataStorage) const
+{
+	using namespace UE::Editor::DataStorage::Queries;
+
+	DataStorage.RegisterQuery(
+		Select(
+			TEXT("Actor World Partition Data Layer Object to Column"),
+			FProcessor(EQueryTickPhase::PrePhysics, DataStorage.GetQueryTickGroupName(EQueryTickGroups::SyncExternalToDataStorage)).SetExecutionMode(EExecutionMode::GameThread),
+			[](IQueryContext&, RowHandle, const FTypedElementUObjectColumn& Object, FWorldPartitionDataLayerColumn& DataLayerColumn)
+			{
+				if (const AActor* Actor = Cast<AActor>(Object.Object))
+				{
+					DataLayerColumn.DataLayers = UE::Editor::DataStorage::Private::BuildDataLayerString(Actor);
+				}
+			})
+		.Where()
+			.All<FTypedElementActorTag, FTypedElementSyncFromWorldTag>()
+		.Compile()
+	);
+}
diff --git a/Engine/Plugins/Experimental/EditorDataStorageFeatures/Source/TedsActorCompatibility/Private/Compatibility/TedsActorWorldPartitionDataLayerQueries.h b/Engine/Plugins/Experimental/EditorDataStorageFeatures/Source/TedsActorCompatibility/Private/Compatibility/TedsActorWorldPartitionDataLayerQueries.h
new file mode 100644
--- /dev/null
+++ b/Engine/Plugins/Experimental/EditorDataStorageFeatures/Source/TedsActorCompatibility/Private/Compatibility/TedsActorWorldPartitionDataLayerQueries.h
@@ -0,0 +1,34 @@
+// Copyright Epic Games, Inc. All Rights Reserved.
+
+#pragma once
+
+#include "Elements/Interfaces/TypedElementDataStorageFactory.h"
+#include "UObject/ObjectMacros.h"
+
+#include "TedsActorWorldPartitionDataLayerQueries.generated.h"
+
+/**
+ * Registers queries that keep FWorldPartitionDataLayerColumn in sync with the actor's
+ * data layer assignments. Read-only — no write-back to world.
+ */
+UCLASS()
+class UActorWorldPartitionDataLayerDataStorageFactory : public UEditorDataStorageFactory
+{
+	GENERATED_BODY()
+
+public:
+	~UActorWorldPartitionDataLayerDataStorageFactory() override = default;
+
+	void RegisterQueries(UE::Editor::DataStorage::ICoreProvider& DataStorage) override;
+
+private:
+	/**
+	 * Adds FWorldPartitionDataLayerColumn to actor rows in WP worlds that do not already have the column.
+	 */
+	void RegisterActorAddWorldPartitionDataLayerColumn(UE::Editor::DataStorage::ICoreProvider& DataStorage) const;
+
+	/**
+	 * Syncs the actor's data layer names into FWorldPartitionDataLayerColumn when they may have changed.
+	 */
+	void RegisterActorWorldPartitionDataLayerToColumnQuery(UE::Editor::DataStorage::ICoreProvider& DataStorage) const;
+};
diff --git a/Engine/Plugins/Experimental/EditorDataStorageFeatures/Source/TedsOutliner/Private/Compatibility/SceneOutlinerTedsBridge.cpp b/Engine/Plugins/Experimental/EditorDataStorageFeatures/Source/TedsOutliner/Private/Compatibility/SceneOutlinerTedsBridge.cpp
--- a/Engine/Plugins/Experimental/EditorDataStorageFeatures/Source/TedsOutliner/Private/Compatibility/SceneOutlinerTedsBridge.cpp
+++ b/Engine/Plugins/Experimental/EditorDataStorageFeatures/Source/TedsOutliner/Private/Compatibility/SceneOutlinerTedsBridge.cpp
@@ -44,6 +44,7 @@ void USceneOutlinerTedsBridgeFactory::PreRegister(UE::Editor::DataStorage::ICore
 	TEDSToOutlinerDefaultColumnMapping.Add(FTedsActorSocketColumn::StaticStruct(), FSceneOutlinerBuiltInColumnTypes::Socket());
 	TEDSToOutlinerDefaultColumnMapping.Add(FEditorDataStorageActorLayersColumn::StaticStruct(), FSceneOutlinerBuiltInColumnTypes::Layer());
 	TEDSToOutlinerDefaultColumnMapping.Add(FWorldPartitionPinnedColumn::StaticStruct(), FSceneOutlinerBuiltInColumnTypes::Pinned());
+	TEDSToOutlinerDefaultColumnMapping.Add(FWorldPartitionDataLayerColumn::StaticStruct(), FSceneOutlinerBuiltInColumnTypes::DataLayer());
 }
 
 void USceneOutlinerTedsBridgeFactory::RegisterTables(UE::Editor::DataStorage::ICoreProvider& DataStorage)
diff --git a/Engine/Plugins/Experimental/EditorDataStorageFeatures/Source/TedsOutliner/Private/TedsOutlinerModule.cpp b/Engine/Plugins/Experimental/EditorDataStorageFeatures/Source/TedsOutliner/Private/TedsOutlinerModule.cpp
--- a/Engine/Plugins/Experimental/EditorDataStorageFeatures/Source/TedsOutliner/Private/TedsOutlinerModule.cpp
+++ b/Engine/Plugins/Experimental/EditorDataStorageFeatures/Source/TedsOutliner/Private/TedsOutlinerModule.cpp
@@ -143,6 +143,7 @@ FTedsOutlinerColumnDescription FTedsOutlinerModule::GetLevelEditorTedsOutlinerCo
 		FTedsPrimaryPackageObjectTag::StaticStruct(),
 		FWorldPartitionPinnedColumn::StaticStruct(),
 		FTypedElementClassTypeInfoColumn::StaticStruct(),
+		FWorldPartitionDataLayerColumn::StaticStruct(),
 		FLevelColumn::StaticStruct(), // PackageShortName
 		FTedsActorMobilityColumn::StaticStruct(),
 		FEditorDataStorageActorLayersColumn::StaticStruct(),
diff --git a/Engine/Plugins/Experimental/EditorDataStorageFeatures/Source/TedsOutliner/Private/Widgets/OutlinerDataLayerWidget.cpp b/Engine/Plugins/Experimental/EditorDataStorageFeatures/Source/TedsOutliner/Private/Widgets/OutlinerDataLayerWidget.cpp
new file mode 100644
--- /dev/null
+++ b/Engine/Plugins/Experimental/EditorDataStorageFeatures/Source/TedsOutliner/Private/Widgets/OutlinerDataLayerWidget.cpp
@@ -0,0 +1,66 @@
+// Copyright Epic Games, Inc. All Rights Reserved.
+
+#include "Widgets/OutlinerDataLayerWidget.h"
+
+#include "SceneOutlinerHelpers.h"
+#include "Columns/TedsActorWorldPartitionColumns.h"
+#include "Elements/Framework/TypedElementAttributeBinding.h"
+#include "TedsOutlinerHelpers.h"
+#include "Widgets/SBoxPanel.h"
+#include "Widgets/Text/STextBlock.h"
+
+#include UE_INLINE_GENERATED_CPP_BY_NAME(OutlinerDataLayerWidget)
+
+#define LOCTEXT_NAMESPACE "OutlinerDataLayerWidget"
+
+void UOutlinerDataLayerWidgetFactory::RegisterWidgetConstructors(UE::Editor::DataStorage::ICoreProvider& DataStorage,
+	UE::Editor::DataStorage::IUiProvider& DataStorageUi) const
+{
+	using namespace UE::Editor::DataStorage::Queries;
+
+	DataStorageUi.RegisterWidgetFactory<FOutlinerDataLayerWidgetConstructor>(
+		DataStorageUi.FindPurpose(IUiProvider::FPurposeInfo("SceneOutliner", "Cell", NAME_None).GeneratePurposeID()),
+		TColumn<FWorldPartitionDataLayerColumn>());
+}
+
+FOutlinerDataLayerWidgetConstructor::FOutlinerDataLayerWidgetConstructor()
+	: FSimpleWidgetConstructor(StaticStruct())
+{
+}
+
+FText FOutlinerDataLayerWidgetConstructor::CreateWidgetDisplayNameText(
+	UE::Editor::DataStorage::ICoreProvider* DataStorage, UE::Editor::DataStorage::RowHandle Row) const
+{
+	return LOCTEXT("SceneOutlinerDataLayerColumn", "Data Layer");
+}
+
+TSharedPtr<SWidget> FOutlinerDataLayerWidgetConstructor::CreateWidget(
+	UE::Editor::DataStorage::ICoreProvider* DataStorage,
+	UE::Editor::DataStorage::IUiProvider* DataStorageUi,
+	UE::Editor::DataStorage::RowHandle TargetRow,
+	UE::Editor::DataStorage::RowHandle WidgetRow,
+	const UE::Editor::DataStorage::FMetaDataView& Arguments)
+{
+	using namespace UE::Editor::DataStorage;
+
+	if (!DataStorage->IsRowAvailable(TargetRow))
+	{
+		return SNew(STextBlock)
+			.Text(LOCTEXT("MissingRowReferenceColumn", "Unable to retrieve row reference."));
+	}
+
+	FAttributeBinder Binder(TargetRow, DataStorage);
+
+	return SNew(SHorizontalBox)
+		+ SHorizontalBox::Slot()
+		.AutoWidth()
+		.VAlign(VAlign_Center)
+		.Padding(8, 0, 0, 0)
+		[
+			SNew(STextBlock)
+				.ColorAndOpacity(FSlateColor::UseSubduedForeground())
+				.Text(Binder.BindText(&FWorldPartitionDataLayerColumn::DataLayers))
+		];
+}
+
+#undef LOCTEXT_NAMESPACE
diff --git a/Engine/Plugins/Experimental/EditorDataStorageFeatures/Source/TedsOutliner/Public/Widgets/OutlinerDataLayerWidget.h b/Engine/Plugins/Experimental/EditorDataStorageFeatures/Source/TedsOutliner/Public/Widgets/OutlinerDataLayerWidget.h
new file mode 100644
--- /dev/null
+++ b/Engine/Plugins/Experimental/EditorDataStorageFeatures/Source/TedsOutliner/Public/Widgets/OutlinerDataLayerWidget.h
@@ -0,0 +1,46 @@
+// Copyright Epic Games, Inc. All Rights Reserved.
+
+#pragma once
+
+#include "Elements/Interfaces/TypedElementDataStorageFactory.h"
+#include "Elements/Interfaces/TypedElementDataStorageInterface.h"
+#include "Elements/Interfaces/TypedElementDataStorageUiInterface.h"
+
+#include "OutlinerDataLayerWidget.generated.h"
+
+#define UE_API TEDSOUTLINER_API
+
+UCLASS()
+class UOutlinerDataLayerWidgetFactory : public UEditorDataStorageFactory
+{
+	GENERATED_BODY()
+
+public:
+	~UOutlinerDataLayerWidgetFactory() override = default;
+
+	void RegisterWidgetConstructors(UE::Editor::DataStorage::ICoreProvider& DataStorage,
+		UE::Editor::DataStorage::IUiProvider& DataStorageUi) const override;
+};
+
+USTRUCT()
+struct FOutlinerDataLayerWidgetConstructor : public FSimpleWidgetConstructor
+{
+	GENERATED_BODY()
+
+public:
+	UE_API FOutlinerDataLayerWidgetConstructor();
+	~FOutlinerDataLayerWidgetConstructor() override = default;
+
+	UE_API virtual FText CreateWidgetDisplayNameText(
+		UE::Editor::DataStorage::ICoreProvider* DataStorage,
+		UE::Editor::DataStorage::RowHandle Row = UE::Editor::DataStorage::InvalidRowHandle) const override;
+
+	UE_API virtual TSharedPtr<SWidget> CreateWidget(
+		UE::Editor::DataStorage::ICoreProvider* DataStorage,
+		UE::Editor::DataStorage::IUiProvider* DataStorageUi,
+		UE::Editor::DataStorage::RowHandle TargetRow,
+		UE::Editor::DataStorage::RowHandle WidgetRow,
+		const UE::Editor::DataStorage::FMetaDataView& Arguments) override;
+};
+
+#undef UE_API
