[TedsOutliner] Migrate World Partition Sub Package column to TEDS Outliner

Adds the WP Sub Package actor column to the TEDS Outliner

- FWorldPartitionSubPackageColumn (FString SubPackage) in TedsActorCompatibility
- Sync queries: world to TEDS (GetWorldAssetPackage on ILevelInstanceInterface)
- Read-only column: no write-back to world
- STextBlock cell widget via FAttributeBinder, registered on TColumn<FWorldPartitionSubPackageColumn>()
- Bridge mapping: FWorldPartitionSubPackageColumn to FSceneOutlinerBuiltInColumnTypes::SubPackage()

#Synced-CL 52470154

diff --git a/Engine/Plugins/Experimental/EditorDataStorageFeatures/Source/TedsActorCompatibility/Private/Compatibility/TedsActorWorldPartitionSubPackageQueries.cpp b/Engine/Plugins/Experimental/EditorDataStorageFeatures/Source/TedsActorCompatibility/Private/Compatibility/TedsActorWorldPartitionSubPackageQueries.cpp
new file mode 100644
--- /dev/null
+++ b/Engine/Plugins/Experimental/EditorDataStorageFeatures/Source/TedsActorCompatibility/Private/Compatibility/TedsActorWorldPartitionSubPackageQueries.cpp
@@ -0,0 +1,69 @@
+// Copyright Epic Games, Inc. All Rights Reserved.
+
+#include "Compatibility/TedsActorWorldPartitionSubPackageQueries.h"
+
+#include "Columns/TedsActorWorldPartitionColumns.h"
+#include "DataStorage/Queries/Types.h"
+#include "Elements/Columns/TypedElementCompatibilityColumns.h"
+#include "Elements/Columns/TypedElementMiscColumns.h"
+#include "Elements/Framework/TypedElementQueryBuilder.h"
+#include "GameFramework/Actor.h"
+#include "LevelInstance/LevelInstanceInterface.h"
+#include "Engine/World.h"
+
+#include UE_INLINE_GENERATED_CPP_BY_NAME(TedsActorWorldPartitionSubPackageQueries)
+
+void UActorWorldPartitionSubPackageDataStorageFactory::RegisterQueries(UE::Editor::DataStorage::ICoreProvider& DataStorage)
+{
+	RegisterActorAddWorldPartitionSubPackageColumn(DataStorage);
+	RegisterActorWorldPartitionSubPackageToColumnQuery(DataStorage);
+}
+
+void UActorWorldPartitionSubPackageDataStorageFactory::RegisterActorAddWorldPartitionSubPackageColumn(UE::Editor::DataStorage::ICoreProvider& DataStorage) const
+{
+	using namespace UE::Editor::DataStorage::Queries;
+
+	DataStorage.RegisterQuery(
+		Select(
+			TEXT("Actor World Partition Sub Package Object to New Column"),
+			FProcessor(EQueryTickPhase::PrePhysics, DataStorage.GetQueryTickGroupName(EQueryTickGroups::SyncExternalToDataStorage)).SetExecutionMode(EExecutionMode::GameThread),
+			[](IQueryContext& Context, RowHandle Row, const FTypedElementUObjectColumn& Object, const FTypedElementWorldColumn& WorldColumn)
+			{
+				// Return if not a WP World
+				if (const UWorld* World = WorldColumn.World.Get(); !World || !World->GetWorldPartition())
+				{
+					return;
+				}
+
+				if (const ILevelInstanceInterface* LevelInstance = Cast<ILevelInstanceInterface>(Object.Object))
+				{
+					Context.AddColumn(Row, FWorldPartitionSubPackageColumn{ .SubPackage = LevelInstance->GetWorldAssetPackage() });
+				}
+			})
+		.Where()
+			.All<FTypedElementActorTag, FTypedElementSyncFromWorldTag, FTypedElementLevelColumn>()
+			.None<FWorldPartitionSubPackageColumn>()
+		.Compile()
+	);
+}
+
+void UActorWorldPartitionSubPackageDataStorageFactory::RegisterActorWorldPartitionSubPackageToColumnQuery(UE::Editor::DataStorage::ICoreProvider& DataStorage) const
+{
+	using namespace UE::Editor::DataStorage::Queries;
+
+	DataStorage.RegisterQuery(
+		Select(
+			TEXT("Actor World Partition Sub Package Object to Column"),
+			FProcessor(EQueryTickPhase::PrePhysics, DataStorage.GetQueryTickGroupName(EQueryTickGroups::SyncExternalToDataStorage)).SetExecutionMode(EExecutionMode::GameThread),
+			[](IQueryContext&, RowHandle, const FTypedElementUObjectColumn& Object, FWorldPartitionSubPackageColumn& SubPackageColumn)
+			{
+				if (const ILevelInstanceInterface* LevelInstance = Cast<ILevelInstanceInterface>(Object.Object))
+				{
+					SubPackageColumn.SubPackage = LevelInstance->GetWorldAssetPackage();
+				}
+			})
+		.Where()
+			.All<FTypedElementActorTag, FTypedElementSyncFromWorldTag>()
+		.Compile()
+	);
+}
diff --git a/Engine/Plugins/Experimental/EditorDataStorageFeatures/Source/TedsActorCompatibility/Private/Compatibility/TedsActorWorldPartitionSubPackageQueries.h b/Engine/Plugins/Experimental/EditorDataStorageFeatures/Source/TedsActorCompatibility/Private/Compatibility/TedsActorWorldPartitionSubPackageQueries.h
new file mode 100644
--- /dev/null
+++ b/Engine/Plugins/Experimental/EditorDataStorageFeatures/Source/TedsActorCompatibility/Private/Compatibility/TedsActorWorldPartitionSubPackageQueries.h
@@ -0,0 +1,36 @@
+// Copyright Epic Games, Inc. All Rights Reserved.
+
+#pragma once
+
+#include "Elements/Interfaces/TypedElementDataStorageFactory.h"
+#include "UObject/ObjectMacros.h"
+
+#include "TedsActorWorldPartitionSubPackageQueries.generated.h"
+
+/**
+ * Registers queries that keep FWorldPartitionSubPackageColumn in sync with the actor's
+ * world asset package name. Read-only — no write-back to world.
+ * Only added to Level Instance actors (those implementing ILevelInstanceInterface).
+ */
+UCLASS()
+class UActorWorldPartitionSubPackageDataStorageFactory : public UEditorDataStorageFactory
+{
+	GENERATED_BODY()
+
+public:
+	~UActorWorldPartitionSubPackageDataStorageFactory() override = default;
+
+	void RegisterQueries(UE::Editor::DataStorage::ICoreProvider& DataStorage) override;
+
+private:
+	/**
+	 * Adds FWorldPartitionSubPackageColumn to actor rows for Level Instance actors
+	 * that do not already have the column.
+	 */
+	void RegisterActorAddWorldPartitionSubPackageColumn(UE::Editor::DataStorage::ICoreProvider& DataStorage) const;
+
+	/**
+	 * Syncs the actor's world asset package name into FWorldPartitionSubPackageColumn when it may have changed.
+	 */
+	void RegisterActorWorldPartitionSubPackageToColumnQuery(UE::Editor::DataStorage::ICoreProvider& DataStorage) const;
+};
diff --git a/Engine/Plugins/Experimental/EditorDataStorageFeatures/Source/TedsOutliner/Private/Compatibility/SceneOutlinerTedsBridge.cpp b/Engine/Plugins/Experimental/EditorDataStorageFeatures/Source/TedsOutliner/Private/Compatibility/SceneOutlinerTedsBridge.cpp
--- a/Engine/Plugins/Experimental/EditorDataStorageFeatures/Source/TedsOutliner/Private/Compatibility/SceneOutlinerTedsBridge.cpp
+++ b/Engine/Plugins/Experimental/EditorDataStorageFeatures/Source/TedsOutliner/Private/Compatibility/SceneOutlinerTedsBridge.cpp
@@ -45,6 +45,7 @@ void USceneOutlinerTedsBridgeFactory::PreRegister(UE::Editor::DataStorage::ICore
 	TEDSToOutlinerDefaultColumnMapping.Add(FEditorDataStorageActorLayersColumn::StaticStruct(), FSceneOutlinerBuiltInColumnTypes::Layer());
 	TEDSToOutlinerDefaultColumnMapping.Add(FWorldPartitionPinnedColumn::StaticStruct(), FSceneOutlinerBuiltInColumnTypes::Pinned());
 	TEDSToOutlinerDefaultColumnMapping.Add(FWorldPartitionDataLayerColumn::StaticStruct(), FSceneOutlinerBuiltInColumnTypes::DataLayer());
+	TEDSToOutlinerDefaultColumnMapping.Add(FWorldPartitionSubPackageColumn::StaticStruct(), FSceneOutlinerBuiltInColumnTypes::SubPackage());
 }
 
 void USceneOutlinerTedsBridgeFactory::RegisterTables(UE::Editor::DataStorage::ICoreProvider& DataStorage)
diff --git a/Engine/Plugins/Experimental/EditorDataStorageFeatures/Source/TedsOutliner/Private/TedsOutlinerModule.cpp b/Engine/Plugins/Experimental/EditorDataStorageFeatures/Source/TedsOutliner/Private/TedsOutlinerModule.cpp
--- a/Engine/Plugins/Experimental/EditorDataStorageFeatures/Source/TedsOutliner/Private/TedsOutlinerModule.cpp
+++ b/Engine/Plugins/Experimental/EditorDataStorageFeatures/Source/TedsOutliner/Private/TedsOutlinerModule.cpp
@@ -144,6 +144,7 @@ FTedsOutlinerColumnDescription FTedsOutlinerModule::GetLevelEditorTedsOutlinerCo
 		FWorldPartitionPinnedColumn::StaticStruct(),
 		FTypedElementClassTypeInfoColumn::StaticStruct(),
 		FWorldPartitionDataLayerColumn::StaticStruct(),
+		FWorldPartitionSubPackageColumn::StaticStruct(),
 		FLevelColumn::StaticStruct(), // PackageShortName
 		FTedsActorMobilityColumn::StaticStruct(),
 		FEditorDataStorageActorLayersColumn::StaticStruct(),
diff --git a/Engine/Plugins/Experimental/EditorDataStorageFeatures/Source/TedsOutliner/Private/Widgets/OutlinerSubPackageWidget.cpp b/Engine/Plugins/Experimental/EditorDataStorageFeatures/Source/TedsOutliner/Private/Widgets/OutlinerSubPackageWidget.cpp
new file mode 100644
--- /dev/null
+++ b/Engine/Plugins/Experimental/EditorDataStorageFeatures/Source/TedsOutliner/Private/Widgets/OutlinerSubPackageWidget.cpp
@@ -0,0 +1,67 @@
+// Copyright Epic Games, Inc. All Rights Reserved.
+
+#include "Widgets/OutlinerSubPackageWidget.h"
+
+#include "SceneOutlinerHelpers.h"
+#include "Columns/TedsActorWorldPartitionColumns.h"
+#include "Elements/Framework/TypedElementAttributeBinding.h"
+#include "TedsOutlinerHelpers.h"
+#include "Widgets/SBoxPanel.h"
+#include "Widgets/Text/STextBlock.h"
+
+#include UE_INLINE_GENERATED_CPP_BY_NAME(OutlinerSubPackageWidget)
+
+#define LOCTEXT_NAMESPACE "OutlinerSubPackageWidget"
+
+void UOutlinerSubPackageWidgetFactory::RegisterWidgetConstructors(UE::Editor::DataStorage::ICoreProvider& DataStorage,
+	UE::Editor::DataStorage::IUiProvider& DataStorageUi) const
+{
+	using namespace UE::Editor::DataStorage::Queries;
+
+	DataStorageUi.RegisterWidgetFactory<FOutlinerSubPackageWidgetConstructor>(
+		DataStorageUi.FindPurpose(IUiProvider::FPurposeInfo("SceneOutliner", "Cell", NAME_None).GeneratePurposeID()),
+		TColumn<FWorldPartitionSubPackageColumn>());
+}
+
+FOutlinerSubPackageWidgetConstructor::FOutlinerSubPackageWidgetConstructor()
+	: FSimpleWidgetConstructor(StaticStruct())
+{
+}
+
+TSharedPtr<SWidget> FOutlinerSubPackageWidgetConstructor::CreateWidget(
+	UE::Editor::DataStorage::ICoreProvider* DataStorage,
+	UE::Editor::DataStorage::IUiProvider* DataStorageUi,
+	UE::Editor::DataStorage::RowHandle TargetRow,
+	UE::Editor::DataStorage::RowHandle WidgetRow,
+	const UE::Editor::DataStorage::FMetaDataView& Arguments)
+{
+	using namespace UE::Editor::DataStorage;
+
+	if (!DataStorage->IsRowAvailable(TargetRow))
+	{
+		return SNew(STextBlock)
+			.Text(LOCTEXT("MissingRowReferenceColumn", "Unable to retrieve row reference."));
+	}
+
+	FAttributeBinder Binder(TargetRow, DataStorage);
+
+	return SNew(SHorizontalBox)
+		+ SHorizontalBox::Slot()
+		.AutoWidth()
+		.VAlign(VAlign_Center)
+		.Padding(8, 0, 0, 0)
+		[
+			SNew(STextBlock)
+			.ColorAndOpacity(FSlateColor::UseSubduedForeground())
+			.Text(Binder.BindText(&FWorldPartitionSubPackageColumn::SubPackage))
+		];
+}
+
+FText FOutlinerSubPackageWidgetConstructor::CreateWidgetDisplayNameText(
+	UE::Editor::DataStorage::ICoreProvider* DataStorage,
+	UE::Editor::DataStorage::RowHandle Row) const
+{
+	return LOCTEXT("OutlinerSubPackageWidgetDisplayName", "Sub Package");
+}
+
+#undef LOCTEXT_NAMESPACE
diff --git a/Engine/Plugins/Experimental/EditorDataStorageFeatures/Source/TedsOutliner/Public/Widgets/OutlinerSubPackageWidget.h b/Engine/Plugins/Experimental/EditorDataStorageFeatures/Source/TedsOutliner/Public/Widgets/OutlinerSubPackageWidget.h
new file mode 100644
--- /dev/null
+++ b/Engine/Plugins/Experimental/EditorDataStorageFeatures/Source/TedsOutliner/Public/Widgets/OutlinerSubPackageWidget.h
@@ -0,0 +1,46 @@
+// Copyright Epic Games, Inc. All Rights Reserved.
+
+#pragma once
+
+#include "Elements/Interfaces/TypedElementDataStorageFactory.h"
+#include "Elements/Interfaces/TypedElementDataStorageInterface.h"
+#include "Elements/Interfaces/TypedElementDataStorageUiInterface.h"
+
+#include "OutlinerSubPackageWidget.generated.h"
+
+#define UE_API TEDSOUTLINER_API
+
+UCLASS()
+class UOutlinerSubPackageWidgetFactory : public UEditorDataStorageFactory
+{
+	GENERATED_BODY()
+
+public:
+	~UOutlinerSubPackageWidgetFactory() override = default;
+
+	void RegisterWidgetConstructors(UE::Editor::DataStorage::ICoreProvider& DataStorage,
+		UE::Editor::DataStorage::IUiProvider& DataStorageUi) const override;
+};
+
+USTRUCT()
+struct FOutlinerSubPackageWidgetConstructor : public FSimpleWidgetConstructor
+{
+	GENERATED_BODY()
+
+public:
+	UE_API FOutlinerSubPackageWidgetConstructor();
+	~FOutlinerSubPackageWidgetConstructor() override = default;
+
+	UE_API virtual TSharedPtr<SWidget> CreateWidget(
+		UE::Editor::DataStorage::ICoreProvider* DataStorage,
+		UE::Editor::DataStorage::IUiProvider* DataStorageUi,
+		UE::Editor::DataStorage::RowHandle TargetRow,
+		UE::Editor::DataStorage::RowHandle WidgetRow,
+		const UE::Editor::DataStorage::FMetaDataView& Arguments) override;
+
+	UE_API virtual FText CreateWidgetDisplayNameText(
+		UE::Editor::DataStorage::ICoreProvider* DataStorage,
+		UE::Editor::DataStorage::RowHandle Row = UE::Editor::DataStorage::InvalidRowHandle) const override;
+};
+
+#undef UE_API
