[TedsOutliner] Rename WP Pinned column display name to "Pinned"

Add CreateWidgetDisplayNameText override to FOutlinerPinnedWidgetConstructor
to return "Pinned" instead of falling back to the struct metadata
"World Partition Pinned". Matches the same pattern used by
FOutlinerDataLayerWidgetConstructor for the Data Layer column.

#Synced-CL 52470154

diff --git a/Engine/Plugins/Experimental/EditorDataStorageFeatures/Source/TedsOutliner/Private/Widgets/OutlinerPinnedWidget.cpp b/Engine/Plugins/Experimental/EditorDataStorageFeatures/Source/TedsOutliner/Private/Widgets/OutlinerPinnedWidget.cpp
--- a/Engine/Plugins/Experimental/EditorDataStorageFeatures/Source/TedsOutliner/Private/Widgets/OutlinerPinnedWidget.cpp
+++ b/Engine/Plugins/Experimental/EditorDataStorageFeatures/Source/TedsOutliner/Private/Widgets/OutlinerPinnedWidget.cpp
@@ -28,6 +28,12 @@ FOutlinerPinnedWidgetConstructor::FOutlinerPinnedWidgetConstructor()
 {
 }
 
+FText FOutlinerPinnedWidgetConstructor::CreateWidgetDisplayNameText(
+	UE::Editor::DataStorage::ICoreProvider* DataStorage, UE::Editor::DataStorage::RowHandle Row) const
+{
+	return LOCTEXT("SceneOutlinerPinnedColumn", "Pinned");
+}
+
 void UOutlinerPinnedFactory::RegisterWidgetConstructors(UE::Editor::DataStorage::ICoreProvider& DataStorage,
 	UE::Editor::DataStorage::IUiProvider& DataStorageUi) const
 {
diff --git a/Engine/Plugins/Experimental/EditorDataStorageFeatures/Source/TedsOutliner/Public/Widgets/OutlinerPinnedWidget.h b/Engine/Plugins/Experimental/EditorDataStorageFeatures/Source/TedsOutliner/Public/Widgets/OutlinerPinnedWidget.h
--- a/Engine/Plugins/Experimental/EditorDataStorageFeatures/Source/TedsOutliner/Public/Widgets/OutlinerPinnedWidget.h
+++ b/Engine/Plugins/Experimental/EditorDataStorageFeatures/Source/TedsOutliner/Public/Widgets/OutlinerPinnedWidget.h
@@ -39,6 +39,10 @@ public:
 	TEDSOUTLINER_API FOutlinerPinnedWidgetConstructor();
 	~FOutlinerPinnedWidgetConstructor() override = default;
 
+	TEDSOUTLINER_API virtual FText CreateWidgetDisplayNameText(
+		UE::Editor::DataStorage::ICoreProvider* DataStorage,
+		UE::Editor::DataStorage::RowHandle Row = UE::Editor::DataStorage::InvalidRowHandle) const override;
+
 	TEDSOUTLINER_API virtual TSharedPtr<SWidget> CreateWidget(
 		UE::Editor::DataStorage::ICoreProvider* DataStorage,
 		UE::Editor::DataStorage::IUiProvider* DataStorageUi,
