[TedsOutliner] Migrate World Partition Content Bundle column to TEDS Outliner

Adds the WP Content Bundle actor column to the TEDS Outliner

- FWorldPartitionContentBundleColumn (FString ContentBundle) in TedsActorCompatibility
- Sync queries: world to TEDS (UContentBundleEngineSubsystem descriptor display name lookup)
- Read-only column: no write-back to world
- STextBlock cell widget via FAttributeBinder::BindText, registered on TColumn<FWorldPartitionContentBundleColumn>()
- Registered in TedsOutlinerModule
- Bridge mapping: FWorldPartitionContentBundleColumn to FSceneOutlinerBuiltInColumnTypes::ContentBundle()

#Synced-CL 52470154

diff --git a/Engine/Plugins/Experimental/EditorDataStorageFeatures/Source/TedsActorCompatibility/Private/Compatibility/TedsActorWorldPartitionContentBundleQueries.cpp b/Engine/Plugins/Experimental/EditorDataStorageFeatures/Source/TedsActorCompatibility/Private/Compatibility/TedsActorWorldPartitionContentBundleQueries.cpp
new file mode 100644
--- /dev/null
+++ b/Engine/Plugins/Experimental/EditorDataStorageFeatures/Source/TedsActorCompatibility/Private/Compatibility/TedsActorWorldPartitionContentBundleQueries.cpp
@@ -0,0 +1,89 @@
+// Copyright Epic Games, Inc. All Rights Reserved.
+
+#include "Compatibility/TedsActorWorldPartitionContentBundleQueries.h"
+
+#include "Columns/TedsActorWorldPartitionColumns.h"
+#include "DataStorage/Queries/Types.h"
+#include "Elements/Columns/TypedElementCompatibilityColumns.h"
+#include "Elements/Columns/TypedElementMiscColumns.h"
+#include "Elements/Framework/TypedElementQueryBuilder.h"
+#include "GameFramework/Actor.h"
+#include "WorldPartition/ContentBundle/ContentBundleDescriptor.h"
+#include "WorldPartition/ContentBundle/ContentBundleEngineSubsystem.h"
+#include "WorldPartition/WorldPartition.h"
+
+#include UE_INLINE_GENERATED_CPP_BY_NAME(TedsActorWorldPartitionContentBundleQueries)
+
+namespace UE::Editor::DataStorage::Private
+{
+	/** Looks up the content bundle descriptor display name for a loaded actor, or returns empty string. */
+	static FString BuildContentBundleString(const AActor* Actor)
+	{
+		if (const UContentBundleEngineSubsystem* Subsystem = UContentBundleEngineSubsystem::Get())
+		{
+			if (const UContentBundleDescriptor* Descriptor = Subsystem->GetContentBundleDescriptor(Actor->GetContentBundleGuid()))
+			{
+				return Descriptor->GetDisplayName();
+			}
+		}
+		
+		return TEXT("");
+	}
+} // namespace UE::Editor::DataStorage::Private
+
+void UActorWorldPartitionContentBundleDataStorageFactory::RegisterQueries(UE::Editor::DataStorage::ICoreProvider& DataStorage)
+{
+	RegisterActorAddWorldPartitionContentBundleColumn(DataStorage);
+	RegisterActorWorldPartitionContentBundleToColumnQuery(DataStorage);
+}
+
+void UActorWorldPartitionContentBundleDataStorageFactory::RegisterActorAddWorldPartitionContentBundleColumn(
+	UE::Editor::DataStorage::ICoreProvider& DataStorage) const
+{
+	using namespace UE::Editor::DataStorage::Queries;
+
+	DataStorage.RegisterQuery(
+		Select(
+			TEXT("Actor World Partition Content Bundle Object to New Column"),
+			FProcessor(EQueryTickPhase::PrePhysics, DataStorage.GetQueryTickGroupName(EQueryTickGroups::SyncExternalToDataStorage))
+				.SetExecutionMode(EExecutionMode::GameThread),
+			[](IQueryContext& Context, RowHandle Row, const FTypedElementUObjectColumn& Object)
+			{
+				if (const AActor* Actor = Cast<AActor>(Object.Object); Actor && Actor->GetContentBundleGuid().IsValid())
+				{
+					Context.AddColumn(Row, FWorldPartitionContentBundleColumn
+						{
+							.ContentBundle = UE::Editor::DataStorage::Private::BuildContentBundleString(Actor)
+						});
+				}
+			})
+		.Where()
+			.All<FTypedElementActorTag, FTypedElementSyncFromWorldTag, FLevelColumn>()
+			.None<FWorldPartitionContentBundleColumn>()
+		.Compile()
+	);
+}
+
+void UActorWorldPartitionContentBundleDataStorageFactory::RegisterActorWorldPartitionContentBundleToColumnQuery(
+	UE::Editor::DataStorage::ICoreProvider& DataStorage) const
+{
+	using namespace UE::Editor::DataStorage::Queries;
+
+	DataStorage.RegisterQuery(
+		Select(
+			TEXT("Actor World Partition Content Bundle Object to Column"),
+			FProcessor(EQueryTickPhase::PrePhysics, DataStorage.GetQueryTickGroupName(EQueryTickGroups::SyncExternalToDataStorage))
+				.SetExecutionMode(EExecutionMode::GameThread),
+			[](IQueryContext&, RowHandle, const FTypedElementUObjectColumn& Object,
+			   FWorldPartitionContentBundleColumn& ContentBundleColumn)
+			{
+				if (const AActor* Actor = Cast<AActor>(Object.Object))
+				{
+					ContentBundleColumn.ContentBundle = UE::Editor::DataStorage::Private::BuildContentBundleString(Actor);
+				}
+			})
+		.Where()
+			.All<FTypedElementActorTag, FTypedElementSyncFromWorldTag>()
+		.Compile()
+	);
+}
diff --git a/Engine/Plugins/Experimental/EditorDataStorageFeatures/Source/TedsActorCompatibility/Private/Compatibility/TedsActorWorldPartitionContentBundleQueries.h b/Engine/Plugins/Experimental/EditorDataStorageFeatures/Source/TedsActorCompatibility/Private/Compatibility/TedsActorWorldPartitionContentBundleQueries.h
new file mode 100644
--- /dev/null
+++ b/Engine/Plugins/Experimental/EditorDataStorageFeatures/Source/TedsActorCompatibility/Private/Compatibility/TedsActorWorldPartitionContentBundleQueries.h
@@ -0,0 +1,36 @@
+// Copyright Epic Games, Inc. All Rights Reserved.
+
+#pragma once
+
+#include "Elements/Interfaces/TypedElementDataStorageFactory.h"
+#include "UObject/ObjectMacros.h"
+
+#include "TedsActorWorldPartitionContentBundleQueries.generated.h"
+
+/**
+ * Registers queries that keep FWorldPartitionContentBundleColumn in sync with the actor's
+ * content bundle assignment. Read-only — no write-back to world.
+ */
+UCLASS()
+class UActorWorldPartitionContentBundleDataStorageFactory : public UEditorDataStorageFactory
+{
+	GENERATED_BODY()
+
+public:
+	~UActorWorldPartitionContentBundleDataStorageFactory() override = default;
+
+	void RegisterQueries(UE::Editor::DataStorage::ICoreProvider& DataStorage) override;
+
+private:
+	/**
+	 * Adds FWorldPartitionContentBundleColumn to actor rows in WP worlds with content bundles
+	 * enabled that do not already have the column.
+	 */
+	void RegisterActorAddWorldPartitionContentBundleColumn(UE::Editor::DataStorage::ICoreProvider& DataStorage) const;
+
+	/**
+	 * Syncs the actor's content bundle display name into FWorldPartitionContentBundleColumn
+	 * when it may have changed.
+	 */
+	void RegisterActorWorldPartitionContentBundleToColumnQuery(UE::Editor::DataStorage::ICoreProvider& DataStorage) const;
+};
diff --git a/Engine/Plugins/Experimental/EditorDataStorageFeatures/Source/TedsActorCompatibility/Public/Columns/TedsActorWorldPartitionColumns.h b/Engine/Plugins/Experimental/EditorDataStorageFeatures/Source/TedsActorCompatibility/Public/Columns/TedsActorWorldPartitionColumns.h
--- a/Engine/Plugins/Experimental/EditorDataStorageFeatures/Source/TedsActorCompatibility/Public/Columns/TedsActorWorldPartitionColumns.h
+++ b/Engine/Plugins/Experimental/EditorDataStorageFeatures/Source/TedsActorCompatibility/Public/Columns/TedsActorWorldPartitionColumns.h
@@ -32,6 +32,19 @@ struct FWorldPartitionDataLayerColumn final : public FEditorDataStorageColumn
 	FString DataLayers;
 };
 
+/**
+ * Column that stores the World Partition content bundle display name for an actor.
+ * Only added to rows for actors in worlds with content bundles enabled.
+ */
+USTRUCT(meta = (DisplayName = "World Partition Content Bundle"))
+struct FWorldPartitionContentBundleColumn final : public FEditorDataStorageColumn
+{
+	GENERATED_BODY()
+
+	UPROPERTY(meta = (Sortable, Searchable))
+	FString ContentBundle;
+};
+
 /**
  * Column that stores the World Asset Package name for Level Instance actors.
  * Only added to rows for actors that implement ILevelInstanceInterface.
diff --git a/Engine/Plugins/Experimental/EditorDataStorageFeatures/Source/TedsOutliner/Private/Compatibility/SceneOutlinerTedsBridge.cpp b/Engine/Plugins/Experimental/EditorDataStorageFeatures/Source/TedsOutliner/Private/Compatibility/SceneOutlinerTedsBridge.cpp
--- a/Engine/Plugins/Experimental/EditorDataStorageFeatures/Source/TedsOutliner/Private/Compatibility/SceneOutlinerTedsBridge.cpp
+++ b/Engine/Plugins/Experimental/EditorDataStorageFeatures/Source/TedsOutliner/Private/Compatibility/SceneOutlinerTedsBridge.cpp
@@ -46,6 +46,7 @@ void USceneOutlinerTedsBridgeFactory::PreRegister(UE::Editor::DataStorage::ICore
 	TEDSToOutlinerDefaultColumnMapping.Add(FWorldPartitionPinnedColumn::StaticStruct(), FSceneOutlinerBuiltInColumnTypes::Pinned());
 	TEDSToOutlinerDefaultColumnMapping.Add(FWorldPartitionDataLayerColumn::StaticStruct(), FSceneOutlinerBuiltInColumnTypes::DataLayer());
 	TEDSToOutlinerDefaultColumnMapping.Add(FWorldPartitionSubPackageColumn::StaticStruct(), FSceneOutlinerBuiltInColumnTypes::SubPackage());
+	TEDSToOutlinerDefaultColumnMapping.Add(FWorldPartitionContentBundleColumn::StaticStruct(), FSceneOutlinerBuiltInColumnTypes::ContentBundle());
 }
 
 void USceneOutlinerTedsBridgeFactory::RegisterTables(UE::Editor::DataStorage::ICoreProvider& DataStorage)
diff --git a/Engine/Plugins/Experimental/EditorDataStorageFeatures/Source/TedsOutliner/Private/TedsOutlinerModule.cpp b/Engine/Plugins/Experimental/EditorDataStorageFeatures/Source/TedsOutliner/Private/TedsOutlinerModule.cpp
--- a/Engine/Plugins/Experimental/EditorDataStorageFeatures/Source/TedsOutliner/Private/TedsOutlinerModule.cpp
+++ b/Engine/Plugins/Experimental/EditorDataStorageFeatures/Source/TedsOutliner/Private/TedsOutlinerModule.cpp
@@ -145,6 +145,7 @@ FTedsOutlinerColumnDescription FTedsOutlinerModule::GetLevelEditorTedsOutlinerCo
 		FTypedElementClassTypeInfoColumn::StaticStruct(),
 		FWorldPartitionDataLayerColumn::StaticStruct(),
 		FWorldPartitionSubPackageColumn::StaticStruct(),
+		FWorldPartitionContentBundleColumn::StaticStruct(),
 		FLevelColumn::StaticStruct(), // PackageShortName
 		FTedsActorMobilityColumn::StaticStruct(),
 		FEditorDataStorageActorLayersColumn::StaticStruct(),
diff --git a/Engine/Plugins/Experimental/EditorDataStorageFeatures/Source/TedsOutliner/Private/Widgets/OutlinerContentBundleWidget.cpp b/Engine/Plugins/Experimental/EditorDataStorageFeatures/Source/TedsOutliner/Private/Widgets/OutlinerContentBundleWidget.cpp
new file mode 100644
--- /dev/null
+++ b/Engine/Plugins/Experimental/EditorDataStorageFeatures/Source/TedsOutliner/Private/Widgets/OutlinerContentBundleWidget.cpp
@@ -0,0 +1,66 @@
+// Copyright Epic Games, Inc. All Rights Reserved.
+
+#include "Widgets/OutlinerContentBundleWidget.h"
+
+#include "Columns/TedsActorWorldPartitionColumns.h"
+#include "Elements/Framework/TypedElementAttributeBinding.h"
+#include "TedsOutlinerHelpers.h"
+#include "Widgets/SBoxPanel.h"
+#include "Widgets/Text/STextBlock.h"
+
+#include UE_INLINE_GENERATED_CPP_BY_NAME(OutlinerContentBundleWidget)
+
+#define LOCTEXT_NAMESPACE "OutlinerContentBundleWidget"
+
+void UOutlinerContentBundleWidgetFactory::RegisterWidgetConstructors(
+	UE::Editor::DataStorage::ICoreProvider& DataStorage,
+	UE::Editor::DataStorage::IUiProvider& DataStorageUi) const
+{
+	using namespace UE::Editor::DataStorage::Queries;
+
+	DataStorageUi.RegisterWidgetFactory<FOutlinerContentBundleWidgetConstructor>(
+		DataStorageUi.FindPurpose(IUiProvider::FPurposeInfo("SceneOutliner", "Cell", NAME_None).GeneratePurposeID()),
+		TColumn<FWorldPartitionContentBundleColumn>());
+}
+
+FOutlinerContentBundleWidgetConstructor::FOutlinerContentBundleWidgetConstructor()
+	: FSimpleWidgetConstructor(StaticStruct())
+{
+}
+
+FText FOutlinerContentBundleWidgetConstructor::CreateWidgetDisplayNameText(
+	UE::Editor::DataStorage::ICoreProvider* DataStorage, UE::Editor::DataStorage::RowHandle Row) const
+{
+	return LOCTEXT("SceneOutlinerContentBundleColumn", "Content Bundle");
+}
+
+TSharedPtr<SWidget> FOutlinerContentBundleWidgetConstructor::CreateWidget(
+	UE::Editor::DataStorage::ICoreProvider* DataStorage,
+	UE::Editor::DataStorage::IUiProvider* DataStorageUi,
+	UE::Editor::DataStorage::RowHandle TargetRow,
+	UE::Editor::DataStorage::RowHandle WidgetRow,
+	const UE::Editor::DataStorage::FMetaDataView& Arguments)
+{
+	using namespace UE::Editor::DataStorage;
+
+	if (!DataStorage->IsRowAvailable(TargetRow))
+	{
+		return SNew(STextBlock)
+			.Text(LOCTEXT("MissingRowReferenceColumn", "Unable to retrieve row reference."));
+	}
+
+	FAttributeBinder Binder(TargetRow, DataStorage);
+
+	return SNew(SHorizontalBox)
+		+ SHorizontalBox::Slot()
+		.AutoWidth()
+		.VAlign(VAlign_Center)
+		.Padding(8, 0, 0, 0)
+		[
+			SNew(STextBlock)
+				.ColorAndOpacity(FSlateColor::UseSubduedForeground())
+				.Text(Binder.BindText(&FWorldPartitionContentBundleColumn::ContentBundle))
+		];
+}
+
+#undef LOCTEXT_NAMESPACE
diff --git a/Engine/Plugins/Experimental/EditorDataStorageFeatures/Source/TedsOutliner/Public/Widgets/OutlinerContentBundleWidget.h b/Engine/Plugins/Experimental/EditorDataStorageFeatures/Source/TedsOutliner/Public/Widgets/OutlinerContentBundleWidget.h
new file mode 100644
--- /dev/null
+++ b/Engine/Plugins/Experimental/EditorDataStorageFeatures/Source/TedsOutliner/Public/Widgets/OutlinerContentBundleWidget.h
@@ -0,0 +1,46 @@
+// Copyright Epic Games, Inc. All Rights Reserved.
+
+#pragma once
+
+#include "Elements/Interfaces/TypedElementDataStorageFactory.h"
+#include "Elements/Interfaces/TypedElementDataStorageInterface.h"
+#include "Elements/Interfaces/TypedElementDataStorageUiInterface.h"
+
+#include "OutlinerContentBundleWidget.generated.h"
+
+#define UE_API TEDSOUTLINER_API
+
+UCLASS()
+class UOutlinerContentBundleWidgetFactory : public UEditorDataStorageFactory
+{
+	GENERATED_BODY()
+
+public:
+	~UOutlinerContentBundleWidgetFactory() override = default;
+
+	void RegisterWidgetConstructors(UE::Editor::DataStorage::ICoreProvider& DataStorage,
+		UE::Editor::DataStorage::IUiProvider& DataStorageUi) const override;
+};
+
+USTRUCT()
+struct FOutlinerContentBundleWidgetConstructor : public FSimpleWidgetConstructor
+{
+	GENERATED_BODY()
+
+public:
+	UE_API FOutlinerContentBundleWidgetConstructor();
+	~FOutlinerContentBundleWidgetConstructor() override = default;
+
+	UE_API virtual FText CreateWidgetDisplayNameText(
+		UE::Editor::DataStorage::ICoreProvider* DataStorage,
+		UE::Editor::DataStorage::RowHandle Row = UE::Editor::DataStorage::InvalidRowHandle) const override;
+
+	UE_API virtual TSharedPtr<SWidget> CreateWidget(
+		UE::Editor::DataStorage::ICoreProvider* DataStorage,
+		UE::Editor::DataStorage::IUiProvider* DataStorageUi,
+		UE::Editor::DataStorage::RowHandle TargetRow,
+		UE::Editor::DataStorage::RowHandle WidgetRow,
+		const UE::Editor::DataStorage::FMetaDataView& Arguments) override;
+};
+
+#undef UE_API
