[TEDS] Made the actor table derive from the more general UObject table.

#Synced-CL 52445887

diff --git a/Engine/Plugins/Experimental/EditorDataStorageFeatures/Source/TedsActorCompatibility/Private/TedsActorCompatibilityFactory.cpp b/Engine/Plugins/Experimental/EditorDataStorageFeatures/Source/TedsActorCompatibility/Private/TedsActorCompatibilityFactory.cpp
--- a/Engine/Plugins/Experimental/EditorDataStorageFeatures/Source/TedsActorCompatibility/Private/TedsActorCompatibilityFactory.cpp
+++ b/Engine/Plugins/Experimental/EditorDataStorageFeatures/Source/TedsActorCompatibility/Private/TedsActorCompatibilityFactory.cpp
@@ -2,14 +2,10 @@
 
 #include "TedsActorCompatibilityFactory.h"
 
-#include "Editor.h"
+#include "Engine/Engine.h"
 #include "Engine/World.h"
 #include "Elements/Columns/TypedElementCompatibilityColumns.h"
 #include "Elements/Columns/TypedElementLabelColumns.h"
-#include "Elements/Columns/TypedElementMiscColumns.h"
-#include "Elements/Columns/TypedElementPackageColumns.h"
-#include "Elements/Columns/TypedElementTransformColumns.h"
-#include "Elements/Columns/TypedElementTypeInfoColumns.h"
 #include "Columns/TedsActorMobilityColumns.h"
 #include "DataStorage/Features.h"
 #include "GameFramework/Actor.h"
@@ -39,12 +35,13 @@ void UTedsActorCompatibilityFactory::RegisterTable(UE::Editor::DataStorage::ICor
 {
 	using namespace UE::Editor::DataStorage;
 
-	StandardActorTable = DataStorage.RegisterTable(TTypedElementColumnTypeList<
-		FTypedElementUObjectColumn, FTypedElementUObjectIdColumn, FUObjectIdNameColumn,
-		FTypedElementClassTypeInfoColumn, FTypedElementLabelColumn, FTypedElementLabelHashColumn, 
-		FTedsActorMobilityColumn, FTypedElementActorTag, FTypedElementSyncFromWorldTag,
-		FLevelColumn>(),
-		FName("Editor_StandardActorTable"));
+	TableHandle UObjectTable = DataStorage.FindTable("Editor_StandardUObjectTable");
+	checkf(DataStorage.IsTableValid(UObjectTable), TEXT("Actor factor is unable to retrieve the UObject table, "
+		"which is a required template for the actor table."))
+
+	StandardActorTable = DataStorage.RegisterTable<
+		FTypedElementLabelColumn, FTypedElementLabelHashColumn, FTedsActorMobilityColumn, FTypedElementActorTag, FLevelColumn>(
+		FName("Editor_StandardActorTable"), FTableRegistrationOptions{ .SourceTable = UObjectTable });
 
 	DataStorageCompatibility->RegisterTypeTableAssociation(AActor::StaticClass(), StandardActorTable);
 }
