[TEDS] Added tab to debug tables to the TEDS Debugger

#Synced-CL 52445887

diff --git a/Engine/Plugins/Experimental/EditorDataStorageFeatures/Source/TedsDebugger/Private/STedsDebugger.cpp b/Engine/Plugins/Experimental/EditorDataStorageFeatures/Source/TedsDebugger/Private/STedsDebugger.cpp
--- a/Engine/Plugins/Experimental/EditorDataStorageFeatures/Source/TedsDebugger/Private/STedsDebugger.cpp
+++ b/Engine/Plugins/Experimental/EditorDataStorageFeatures/Source/TedsDebugger/Private/STedsDebugger.cpp
@@ -16,6 +16,7 @@
 #include "TedsOutlinerMode.h"
 #include "Widgets/SScopeHierarchyTab.h"
 #include "Widgets/STedsDiscoveryDebugger.h"
+#include "Widgets/STedsTableDebugger.h"
 
 #define LOCTEXT_NAMESPACE "STedsDebugger"
 
@@ -24,6 +25,7 @@ namespace UE::Editor::DataStorage::Debug
 	namespace Private
 	{
 		const FName ToolbarTabName = TEXT("TEDS Debugger Toolbar");
+		const FName TableTabName = TEXT("Tables");
 		const FName DiscoveryTabName = TEXT("Discover");
 		const FName ContextHierarchyTabName = TEXT("Scope Hierarchy");
 	}
@@ -56,7 +58,7 @@ void STedsDebugger::Construct(const FArguments& InArgs, const TSharedRef<SDockTa
 	RegisterTabSpawners();
 
 	// Setup the default layout
-	TSharedRef<FTabManager::FLayout> Layout = FTabManager::NewLayout("TedsDebuggerLayout_v1.2")
+	TSharedRef<FTabManager::FLayout> Layout = FTabManager::NewLayout("TedsDebuggerLayout_v1.3")
 	->AddArea
 	(
 		FTabManager::NewPrimaryArea()
@@ -78,6 +80,7 @@ void STedsDebugger::Construct(const FArguments& InArgs, const TSharedRef<SDockTa
 					->AddTab(GetQueryTabName(EQueryEditorTabs::QueryEditorTwo), ETabState::ClosedTab)
 					->AddTab(GetQueryTabName(EQueryEditorTabs::QueryEditorThree), ETabState::ClosedTab)
 					->AddTab(GetQueryTabName(EQueryEditorTabs::QueryEditorFour), ETabState::ClosedTab)
+					->AddTab(Private::TableTabName, ETabState::OpenedTab)
 					->AddTab(Private::DiscoveryTabName, ETabState::OpenedTab)
 					->AddTab(Private::ContextHierarchyTabName, ETabState::OpenedTab)
 					->SetHideTabWell(false)
@@ -187,6 +190,16 @@ TSharedRef<SDockTab> STedsDebugger::SpawnQueryEditorTab(const FSpawnTabArgs& Arg
 	return DockTab;
 }
 
+TSharedRef<SDockTab> STedsDebugger::SpawnTableTab(const FSpawnTabArgs& Args)
+{
+	TSharedRef<SDockTab> DockTab = SNew(SDockTab)
+		.TabRole(ETabRole::PanelTab)
+		.Label(LOCTEXT("TedsDebugger_TableDisplayName", "Tables"));
+
+	DockTab->SetContent(SNew(STableDebuggerTab));
+	return DockTab;
+}
+
 TSharedRef<SDockTab> STedsDebugger::SpawnDiscoverTab(const FSpawnTabArgs& Args)
 {
 	TSharedRef<SDockTab> DockTab = SNew(SDockTab)
@@ -252,6 +265,14 @@ void STedsDebugger::RegisterTabSpawners()
 		.SetTooltipText(QueryEditorTooltip)
 		.SetIcon(FSlateIcon(FAppStyle::GetAppStyleSetName(), "Debug"));
 
+	TabManager->RegisterTabSpawner(
+		Private::TableTabName,
+		FOnSpawnTab::CreateRaw(this, &STedsDebugger::SpawnTableTab))
+		.SetGroup(AppMenuGroup)
+		.SetDisplayName(LOCTEXT("TedsDebugger_TableDisplayName", "Tables"))
+		.SetTooltipText(LOCTEXT("TedsDebugger_TableToolTip", "Opens the Tables tab. Contains tooling to inspect tables in the data storage."))
+		.SetIcon(FSlateIcon(FAppStyle::GetAppStyleSetName(), "Icons.TableTreeView"));
+
 	TabManager->RegisterTabSpawner(
 		Private::DiscoveryTabName,
 		FOnSpawnTab::CreateRaw(this, &STedsDebugger::SpawnDiscoverTab))
diff --git a/Engine/Plugins/Experimental/EditorDataStorageFeatures/Source/TedsDebugger/Private/STedsDebugger.h b/Engine/Plugins/Experimental/EditorDataStorageFeatures/Source/TedsDebugger/Private/STedsDebugger.h
--- a/Engine/Plugins/Experimental/EditorDataStorageFeatures/Source/TedsDebugger/Private/STedsDebugger.h
+++ b/Engine/Plugins/Experimental/EditorDataStorageFeatures/Source/TedsDebugger/Private/STedsDebugger.h
@@ -40,6 +40,7 @@ namespace UE::Editor::DataStorage
 			static FName GetQueryTabName(const EQueryEditorTabs QueryTabIdentifier);
 			TSharedRef<SDockTab> SpawnToolbar(const FSpawnTabArgs& Args);
 			TSharedRef<SDockTab> SpawnQueryEditorTab(const FSpawnTabArgs& Args, const EQueryEditorTabs QueryTabIdentifier);
+			TSharedRef<SDockTab> SpawnTableTab(const FSpawnTabArgs& Args);
 			TSharedRef<SDockTab> SpawnDiscoverTab(const FSpawnTabArgs& Args);
 			TSharedRef<SDockTab> SpawnContextHierarchyTab(const FSpawnTabArgs& Args);
 			void FillWindowMenu( FMenuBuilder& MenuBuilder);
diff --git a/Engine/Plugins/Experimental/EditorDataStorageFeatures/Source/TedsDebugger/Private/Widgets/STedsTableDebugger.cpp b/Engine/Plugins/Experimental/EditorDataStorageFeatures/Source/TedsDebugger/Private/Widgets/STedsTableDebugger.cpp
new file mode 100644
--- /dev/null
+++ b/Engine/Plugins/Experimental/EditorDataStorageFeatures/Source/TedsDebugger/Private/Widgets/STedsTableDebugger.cpp
@@ -0,0 +1,304 @@
+// Copyright Epic Games, Inc. All Rights Reserved.
+
+#include "STedsTableDebugger.h"
+
+#include "DataStorage/Features.h"
+#include "Elements/Columns/TypedElementHiearchyColumns.h"
+#include "Elements/Interfaces/TypedElementDataStorageInterface.h"
+#include "Elements/Interfaces/TypedElementDataStorageUiInterface.h"
+#include "TedsRowArrayNode.h"
+#include "Widgets/Layout/SScrollBox.h"
+#include "Widgets/Layout/SSplitter.h"
+#include "Widgets/SBoxPanel.h"
+#include "Widgets/Text/STextBlock.h"
+#include "Widgets/STedsTableViewer.h"
+
+#define LOCTEXT_NAMESPACE "TedsTableDebuggerTab"
+
+namespace UE::Editor::DataStorage::Debug
+{
+	void STableDebuggerTab::Construct(const FArguments& InArgs)
+	{
+		Storage = GetMutableDataStorageFeature<ICoreProvider>(StorageFeatureName);
+		UiProvider = GetMutableDataStorageFeature<IUiProvider>(UiFeatureName);
+        checkf(Storage && UiProvider, TEXT("Cannot create the Table Debugger tab without TEDS Core or TEDS UI"));
+		
+		RefreshRootTables();
+
+		ChildSlot
+		[
+			SNew(SBorder)
+			.BorderImage(FAppStyle::GetBrush("Brushes.Background"))
+			[
+				SNew(SSplitter)
+				.Orientation(Orient_Horizontal)
+				+SSplitter::Slot()
+				.Value(1.0f)
+				[
+					SNew(SSplitter)
+					.Orientation(Orient_Vertical)
+					+SSplitter::Slot()
+					.Value(1.0f)
+					[
+						CreateTreeView()
+					]
+					+SSplitter::Slot()
+					.Value(1.0f)
+					[
+						CreateInfoView()
+					]
+				]
+				+SSplitter::Slot()
+				.Value(2.0f)
+				[
+					CreateRowView()
+				]
+			]
+		];
+	}
+
+	void STableDebuggerTab::RefreshRootTables()
+	{
+		RootTables.Reset();
+		Storage->ListTables(ETableType::Declared, [this](TableHandle Table)
+			{
+				RootTables.Emplace(Table);
+			});
+	}
+
+	void STableDebuggerTab::OnTableSelection(TreeViewItem Item, ESelectInfo::Type SelectInfo)
+	{
+		// Fill in the row view
+		FRowHandleArray& Rows = QueryStack->GetMutableRows();
+		Rows.Reset();
+		if (Item)
+		{
+			Storage->ListTableRows(Item, [&Rows](FRowHandleArrayView Chunk)
+				{
+					Rows.Append(Chunk);
+				});
+
+			RowHandle GeneralPurposeRowHandle = UiProvider->FindPurpose(UiProvider->GetGeneralWidgetPurposeID());
+			RowHandle DefaultPurposeRowHandle = UiProvider->FindPurpose(UiProvider->GetDefaultWidgetPurposeID());
+
+			// Temporarily add the default purpose as a parent of the general purpose so the debugger can support both
+			Storage->AddColumn(GeneralPurposeRowHandle, FTableRowParentColumn{ .Parent = DefaultPurposeRowHandle });
+
+			const UScriptStruct* TagType = FEditorDataStorageTag::StaticStruct();
+			TArray<TWeakObjectPtr<const UScriptStruct>> Columns;
+			Storage->ListTableColumns(Item, [&Columns, TagType](const UScriptStruct* ColumnType)
+				{
+					Columns.Add(ColumnType);
+					return true;
+				});
+			RowView->SetColumns(Columns);
+			if (RowForeignKeyColumn && !Columns.IsEmpty())
+			{
+				RowView->AddCustomRowWidget(RowForeignKeyColumn.ToSharedRef());
+			}
+
+			// Remove the parenting chain after we have used it to generate widgets
+			Storage->RemoveColumns<FTableRowParentColumn>(GeneralPurposeRowHandle);
+		}
+		QueryStack->MarkDirty();
+
+		// Fill in the info view
+		auto TableTypeToString = [](ETableType Type)
+			{
+				switch (Type)
+				{
+				case ETableType::Invalid:
+					return FString(TEXT("Invalid"));
+				case ETableType::Declared:
+					return FString(TEXT("Declared"));
+				case ETableType::Derivative:
+					return FString(TEXT("Derivative"));
+				case ETableType::Variant:
+					return FString(TEXT("Variant"));
+				default:
+					return FString(TEXT("Unknown"));
+				}
+			};
+
+		InfoLines.Reset();
+		if (Item)
+		{
+			FTableInfoView TableInfo = Storage->GetTableInfo(Item);
+			if (TableInfo.Type != ETableType::Invalid)
+			{
+				InfoLines.Add(MakeShared<FTableInfoEntry>(FTableInfoEntry
+					{
+						.Key = TEXT("Name"),
+						.Value = TableInfo.Name.ToString()
+					}));
+				InfoLines.Add(MakeShared<FTableInfoEntry>(FTableInfoEntry
+					{
+						.Key = TEXT("Type"),
+						.Value = TableTypeToString(TableInfo.Type)
+					}));
+				InfoLines.Add(MakeShared<FTableInfoEntry>(FTableInfoEntry
+					{
+						.Key = TEXT("Derivative count"),
+						.Value = FString::FromInt(TableInfo.Derivatives.Num())
+					}));
+				InfoLines.Add(MakeShared<FTableInfoEntry>(FTableInfoEntry
+					{
+						.Key = TEXT("Variant count"),
+						.Value = FString::FromInt(TableInfo.Variants.Num())
+					}));
+
+				TStringBuilder<256> TempString;
+				Storage->ListTableForeignKeyDomains(Item, false, [&TempString](const FName& DomainName)
+					{
+						TempString.Append(DomainName.ToString());
+						TempString.Add(TEXT('\n'));
+					});
+				InfoLines.Add(MakeShared<FTableInfoEntry>(FTableInfoEntry
+					{
+						.Key = TEXT("Foreign key domains (owned)"),
+						.Value = TempString.ToString()
+					}));
+
+				TempString.Reset();
+				Storage->ListTableForeignKeyDomains(Item, true, [&TempString](const FName& DomainName)
+					{
+						TempString.Append(DomainName.ToString());
+						TempString.Add(TEXT('\n'));
+					});
+				InfoLines.Add(MakeShared<FTableInfoEntry>(FTableInfoEntry
+					{
+						.Key = TEXT("Foreign key domains (all)"),
+						.Value = TempString.ToString()
+					}));
+
+				TempString.Reset();
+				Storage->ListTableColumns(Item, [&TempString](const UScriptStruct* Column)
+					{
+						TempString.Append(Column->GetName());
+						TempString.Add(TEXT('\n'));
+						return true;
+					});
+				InfoLines.Add(MakeShared<FTableInfoEntry>(FTableInfoEntry
+					{
+						.Key = TEXT("Columns"),
+						.Value = TempString.ToString()
+					}));
+			}
+		}
+
+		InfoView->RequestListRefresh();
+	}
+
+	TSharedRef<SWidget> STableDebuggerTab::CreateTreeView()
+	{
+		return SNew(STreeView<TreeViewItem>)
+			.TreeItemsSource(&RootTables)
+			.OnGenerateRow(this, &STableDebuggerTab::MakeTableRowWidget)
+			.OnGetChildren(this, &STableDebuggerTab::HandleGetChildrenForTree)
+			.SelectionMode(ESelectionMode::Single)
+			.OnSelectionChanged(this, &STableDebuggerTab::OnTableSelection);
+	}
+
+	TSharedRef<SWidget> STableDebuggerTab::CreateInfoView()
+	{
+		return SAssignNew(InfoView, SListView<InfoItem>)
+			.ListItemsSource(&InfoLines)
+			.HeaderRow
+			(
+				SNew(SHeaderRow)
+				+SHeaderRow::Column(TEXT("Key"))
+					.DefaultLabel(FText::GetEmpty())
+				+SHeaderRow::Column(TEXT("Value"))
+					.DefaultLabel(FText::GetEmpty())
+			)
+			.OnGenerateRow(this, &STableDebuggerTab::MakeInfoRowWidget);
+	}
+
+	TSharedRef<SWidget> STableDebuggerTab::CreateRowView()
+	{
+		CreateRowForeignKeyColumn();
+
+		QueryStack = MakeShared<QueryStack::FRowArrayNode>();
+
+		SAssignNew(RowView, STedsTableViewer)
+			.QueryStack(QueryStack)
+			.EmptyRowsMessage(LOCTEXT("NoSelectionMessageForTable", "Select a table to see content."));
+		return RowView.ToSharedRef();
+	}
+	
+	void STableDebuggerTab::CreateRowForeignKeyColumn()
+	{
+		auto AssignWidgetToColumn = [this](IUiProvider::FWidgetConstructorPtr WidgetConstructor, TConstArrayView<TWeakObjectPtr<const UScriptStruct>>)
+			{
+				RowForeignKeyColumn = MakeShared<FTedsTableViewerColumn>(TEXT("Row Foreign Key"), WidgetConstructor);
+				return false;
+			};
+
+		UiProvider->CreateWidgetConstructors(
+			UiProvider->FindPurpose(IUiProvider::FPurposeInfo("General", "Cell", "RowForeignKey").GeneratePurposeID()),
+			FMetaDataView(), AssignWidgetToColumn);
+	}
+
+	TSharedRef<ITableRow> STableDebuggerTab::MakeTableRowWidget(TreeViewItem Item, const TSharedRef<STableViewBase>& Tree)
+	{
+		return SNew(STableTreeViewRow, *Storage, Item, Tree);
+	}
+
+	TSharedRef<ITableRow> STableDebuggerTab::MakeInfoRowWidget(InfoItem Item, const TSharedRef<STableViewBase>& Tree)
+	{
+		return SNew(SInfoRow, Item, Tree);
+	}
+
+	void STableDebuggerTab::HandleGetChildrenForTree(TreeViewItem Item, TArray<TreeViewItem>& Children)
+	{
+		FTableInfoView TableInfo = Storage->GetTableInfo(Item);
+		checkf(TableInfo.Type != ETableType::Invalid, TEXT("Received an invalid table handle while retrieving children."));
+
+		for (TableHandle Derivative : TableInfo.Derivatives)
+		{
+			Children.Emplace(Derivative);
+		}
+		for (TableHandle Variant : TableInfo.Variants)
+		{
+			Children.Emplace(Variant);
+		}
+	}
+
+	void STableDebuggerTab::STableTreeViewRow::Construct(const FArguments& InArgs, const ICoreProvider& InStorage, const TreeViewItem& InItem, 
+		const TSharedRef<STableViewBase>& InOwner)
+	{
+		FTableInfoView TableInfo = InStorage.GetTableInfo(InItem);
+		checkf(TableInfo.Type != ETableType::Invalid, TEXT("Received an invalid table handle for construction."));
+		
+		Super::Construct(Super::FArguments().Content()
+		[
+			SNew(STextBlock).Text(FText::FromName(TableInfo.Name))
+		], InOwner);
+	}
+
+	void STableDebuggerTab::SInfoRow::Construct(const FArguments& InArgs, const InfoItem& InItem, const TSharedRef<STableViewBase>& InOwner)
+	{
+		Item = InItem;
+
+		SMultiColumnTableRow::Construct(FSuperRowType::FArguments(), InOwner);
+	}
+
+	TSharedRef<SWidget> STableDebuggerTab::SInfoRow::GenerateWidgetForColumn(const FName& ColumnName)
+	{
+		if (ColumnName == TEXT("Key"))
+		{
+			return SNew(STextBlock)
+				.Text(FText::FromString(Item->Key));
+		}
+		else if (ColumnName == TEXT("Value"))
+		{
+			return SNew(STextBlock)
+				.Text(FText::FromString(Item->Value));
+		}
+
+		return SNullWidget::NullWidget;
+	}
+	
+} // namespace UE::Editor::DataStorage::Debug
+
+#undef LOCTEXT_NAMESPACE
diff --git a/Engine/Plugins/Experimental/EditorDataStorageFeatures/Source/TedsDebugger/Private/Widgets/STedsTableDebugger.h b/Engine/Plugins/Experimental/EditorDataStorageFeatures/Source/TedsDebugger/Private/Widgets/STedsTableDebugger.h
new file mode 100644
--- /dev/null
+++ b/Engine/Plugins/Experimental/EditorDataStorageFeatures/Source/TedsDebugger/Private/Widgets/STedsTableDebugger.h
@@ -0,0 +1,93 @@
+// Copyright Epic Games, Inc. All Rights Reserved.
+#pragma once
+
+#include "Containers/UnrealString.h"
+#include "DataStorage/CommonTypes.h"
+#include "TedsTableViewerColumn.h"
+#include "TypedElementUITypes.h"
+#include "Widgets/SCompoundWidget.h"
+#include "Widgets/Views/STableRow.h"
+#include "Widgets/Views/STreeView.h"
+#include "Widgets/Views/SListView.h"
+
+namespace UE::Editor::DataStorage
+{
+	namespace QueryStack
+	{
+		class FRowArrayNode;
+	}
+
+	class ICoreProvider;
+	class IUiProvider;
+	class STedsTableViewer;
+}
+
+namespace UE::Editor::DataStorage::Debug
+{
+	class STableDebuggerTab final : public SCompoundWidget
+	{
+	public:
+
+		SLATE_BEGIN_ARGS(STableDebuggerTab)
+		{}
+		SLATE_END_ARGS()
+
+		void Construct(const FArguments& InArgs);
+
+	private:
+		using TreeViewItem = FTedsTableHandle;
+		class STableTreeViewRow final : public STableRow<TreeViewItem>
+		{
+		public:
+			using Super = STableRow<TreeViewItem>;
+
+			SLATE_BEGIN_ARGS(STableTreeViewRow) {}
+			SLATE_END_ARGS()
+
+			void Construct(const FArguments& InArgs, const ICoreProvider& InStorage, const TreeViewItem& InItem,
+				const TSharedRef<STableViewBase>& InOwner);
+		};
+
+		struct FTableInfoEntry final
+		{
+			FString Key;
+			FString Value;
+		};
+
+		using InfoItem = TSharedPtr<FTableInfoEntry>;
+		class SInfoRow final : public SMultiColumnTableRow<InfoItem>
+		{
+		public:
+			SLATE_BEGIN_ARGS(SInfoRow) {}
+			SLATE_END_ARGS()
+
+			void Construct(const FArguments& InArgs, const InfoItem& InItem, const TSharedRef<STableViewBase>& InOwner);
+			
+		protected:
+			virtual TSharedRef<SWidget> GenerateWidgetForColumn(const FName& ColumnName) override;
+
+		private:
+			InfoItem Item;
+		};
+
+		void RefreshRootTables();
+		void OnTableSelection(TreeViewItem Item, ESelectInfo::Type SelectInfo);
+
+		TSharedRef<SWidget> CreateTreeView();
+		TSharedRef<SWidget> CreateInfoView();
+		TSharedRef<SWidget> CreateRowView();
+		void CreateRowForeignKeyColumn();
+		TSharedRef<ITableRow> MakeTableRowWidget(TreeViewItem Item, const TSharedRef<STableViewBase>& Tree);
+		TSharedRef<ITableRow> MakeInfoRowWidget(InfoItem Item, const TSharedRef<STableViewBase>& Tree);
+		void HandleGetChildrenForTree(TreeViewItem Item, TArray<TreeViewItem>& Children);
+
+		TArray<TreeViewItem> RootTables;
+		TArray<InfoItem> InfoLines;
+		TSharedPtr<SListView<InfoItem>> InfoView;
+		TSharedPtr<STedsTableViewer> RowView;
+		TSharedPtr<QueryStack::FRowArrayNode> QueryStack;
+		TSharedPtr<FTedsTableViewerColumn> RowForeignKeyColumn;
+		ICoreProvider* Storage = nullptr;
+		IUiProvider* UiProvider = nullptr;
+	};
+}
diff --git a/Engine/Plugins/Experimental/EditorDataStorageFeatures/Source/TedsTableViewer/Private/Widgets/RowForeignKeyWidget.cpp b/Engine/Plugins/Experimental/EditorDataStorageFeatures/Source/TedsTableViewer/Private/Widgets/RowForeignKeyWidget.cpp
new file mode 100644
--- /dev/null
+++ b/Engine/Plugins/Experimental/EditorDataStorageFeatures/Source/TedsTableViewer/Private/Widgets/RowForeignKeyWidget.cpp
@@ -0,0 +1,155 @@
+// Copyright Epic Games, Inc. All Rights Reserved.
+
+#include "Widgets/RowForeignKeyWidget.h"
+
+#include "Elements/Columns/TypedElementMiscColumns.h"
+#include "Widgets/Layout/SBox.h"
+#include "Widgets/SCompoundWidget.h"
+#include "Widgets/Text/STextBlock.h"
+
+#include UE_INLINE_GENERATED_CPP_BY_NAME(RowForeignKeyWidget)
+
+#define LOCTEXT_NAMESPACE "RowForeignKeyWidget"
+
+namespace UE::Editor::DataStorage::Ui
+{
+	class SRowForeignKeyCellWidget : public SCompoundWidget
+	{
+	public:
+		SLATE_BEGIN_ARGS(SRowForeignKeyCellWidget)
+			: _DataStorage(nullptr)
+			, _Row(InvalidRowHandle)
+			{
+			}
+
+			SLATE_ATTRIBUTE(const ICoreProvider*, DataStorage)
+			SLATE_ATTRIBUTE(RowHandle, Row)
+		SLATE_END_ARGS()
+
+		void Construct(const FArguments& InArgs)
+		{
+			DataStorage = InArgs._DataStorage.Get();
+			Row = InArgs._Row.Get();
+			Table = DataStorage->FindTable(Row);
+			
+			ChildSlot
+				[
+					SAssignNew(TextBlock, STextBlock)
+						.ColorAndOpacity(FSlateColor::UseForeground())
+						.OnDoubleClicked_Lambda(
+							[this](const FGeometry&, const FPointerEvent&) 
+							{ 
+								SetText();
+								return FReply::Handled();
+							})
+				];
+
+			SetText();
+		}
+
+		void SetText()
+		{
+			FName Domain;
+			int32 Counter = 0;
+			// Note: Tables are never destroyed and foreign keys can't be unregistered. This means that the list of foreign key domains
+			// for a given table is also immutable, so we can rely on the index to remain valid across calls.
+			DataStorage->ListTableForeignKeyDomains(Table, true, [this, &Counter, &Domain](const FName& DomainName)
+				{
+					if (Counter == Index)
+					{
+						Domain = DomainName;
+					}
+					Counter++;
+				});
+			if (Counter > 0 && !Domain.IsNone())
+			{
+				Index++;
+				if (Index >= Counter)
+				{
+					Index = 0;
+				}
+
+				FString Text = Domain.ToString();
+				Text.Append(TEXT(": "));
+				FForeignKey Key = DataStorage->BuildForeignKey(Domain, Row);
+				Text.Append(Key.ToString());
+				TextBlock->SetText(FText::FromString(MoveTemp(Text)));
+			}
+			else
+			{
+				TextBlock->SetText(FText::GetEmpty());
+			}
+		}
+
+	private:
+		TSharedPtr<STextBlock> TextBlock;
+		const ICoreProvider* DataStorage = nullptr;
+		RowHandle Row = InvalidRowHandle;
+		TableHandle Table = InvalidTableHandle;
+		int32 Index = 0;
+	};
+} // namespace UE::Editor::DataStorage::Ui
+
+void URowForeignKeyWidgetFactory::RegisterWidgetConstructors(UE::Editor::DataStorage::ICoreProvider& DataStorage, UE::Editor::DataStorage::IUiProvider& DataStorageUi) const
+{
+	DataStorageUi.RegisterWidgetFactory(DataStorageUi.FindPurpose(
+		UE::Editor::DataStorage::IUiProvider::FPurposeInfo("General", "Cell", "RowForeignKey").GeneratePurposeID()),
+		FRowForeignKeyWidgetConstructor::StaticStruct());
+}
+
+void URowForeignKeyWidgetFactory::RegisterWidgetPurposes(UE::Editor::DataStorage::IUiProvider& DataStorageUi) const
+{
+	DataStorageUi.RegisterWidgetPurpose(
+		UE::Editor::DataStorage::IUiProvider::FPurposeInfo("General", "Cell", "RowForeignKey",
+			UE::Editor::DataStorage::IUiProvider::EPurposeType::UniqueByName,
+			LOCTEXT("GeneralRowForeignKeyPurpose", "Specific purpose to request a widget to display row foreign keys.")));
+}
+
+FRowForeignKeyWidgetConstructor::FRowForeignKeyWidgetConstructor()
+	: Super(FRowForeignKeyWidgetConstructor::StaticStruct())
+{
+	
+}
+
+TSharedPtr<SWidget> FRowForeignKeyWidgetConstructor::CreateWidget(const UE::Editor::DataStorage::FMetaDataView& Arguments)
+{
+	return SNew(SBox)
+		.HAlign(HAlign_Left)
+		.VAlign(VAlign_Center)
+		.Padding(8, 0, 0, 0);}
+
+bool FRowForeignKeyWidgetConstructor::FinalizeWidget(UE::Editor::DataStorage::ICoreProvider* DataStorage, 
+	UE::Editor::DataStorage::IUiProvider* DataStorageUi, UE::Editor::DataStorage::RowHandle Row, const TSharedPtr<SWidget>& Widget)
+{
+	using namespace UE::Editor::DataStorage::Ui;
+
+	checkf(Widget->GetType() == SBox::StaticWidgetClass().GetWidgetType(),
+		TEXT("Stored widget with FRowForeignKeyWidgetConstructor doesn't match type %s, but was a %s."),
+		*(SBox::StaticWidgetClass().GetWidgetType().ToString()),
+		*(Widget->GetTypeAsString()));
+	
+	SBox* BoxWidget = static_cast<SBox*>(Widget.Get());
+
+	UE::Editor::DataStorage::RowHandle TargetRowHandle = UE::Editor::DataStorage::InvalidRowHandle;
+
+	if (const FTypedElementRowReferenceColumn* RowReferenceColumn = DataStorage->GetColumn<FTypedElementRowReferenceColumn>(Row))
+	{
+		BoxWidget->SetContent(
+			SNew(SRowForeignKeyCellWidget)
+			.DataStorage(DataStorage)
+			.Row(RowReferenceColumn->Row));
+	}
+	else
+	{
+		BoxWidget->SetContent(SNullWidget::NullWidget);
+	}
+	return true;
+}
+
+FText FRowForeignKeyWidgetConstructor::CreateWidgetDisplayNameText(UE::Editor::DataStorage::ICoreProvider* DataStorage,
+	UE::Editor::DataStorage::RowHandle Row) const
+{
+	return LOCTEXT("RowForeignKeyColumnName", "Row Foreign Key");
+}
+
+#undef LOCTEXT_NAMESPACE //"RowForeignKeyWidget"
diff --git a/Engine/Plugins/Experimental/EditorDataStorageFeatures/Source/TedsTableViewer/Public/TypedElementUITypes.h b/Engine/Plugins/Experimental/EditorDataStorageFeatures/Source/TedsTableViewer/Public/TypedElementUITypes.h
--- a/Engine/Plugins/Experimental/EditorDataStorageFeatures/Source/TedsTableViewer/Public/TypedElementUITypes.h
+++ b/Engine/Plugins/Experimental/EditorDataStorageFeatures/Source/TedsTableViewer/Public/TypedElementUITypes.h
@@ -60,6 +60,66 @@ struct TListTypeTraits<FTedsRowHandle>
 // Template declaration to enable using row handles inside of slate widgets like SListView
 template <>
 struct TIsValidListItem<FTedsRowHandle>
+{
+	enum
+	{
+		Value = true
+	};
+};
+
+/* Template declaration to describe how a table handle behaves as a type for slate widgets like SListView, STreeView etc
+ * This allows you to use Table Handles with slate widgets that would otherwise only work on pointers by using the wrapper 
+ * struct e.g STreeView<FTedsTableHandle>
+ */
+template <>
+struct TListTypeTraits<FTedsTableHandle>
+{
+	typedef FTedsTableHandle NullableType;
+
+	using MapKeyFuncs = TDefaultMapHashableKeyFuncs<FTedsTableHandle, TSharedRef<ITableRow>, false>;
+	using MapKeyFuncsSparse = TDefaultMapHashableKeyFuncs<FTedsTableHandle, FSparseItemInfo, false>;
+	using SetKeyFuncs = DefaultKeyFuncs<FTedsTableHandle>;
+
+	template<typename U>
+	static void AddReferencedObjects(FReferenceCollector&,
+		TArray<FTedsTableHandle>&,
+		TSet<FTedsTableHandle>&,
+		TMap<const U*, FTedsTableHandle>&)
+	{
+	}
+
+	static bool IsPtrValid(const FTedsTableHandle& InPtr)
+	{
+		return InPtr != UE::Editor::DataStorage::InvalidTableHandle;
+	}
+
+	static void ResetPtr(FTedsTableHandle& InPtr)
+	{
+		InPtr = UE::Editor::DataStorage::InvalidTableHandle;
+	}
+
+	static FTedsTableHandle MakeNullPtr()
+	{
+		static FTedsTableHandle InvalidTableHandle{ .TableHandle = UE::Editor::DataStorage::InvalidTableHandle };
+		return InvalidTableHandle;
+	}
+
+	static FTedsTableHandle NullableItemTypeConvertToItemType(const FTedsTableHandle& InPtr)
+	{
+		return InPtr;
+	}
+
+	static FString DebugDump(FTedsTableHandle InPtr)
+	{
+		return FString::Printf(TEXT("%llu"), InPtr.TableHandle);
+	}
+
+	class SerializerType {};
+};
+
+// Template declaration to enable using table handles inside of slate widgets like STreeView
+template <>
+struct TIsValidListItem<FTedsTableHandle>
 {
 	enum
 	{
diff --git a/Engine/Plugins/Experimental/EditorDataStorageFeatures/Source/TedsTableViewer/Public/Widgets/RowForeignKeyWidget.h b/Engine/Plugins/Experimental/EditorDataStorageFeatures/Source/TedsTableViewer/Public/Widgets/RowForeignKeyWidget.h
new file mode 100644
--- /dev/null
+++ b/Engine/Plugins/Experimental/EditorDataStorageFeatures/Source/TedsTableViewer/Public/Widgets/RowForeignKeyWidget.h
@@ -0,0 +1,61 @@
+// Copyright Epic Games, Inc. All Rights Reserved.
+
+#pragma once
+
+#include "DataStorage/Queries/Conditions.h"
+#include "Elements/Interfaces/TypedElementDataStorageFactory.h"
+#include "Elements/Interfaces/TypedElementDataStorageUiInterface.h"
+
+#include "RowForeignKeyWidget.generated.h"
+
+#define UE_API TEDSTABLEVIEWER_API
+
+namespace UE::Editor::DataStorage
+{
+	class ICoreProvider;
+} // namespace UE::Editor::DataStorage
+
+class UScriptStruct;
+
+namespace UE::Editor::DataStorage::Ui
+{
+
+} // namespace UE::Editor::DataStorage::Ui
+
+UCLASS(MinimalAPI)
+class URowForeignKeyWidgetFactory : public UEditorDataStorageFactory
+{
+	GENERATED_BODY()
+
+public:
+	virtual ~URowForeignKeyWidgetFactory() override = default;
+
+	UE_API virtual void RegisterWidgetConstructors(UE::Editor::DataStorage::ICoreProvider& DataStorage,
+		UE::Editor::DataStorage::IUiProvider& DataStorageUi) const override;
+
+	UE_API void RegisterWidgetPurposes(UE::Editor::DataStorage::IUiProvider& DataStorageUi) const override;
+};
+
+// A custom widget to display the row ForeignKey of a row as text
+USTRUCT()
+struct FRowForeignKeyWidgetConstructor : public FTypedElementWidgetConstructor
+{
+	GENERATED_BODY()
+
+public:
+	UE_API FRowForeignKeyWidgetConstructor();
+	virtual ~FRowForeignKeyWidgetConstructor() override = default;
+
+protected:
+	UE_API virtual TSharedPtr<SWidget> CreateWidget(const UE::Editor::DataStorage::FMetaDataView& Arguments) override;
+	UE_API virtual bool FinalizeWidget(
+		UE::Editor::DataStorage::ICoreProvider* DataStorage,
+		UE::Editor::DataStorage::IUiProvider* DataStorageUi,
+		UE::Editor::DataStorage::RowHandle Row,
+		const TSharedPtr<SWidget>& Widget) override;
+		
+	UE_API virtual FText CreateWidgetDisplayNameText(UE::Editor::DataStorage::ICoreProvider* DataStorage, 
+    		UE::Editor::DataStorage::RowHandle Row = UE::Editor::DataStorage::InvalidRowHandle) const override;
+};
+
+#undef UE_API
diff --git a/Engine/Source/Runtime/TypedElementFramework/Public/DataStorage/Handles.h b/Engine/Source/Runtime/TypedElementFramework/Public/DataStorage/Handles.h
--- a/Engine/Source/Runtime/TypedElementFramework/Public/DataStorage/Handles.h
+++ b/Engine/Source/Runtime/TypedElementFramework/Public/DataStorage/Handles.h
@@ -31,7 +31,7 @@ namespace UE::Editor::DataStorage
 /*
  * FTedsRowHandle is a strongly typed wrapper around UE::Editor::DataStorage::RowHandle and should only be used in cases where you need the extra info.
  * E.g for reflection/UHT or for template specializing something that needs to know the semantics of the row handle.
- * For all other cases, you should use the regular typedef UE::Editor::DataStorage::RowHandle
+ * For all other cases, you should use the regular UE::Editor::DataStorage::RowHandle
  */
 USTRUCT()
 struct FTedsRowHandle
@@ -63,4 +63,41 @@ struct FTedsRowHandle
 	}
 };
 
-static_assert(sizeof(FTedsRowHandle::RowHandle) == sizeof(UE::Editor::DataStorage::RowHandle), "RowHandle and RowHandle wrapper sizes should match");
\ No newline at end of file
+static_assert(sizeof(FTedsRowHandle::RowHandle) == sizeof(UE::Editor::DataStorage::RowHandle), "RowHandle and RowHandle wrapper sizes should match");
+
+/*
+ * FTedsTableHandle is a strongly typed wrapper around UE::Editor::DataStorage::TableHandle and should only be used in cases where you need the extra info.
+ * E.g for reflection/UHT or for template specializing something that needs to know the semantics of the table handle.
+ * For all other cases, you should use the regular UE::Editor::DataStorage::TableHandle
+ */
+USTRUCT()
+struct FTedsTableHandle
+{
+	GENERATED_BODY()
+
+	UPROPERTY()
+	uint64 TableHandle = UE::Editor::DataStorage::InvalidTableHandle;
+
+	operator UE::Editor::DataStorage::TableHandle() const
+	{
+		return TableHandle;
+	}
+
+	FTedsTableHandle& operator=(UE::Editor::DataStorage::TableHandle InTableHandle)
+	{
+		TableHandle = InTableHandle;
+		return *this;
+	}
+
+	bool IsValid() const
+	{
+		return TableHandle != UE::Editor::DataStorage::InvalidTableHandle;
+	}
+
+	friend uint32 GetTypeHash(const FTedsTableHandle& Key)
+	{
+		return GetTypeHash(Key.TableHandle);
+	}
+};
+
+static_assert(sizeof(FTedsTableHandle::TableHandle) == sizeof(UE::Editor::DataStorage::TableHandle), "TableHandle and TableHandle wrapper sizes should match");
\ No newline at end of file
