[TEDS] Expand SRowDetails configurability with new Slate arguments and metadata support

Added new arguments to `SRowDetails`:
- `Columns`: explicit ordered list of TEDS columns to display; falls back to auto-discovery sorted by display name
- `ShowUnpopulatedColumns`: show/hide requested columns absent from the row
- `SelectionMode`: configures list view selection behavior
- `HeaderRow`: optional custom header row replacing the default Name/Value columns
- `NameMetaData`/`DataMetaData`/`ColumnMetaData`: metadata passed to widget constructors for per-column customization
- `RowPadding`: padding applied inside each row

Additional changes:
- `SRowDetailsRow` now holds a weak reference to its owning `SRowDetails` instead of raw `ICoreProvider`/`IUiProvider` pointers
- `ClearRow` now properly removes widget rows from TEDS storage
- `NameColumn`/`DataColumn` promoted from private statics to public `UE_API` members
- Name column widget now supports font style override via metadata

#Synced-CL 52471066

diff --git a/Engine/Plugins/Experimental/EditorDataStorageFeatures/Source/TedsTableViewer/Private/Widgets/SRowDetails.cpp b/Engine/Plugins/Experimental/EditorDataStorageFeatures/Source/TedsTableViewer/Private/Widgets/SRowDetails.cpp
--- a/Engine/Plugins/Experimental/EditorDataStorageFeatures/Source/TedsTableViewer/Private/Widgets/SRowDetails.cpp
+++ b/Engine/Plugins/Experimental/EditorDataStorageFeatures/Source/TedsTableViewer/Private/Widgets/SRowDetails.cpp
@@ -14,22 +14,19 @@
 
 namespace UE::Editor::DataStorage
 {
-	namespace Widgets::Private
-	{
-		static const FName NameColumn(TEXT("Name"));
-		static const FName DataColumn(TEXT("Data"));
-	} // namespace Widgets::Private
-
 	//
 	// SRowDetails
 	//
 
+	const FName SRowDetails::NameColumn = "Name";
+	const FName SRowDetails::DataColumn = "Data";
+
 	void SRowDetails::Construct(const FArguments& InArgs)
 	{
 		bShowAllDetails = InArgs._ShowAllDetails;
-		
-		WidgetPurpose = InArgs._WidgetPurposeOverride;
+		bShowUnpopulatedColumns = InArgs._ShowUnpopulatedColumns;
 
+		WidgetPurpose = InArgs._WidgetPurposeOverride;
 		if(!WidgetPurpose.IsSet())
 		{
 			WidgetPurpose = IUiProvider::FPurposeInfo("RowDetails", "Cell", "Large").GeneratePurposeID();
@@ -39,12 +36,20 @@ namespace UE::Editor::DataStorage
 
 		DataStorage = GetMutableDataStorageFeature<ICoreProvider>(StorageFeatureName);
 		DataStorageUi = GetMutableDataStorageFeature<IUiProvider>(UiFeatureName);
+		check(DataStorage);
+		check(DataStorageUi);
+
+		RequestedTedsColumns = InArgs._Columns;
+		NameMetaData = InArgs._NameMetaData;
+		DataMetaData = InArgs._DataMetaData;
+		ColumnMetaData = InArgs._ColumnMetaData;
+		RowPadding = InArgs._RowPadding;
 
-		
 		ChildSlot
 		[
 			SAssignNew(ListView, SListView<RowDetailsItemPtr>)
 				.ListItemsSource(&Items)
+				.SelectionMode(InArgs._SelectionMode)
 				.OnGenerateRow(this, &SRowDetails::CreateRow)
 				.Visibility_Lambda([this]()
 					{
@@ -52,102 +57,153 @@ namespace UE::Editor::DataStorage
 					})
 				.HeaderRow
 				(
-					SNew(SHeaderRow)
-					+SHeaderRow::Column(Widgets::Private::NameColumn)
-						.DefaultLabel(FText::FromString(TEXT("Name")))
-						.FillWidth(0.3f)
-					+SHeaderRow::Column(Widgets::Private::DataColumn)
-						.DefaultLabel(FText::FromString(TEXT("Value")))
-						.FillWidth(0.7f)
+					InArgs._HeaderRow
+						? InArgs._HeaderRow.ToSharedRef()
+						: SNew(SHeaderRow)
+						+SHeaderRow::Column(NameColumn)
+							.DefaultLabel(FText::FromString(TEXT("Name")))
+							.FillWidth(0.3f)
+						+SHeaderRow::Column(DataColumn)
+							.DefaultLabel(FText::FromString(TEXT("Value")))
+							.FillWidth(0.7f)
 				)
 		];
 	}
 	
 	void SRowDetails::SetRow(RowHandle Row)
 	{
-		if (DataStorage->IsRowAssigned(Row))
+		ClearRow();
+
+		if (!DataStorage->IsRowAssigned(Row))
+		{
+			return;
+		}
+
+		// A Map of TEDS Columns -> UI columns so we can add them in the same order they were specified
+		TMap<TWeakObjectPtr<const UScriptStruct>, TSharedRef<FRowDetailsItem>> NewColumnMap;
+
+		TArray<TWeakObjectPtr<const UScriptStruct>> Columns = RequestedTedsColumns;
+		if (Columns.IsEmpty())
 		{
-			Items.Reset();
-			
-			TArray<TWeakObjectPtr<const UScriptStruct>> Columns;
 			DataStorage->ListColumns(Row, [&Columns](const UScriptStruct& ColumnType)
+			{
+				Columns.Emplace(&ColumnType);
+				return true;
+			});
+			// Sorting by display name by default
+			Columns.Sort(
+				[](const TWeakObjectPtr<const UScriptStruct>& Lhs, const TWeakObjectPtr<const UScriptStruct>& Rhs)
 				{
-					Columns.Emplace(&ColumnType);
-					return true;
+					return Lhs->GetDisplayNameText().CompareTo(Rhs->GetDisplayNameText()) < 0;
 				});
+		}
+		TArray<TWeakObjectPtr<const UScriptStruct>> ColumnsCopy = Columns;
+
+		RowHandle PurposeRow = DataStorageUi->FindPurpose(WidgetPurpose);
+
+		static const FMetaData EmptyMetaData;
+		// Using the DataMetaData here because the widget constructor will be used for the Data column
+		FComboMetaDataView<FGenericMetaDataView, FColumnsMetaDataView> MetaDataView =
+			FComboMetaDataView(FGenericMetaDataView(DataMetaData ? *DataMetaData : EmptyMetaData))
+			.Next(FColumnsMetaDataView(ColumnMetaData));
 
-			RowHandle PurposeRow = DataStorageUi->FindPurpose(WidgetPurpose);
-			
-			DataStorageUi->CreateWidgetConstructors(PurposeRow, IUiProvider::EMatchApproach::LongestMatch,
-					Columns, {}, [this, Row](
-						IUiProvider::FWidgetConstructorPtr WidgetConstructor, TConstArrayView<TWeakObjectPtr<const UScriptStruct>> Columns)
+		// Create the widget constructors for the columns
+		DataStorageUi->CreateWidgetConstructors(PurposeRow, IUiProvider::EMatchApproach::LongestMatch, ColumnsCopy,
+			MetaDataView, [this, Row, &NewColumnMap](
+			IUiProvider::FWidgetConstructorPtr WidgetConstructor, TConstArrayView<TWeakObjectPtr<const UScriptStruct>> MatchedColumns)
+			{
+				for (const TWeakObjectPtr<const UScriptStruct>& Column : MatchedColumns)
+				{
+					if (!bShowUnpopulatedColumns && !DataStorage->HasColumns(Row, {Column}))
 					{
-						Items.Add(MakeShared<FRowDetailsItem>(nullptr, MoveTemp(WidgetConstructor), Row));
-						return true;
-					});
-			
+						continue;
+					}
+					NewColumnMap.Emplace(Column, MakeShared<FRowDetailsItem>(Column, WidgetConstructor, Row));
+				}
+				return true;
+			});
+
 			if (bShowAllDetails)
 			{
-				// Create defaults for the remaining widgets.
-				for (TWeakObjectPtr<const UScriptStruct> Column : Columns)
+				// ColumnsCopy has been modified in-place by CreateWidgetConstructors above, which removes matched columns.
+				// Only unmatched columns remain here.
+				for (TWeakObjectPtr<const UScriptStruct> Column : ColumnsCopy)
 				{
-					DataStorageUi->CreateWidgetConstructors(DataStorageUi->FindPurpose(DataStorageUi->GetDefaultWidgetPurposeID()), {},
-						[this, Column, Row](
+					DataStorageUi->CreateWidgetConstructors(DataStorageUi->FindPurpose(DataStorageUi->GetDefaultWidgetPurposeID()),
+						MetaDataView, [this, Column, Row, &NewColumnMap](
 							IUiProvider::FWidgetConstructorPtr Constructor, TConstArrayView<TWeakObjectPtr<const UScriptStruct>> Columns)
 						{
-							Items.Add(MakeShared<FRowDetailsItem>(Column, MoveTemp(Constructor), Row));
+							NewColumnMap.Emplace(Column, MakeShared<FRowDetailsItem>(Column, MoveTemp(Constructor), Row));
 							return true;
 						});
 				}
 			}
-			
-			ListView->RequestListRefresh();
-		}
-		else
+
+		// Add the actual UI columns in the order the Teds Columns were specified
+		for (const TWeakObjectPtr<const UScriptStruct>& ColumnType : Columns)
 		{
-			ClearRow();
+			if (const TSharedRef<FRowDetailsItem>* FoundColumn = NewColumnMap.Find(ColumnType))
+			{
+				Items.Add(*FoundColumn);
+				NewColumnMap.Remove(ColumnType);
+			}
 		}
+
+		ListView->RequestListRefresh();
 	}
-	
+
 	void SRowDetails::ClearRow()
 	{
+		for (const RowDetailsItemPtr& Item : Items)
+		{
+			if (DataStorage->IsRowAvailable(Item->WidgetRow))
+			{
+				DataStorage->RemoveRow(Item->WidgetRow);
+			}
+		}
 		Items.Reset();
 		ListView->RequestListRefresh();
 	}
 	
 	TSharedRef<ITableRow> SRowDetails::CreateRow(RowDetailsItemPtr InItem, const TSharedRef<STableViewBase>& OwnerTable)
 	{
-		TSharedPtr<SRowDetailsRow> Row = SNew(SRowDetailsRow, OwnerTable, DataStorage, DataStorageUi)
-			.Item(InItem);
-		
-		return Row.ToSharedRef();
+		TSharedRef<SRowDetailsRow> Row = SNew(SRowDetailsRow, OwnerTable, SharedThis(this))
+			.Item(MoveTemp(InItem))
+			.Padding(RowPadding);
+
+		return Row;
 	}
 	
 	//
 	// SRowDetailsRow
 	//
-	void SRowDetailsRow::Construct(const FArguments& Args, const TSharedRef<STableViewBase>& OwnerTableView, ICoreProvider* InDataStorage,
-			IUiProvider* InDataStorageUi)
+	void SRowDetailsRow::Construct(const FArguments& InArgs, const TSharedRef<STableViewBase>& OwnerTableView, const TSharedRef<SRowDetails>& InOwnerRowDetails)
 	{
-		Item = Args._Item;
+		Item = InArgs._Item;
+		check(Item);
+		OwnerRowDetails = InOwnerRowDetails;
 
-		DataStorage = InDataStorage;
-		DataStorageUi = InDataStorageUi;
-		
-		SMultiColumnTableRow<RowDetailsItemPtr>::Construct(FSuperRowType::FArguments(), OwnerTableView);
+		SMultiColumnTableRow<RowDetailsItemPtr>::Construct(FSuperRowType::FArguments().Padding(InArgs._Padding), OwnerTableView);
 	}
 	
 	TSharedRef<SWidget> SRowDetailsRow::GenerateWidgetForColumn(const FName& ColumnName)
-	{		
+	{
+		TSharedPtr<SRowDetails> RowDetails = OwnerRowDetails.Pin();
+		if (!RowDetails)
+		{
+			return SNullWidget::NullWidget;
+		}
+
+		ICoreProvider* DataStorage = RowDetails->GetDataStorage();
 		if (!DataStorage->IsRowAvailable(Item->WidgetRow))
 		{
 			Item->WidgetRow = DataStorage->AddRow(DataStorage->FindTable(FName("Editor_WidgetTable")));
-			
+
 			DataStorage->AddColumn<FTypedElementRowReferenceColumn>(Item->WidgetRow, FTypedElementRowReferenceColumn
 				{
 					.Row = Item->Row
 				});
-			
+
 			if (Item->ColumnType.IsValid() &&
 				Item->WidgetConstructor->GetAdditionalColumnsList().Contains(FTypedElementScriptStructTypeInfoColumn::StaticStruct()))
 			{
@@ -157,20 +213,36 @@ namespace UE::Editor::DataStorage
 					});
 			}
 		}
-		if (ColumnName == Widgets::Private::NameColumn)
+
+		static const FMetaData EmptyMetaData;
+		if (ColumnName == SRowDetails::NameColumn)
 		{
-			return SNew(STextBlock)
-					.Text(Item->WidgetConstructor->CreateWidgetDisplayNameText(
-							DataStorage, Item->WidgetRow));
+			FComboMetaDataView<FGenericMetaDataView, FColumnsMetaDataView> MetaDataView =
+				FComboMetaDataView(FGenericMetaDataView(RowDetails->GetNameMetaData() ? *RowDetails->GetNameMetaData() : EmptyMetaData))
+				.Next(FColumnsMetaDataView(RowDetails->GetColumnMetaData()));
+
+			STextBlock::FArguments Args = STextBlock::FArguments();
+			if (const FString* const* FontStyle = MetaDataView.FindForColumnOrGeneric(Item->ColumnType, "Font").TryGetExact<const FString*>())
+			{
+				Args._Font = FAppStyle::Get().GetFontStyle(FName(**FontStyle));
+			}
+
+			return SArgumentNew(Args, STextBlock)
+				.Text(Item->WidgetConstructor->CreateWidgetDisplayNameText(DataStorage, Item->WidgetRow));
 		}
-		else if (ColumnName == Widgets::Private::DataColumn)
+		else if (ColumnName == SRowDetails::DataColumn)
 		{
-			return DataStorageUi->ConstructWidget(Item->WidgetRow, *(Item->WidgetConstructor), {}).ToSharedRef();
+			FComboMetaDataView<FGenericMetaDataView, FColumnsMetaDataView> MetaDataView =
+				FComboMetaDataView(FGenericMetaDataView(RowDetails->GetDataMetaData() ? *RowDetails->GetDataMetaData() : EmptyMetaData))
+				.Next(FColumnsMetaDataView(RowDetails->GetColumnMetaData()));
+
+			const TSharedPtr<SWidget> Widget = RowDetails->GetDataStorageUi()->ConstructWidget(Item->WidgetRow, *(Item->WidgetConstructor), MetaDataView);
+			return Widget ? Widget.ToSharedRef() : SNullWidget::NullWidget;
 		}
 		else
 		{
 			return SNew(STextBlock)
-					.Text(LOCTEXT("InvalidColumnType", "Invalid Column Type")); 
+					.Text(LOCTEXT("InvalidColumnType", "Invalid Column Type"));
 		}
 	}
 
diff --git a/Engine/Plugins/Experimental/EditorDataStorageFeatures/Source/TedsTableViewer/Public/Widgets/SRowDetails.h b/Engine/Plugins/Experimental/EditorDataStorageFeatures/Source/TedsTableViewer/Public/Widgets/SRowDetails.h
--- a/Engine/Plugins/Experimental/EditorDataStorageFeatures/Source/TedsTableViewer/Public/Widgets/SRowDetails.h
+++ b/Engine/Plugins/Experimental/EditorDataStorageFeatures/Source/TedsTableViewer/Public/Widgets/SRowDetails.h
@@ -38,20 +38,41 @@ namespace UE::Editor::DataStorage
 	class SRowDetails : public SCompoundWidget
 	{
 	public:
-		
+		UE_API static const FName NameColumn;
+		UE_API static const FName DataColumn;
+
 		~SRowDetails() override = default;
 		
 		SLATE_BEGIN_ARGS(SRowDetails)
-			: _ShowAllDetails(true)
+			: _ShowUnpopulatedColumns(false)
+			, _ShowAllDetails(true)
+			, _SelectionMode(ESelectionMode::Multi)
 		{}
 
+			// List of Columns to show in order. Leave blank to automatically find all matching columns.
+			SLATE_ARGUMENT(TArray<TWeakObjectPtr<const UScriptStruct>>, Columns)
+			// If true, requested columns with no data on this row are shown as empty entries. Has no effect when Columns is empty.
+			SLATE_ARGUMENT(bool, ShowUnpopulatedColumns)
+
 			// Whether or not to show columns that don't have a dedicated widget to represent them
 			SLATE_ARGUMENT(bool, ShowAllDetails)
 
 			// Override for the default widget purposes used to create widgets for the columns
 			SLATE_ARGUMENT(IUiProvider::FPurposeID, WidgetPurposeOverride)
 
-		
+			SLATE_ATTRIBUTE(ESelectionMode::Type, SelectionMode)
+			// Custom header row for this SRowDetails. Leave null to use a default Name/Value header.
+			// If provided, it must contain columns named SRowDetails::NameColumn and SRowDetails::DataColumn.
+			SLATE_ARGUMENT(TSharedPtr<SHeaderRow>, HeaderRow)
+
+			// MetaData for the Name Widget shared by all columns
+			SLATE_ARGUMENT(TSharedPtr<FMetaData>, NameMetaData)
+			// MetaData for the Data Widget shared by all columns
+			SLATE_ARGUMENT(TSharedPtr<FMetaData>, DataMetaData)
+			// MetaData for all columns shared by the Name and the Data widgets
+			SLATE_ARGUMENT(TArray<TSharedPtr<FColumnMetaData>>, ColumnMetaData)
+			// The padding inside each Row
+			SLATE_ATTRIBUTE(FMargin, RowPadding)
 		SLATE_END_ARGS()
 		
 		UE_API void Construct(const FArguments& InArgs);
@@ -61,22 +82,31 @@ namespace UE::Editor::DataStorage
 
 		// Clear the row to view
 		UE_API void ClearRow();
-		
+
+		ICoreProvider* GetDataStorage() const { return DataStorage; }
+		IUiProvider* GetDataStorageUi() const { return DataStorageUi; }
+		const TSharedPtr<FMetaData>& GetNameMetaData() const { return NameMetaData; }
+		const TSharedPtr<FMetaData>& GetDataMetaData() const { return DataMetaData; }
+		const TArray<TSharedPtr<FColumnMetaData>>& GetColumnMetaData() const { return ColumnMetaData; }
 	private:
 		
 		UE_API TSharedRef<ITableRow> CreateRow(RowDetailsItemPtr InItem, const TSharedRef<STableViewBase>& OwnerTable);
 
+		TArray<TWeakObjectPtr<const UScriptStruct>> RequestedTedsColumns;
 		TSharedPtr<SListView<RowDetailsItemPtr>> ListView;
-
 		TArray<RowDetailsItemPtr> Items;
 
 		ICoreProvider* DataStorage = nullptr; 
 		IUiProvider* DataStorageUi = nullptr;
 
-		bool bShowAllDetails = true;
-
 		IUiProvider::FPurposeID WidgetPurpose;
 
+		TSharedPtr<FMetaData> NameMetaData;
+		TSharedPtr<FMetaData> DataMetaData;
+		TArray<TSharedPtr<FColumnMetaData>> ColumnMetaData;
+		TAttribute<FMargin> RowPadding;
+		bool bShowAllDetails = true;
+		bool bShowUnpopulatedColumns = false;
 	};
 
 	class SRowDetailsRow : public SMultiColumnTableRow<RowDetailsItemPtr>
@@ -86,18 +116,17 @@ namespace UE::Editor::DataStorage
 		SLATE_BEGIN_ARGS(SRowDetailsRow) {}
 		
 			SLATE_ARGUMENT(RowDetailsItemPtr, Item)
-
+			// The padding inside each Row
+			SLATE_ATTRIBUTE(FMargin, Padding)
 		SLATE_END_ARGS()
 		
-		void Construct(const FArguments& Args, const TSharedRef<STableViewBase>& OwnerTableView, ICoreProvider* InDataStorage,
-			IUiProvider* InDataStorageUi);
-		
-		TSharedRef<SWidget> GenerateWidgetForColumn(const FName& ColumnName) override;
+		void Construct(const FArguments& InArgs, const TSharedRef<STableViewBase>& OwnerTableView, const TSharedRef<SRowDetails>& InOwnerRowDetails);
+
+		virtual TSharedRef<SWidget> GenerateWidgetForColumn(const FName& ColumnName) override;
 
 	private:
 		RowDetailsItemPtr Item;
-		ICoreProvider* DataStorage = nullptr;
-		IUiProvider* DataStorageUi = nullptr;
+		TWeakPtr<SRowDetails> OwnerRowDetails;
 	};
 } // namespace UE::Editor::DataStorage
 
