[TedsOutliner] Simplify WP Pinned Column queries to use TEDS columns for world and level access

#Synced-CL 52470154

diff --git a/Engine/Plugins/Experimental/EditorDataStorageFeatures/Source/TedsActorCompatibility/Private/Compatibility/TedsActorWorldPartitionPinnedQueries.cpp b/Engine/Plugins/Experimental/EditorDataStorageFeatures/Source/TedsActorCompatibility/Private/Compatibility/TedsActorWorldPartitionPinnedQueries.cpp
--- a/Engine/Plugins/Experimental/EditorDataStorageFeatures/Source/TedsActorCompatibility/Private/Compatibility/TedsActorWorldPartitionPinnedQueries.cpp
+++ b/Engine/Plugins/Experimental/EditorDataStorageFeatures/Source/TedsActorCompatibility/Private/Compatibility/TedsActorWorldPartitionPinnedQueries.cpp
@@ -28,14 +28,14 @@ void UActorWorldPartitionPinnedDataStorageFactory::RegisterActorAddWorldPartitio
 		Select(
 			TEXT("Actor World Partition Pinned Object to New Column"),
 			FProcessor(EQueryTickPhase::PrePhysics, DataStorage.GetQueryTickGroupName(EQueryTickGroups::SyncExternalToDataStorage)).SetExecutionMode(EExecutionMode::GameThread),
-			[](IQueryContext& Context, RowHandle Row, const FTypedElementUObjectColumn& Object)
+			[](IQueryContext& Context, RowHandle Row, const FTypedElementUObjectColumn& Object, const FTypedElementWorldColumn& WorldColumn)
 			{
 				if (AActor* Actor = Cast<AActor>(Object.Object))
 				{
-					if (Actor->GetLevel() && FLoaderAdapterPinnedActors::SupportsPinning(Actor))
+					if (FLoaderAdapterPinnedActors::SupportsPinning(Actor))
 					{
 						bool bIsPinned = false;
-						if (const UWorld* World = Actor->GetWorld())
+						if (const UWorld* World = WorldColumn.World.Get())
 						{
 							if (const UWorldPartition* WorldPartition = World->GetWorldPartition())
 							{
@@ -47,7 +47,7 @@ void UActorWorldPartitionPinnedDataStorageFactory::RegisterActorAddWorldPartitio
 				}
 			})
 		.Where()
-			.All<FTypedElementActorTag, FTypedElementSyncFromWorldTag>()
+			.All<FTypedElementActorTag, FTypedElementSyncFromWorldTag, FLevelColumn>()
 			.None<FWorldPartitionPinnedColumn>()
 		.Compile()
 	);
@@ -61,11 +61,11 @@ void UActorWorldPartitionPinnedDataStorageFactory::RegisterActorWorldPartitionPi
 		Select(
 			TEXT("Actor World Partition Pinned Object to Column"),
 			FProcessor(EQueryTickPhase::PrePhysics, DataStorage.GetQueryTickGroupName(EQueryTickGroups::SyncExternalToDataStorage)).SetExecutionMode(EExecutionMode::GameThread),
-			[](IQueryContext&, RowHandle, const FTypedElementUObjectColumn& Object, FWorldPartitionPinnedColumn& PinnedColumn)
+			[](IQueryContext&, RowHandle, const FTypedElementUObjectColumn& Object, FWorldPartitionPinnedColumn& PinnedColumn, const FTypedElementWorldColumn& WorldColumn)
 			{
 				if (const AActor* Actor = Cast<AActor>(Object.Object))
 				{
-					if (const UWorld* World = Actor->GetWorld())
+					if (const UWorld* World = WorldColumn.World.Get())
 					{
 						if (const UWorldPartition* WorldPartition = World->GetWorldPartition())
 						{
