[TedsOutliner] Double-click row to frame selection in viewport

#Synced-CL 52470154

diff --git a/Engine/Plugins/Experimental/EditorDataStorageFeatures/Source/TedsOutliner/Private/TedsOutlinerMode.cpp b/Engine/Plugins/Experimental/EditorDataStorageFeatures/Source/TedsOutliner/Private/TedsOutlinerMode.cpp
--- a/Engine/Plugins/Experimental/EditorDataStorageFeatures/Source/TedsOutliner/Private/TedsOutlinerMode.cpp
+++ b/Engine/Plugins/Experimental/EditorDataStorageFeatures/Source/TedsOutliner/Private/TedsOutlinerMode.cpp
@@ -2,6 +2,7 @@
 
 #include "TedsOutlinerMode.h"
 
+#include "Editor.h"
 #include "Elements/Columns/TypedElementCompatibilityColumns.h"
 #include "Elements/Columns/TypedElementSelectionColumns.h"
 #include "FolderTreeItem.h"
@@ -15,6 +16,8 @@
 #include "ToolMenus.h"
 #include "DragAndDrop/Widgets/WidgetDropHandler.h"
 #include "Elements/Columns/TypedElementHandleColumn.h"
+#include "Elements/Framework/TypedElementRegistry.h"
+#include "LevelEditorViewport.h"
 #include "SceneOutlinerFilters.h"
 
 #include UE_INLINE_GENERATED_CPP_BY_NAME(TedsOutlinerMode)
@@ -137,6 +140,18 @@ void FTedsOutlinerMode::SynchronizeSelection()
 	OnSelectionChanged();
 }
 
+void FTedsOutlinerMode::OnItemDoubleClick(FSceneOutlinerTreeItemPtr Item)
+{
+	if (Item->CanInteract())
+	{
+		if (UTypedElementSelectionSet* SelectionSet = TedsOutlinerImpl->GetSelectionSet())
+		{
+			constexpr bool bActiveViewportOnly = false;
+			GEditor->MoveViewportCamerasToElement(SelectionSet, bActiveViewportOnly);
+		}
+	}
+}
+
 void FTedsOutlinerMode::OnItemSelectionChanged(FSceneOutlinerTreeItemPtr Item, ESelectInfo::Type SelectionType, const FSceneOutlinerItemSelection& Selection)
 {
 	if(SelectionType == ESelectInfo::Direct)
diff --git a/Engine/Plugins/Experimental/EditorDataStorageFeatures/Source/TedsOutliner/Public/TedsOutlinerMode.h b/Engine/Plugins/Experimental/EditorDataStorageFeatures/Source/TedsOutliner/Public/TedsOutlinerMode.h
--- a/Engine/Plugins/Experimental/EditorDataStorageFeatures/Source/TedsOutliner/Public/TedsOutlinerMode.h
+++ b/Engine/Plugins/Experimental/EditorDataStorageFeatures/Source/TedsOutliner/Public/TedsOutlinerMode.h
@@ -36,6 +36,7 @@ public:
 	UE_API virtual FSceneOutlinerDragValidationInfo ValidateDrop(const ISceneOutlinerTreeItem& DropTarget, const FSceneOutlinerDragDropPayload& Payload) const override;
 	UE_API virtual void OnDrop(ISceneOutlinerTreeItem& DropTarget, const FSceneOutlinerDragDropPayload& Payload, const FSceneOutlinerDragValidationInfo& ValidationInfo) const override;
 	UE_API virtual TSharedPtr<SWidget> CreateContextMenu() override;
+	UE_API virtual void OnItemDoubleClick(FSceneOutlinerTreeItemPtr Item) override;
 	UE_API virtual void Tick() override;
 	UE_API virtual bool ShowViewButton() const override;
 	UE_API virtual void InitializeViewMenuExtender(TSharedPtr<FExtender> Extender) override;
