[Teds Outliner] Allow Show Settings to be saved and loaded

#Synced-CL 52538476

diff --git a/Engine/Plugins/Experimental/EditorDataStorageFeatures/Source/TedsOutliner/Private/TedsOutlinerFilter.cpp b/Engine/Plugins/Experimental/EditorDataStorageFeatures/Source/TedsOutliner/Private/TedsOutlinerFilter.cpp
--- a/Engine/Plugins/Experimental/EditorDataStorageFeatures/Source/TedsOutliner/Private/TedsOutlinerFilter.cpp
+++ b/Engine/Plugins/Experimental/EditorDataStorageFeatures/Source/TedsOutliner/Private/TedsOutlinerFilter.cpp
@@ -15,38 +15,42 @@
 namespace UE::Editor::Outliner
 {
 FTedsOutlinerFilter::FTedsOutlinerFilter(const FName& InFilterName, const FText& InFilterDisplayName, const FText& InFilterToolTip, const FName& InFilterIconName, 
-	const TSharedPtr<FFilterCategory>& InCategory, const QueryHandle& InFilterQuery, const bool bInteractiveFilter)
-	: FTedsFilterBase(InFilterName, InFilterDisplayName, InFilterToolTip, InFilterIconName, InCategory, InFilterQuery)
+	const TSharedPtr<FFilterCategory>& InCategory, const QueryHandle& InFilterQuery, const bool bInteractiveFilter, const bool bActiveByDefault)
+	: FTedsFilterBase(InFilterName, InFilterDisplayName, InFilterToolTip, InFilterIconName, InCategory, InFilterQuery, bActiveByDefault)
 	, bInteractiveFilter(bInteractiveFilter)
 {}
 
 FTedsOutlinerFilter::FTedsOutlinerFilter(const FName& InFilterName, const FText& InFilterDisplayName, const QueryHandle& InFilterQuery,
-	const bool bInteractiveFilter)
+	const bool bInteractiveFilter, const bool bActiveByDefault)
 	: FTedsOutlinerFilter(InFilterName, InFilterDisplayName, FText::FromString(InFilterName.ToString()), FName(), 
-		nullptr, InFilterQuery, bInteractiveFilter)
+		nullptr, InFilterQuery, bInteractiveFilter, bActiveByDefault)
 {}
 
 FTedsOutlinerFilter::FTedsOutlinerFilter(const FName& InFilterName, const FText& InFilterDisplayName, const FText& InFilterToolTip, 
 	const FName& InFilterIconName, const TSharedPtr<FFilterCategory>& InCategory, const TConstQueryFunction<bool>& InFilterQuery,
-	const bool bInteractiveFilter)
-	: FTedsFilterBase(InFilterName, InFilterDisplayName, InFilterToolTip, InFilterIconName, InCategory, InFilterQuery)
+	const bool bInteractiveFilter, const bool bActiveByDefault)
+	: FTedsFilterBase(InFilterName, InFilterDisplayName, InFilterToolTip, InFilterIconName, InCategory, InFilterQuery, bActiveByDefault)
 	, bInteractiveFilter(bInteractiveFilter)
 {}
 
 FTedsOutlinerFilter::FTedsOutlinerFilter(const FName& InFilterName, const FText& InFilterDisplayName, const TConstQueryFunction<bool>& InFilterQuery,
-	const bool bInteractiveFilter)
+	const bool bInteractiveFilter, const bool bActiveByDefault)
 	: FTedsOutlinerFilter(InFilterName, InFilterDisplayName, FText::FromString(InFilterName.ToString()), FName(), 
-		nullptr, InFilterQuery, bInteractiveFilter)
+		nullptr, InFilterQuery, bInteractiveFilter, bActiveByDefault)
 {}
 
-FTedsOutlinerFilter::FTedsOutlinerFilter(const UClass* InClass, const TSharedPtr<FFilterCategory>& InCategory, const bool bInteractiveFilter)
-	: FTedsFilterBase(InClass, InCategory)
+FTedsOutlinerFilter::FTedsOutlinerFilter(const UClass* InClass, const TSharedPtr<FFilterCategory>& InCategory, const bool bInteractiveFilter, const bool bActiveByDefault)
+	: FTedsFilterBase(InClass, InCategory, bActiveByDefault)
 	, bInteractiveFilter(bInteractiveFilter)
 {}
 
 void FTedsOutlinerFilter::SetSceneOutlinerImpl(const TSharedPtr<FTedsOutlinerImpl>& InTedsOutlinerImpl)
 {
 	TedsOutlinerImpl = InTedsOutlinerImpl;
+	if(IsActiveByDefault())
+	{
+		ActiveStateChanged(true);
+	}
 }
 
 void FTedsOutlinerFilter::ActiveStateChanged(bool bActive)
diff --git a/Engine/Plugins/Experimental/EditorDataStorageFeatures/Source/TedsOutliner/Private/TedsOutlinerMode.cpp b/Engine/Plugins/Experimental/EditorDataStorageFeatures/Source/TedsOutliner/Private/TedsOutlinerMode.cpp
--- a/Engine/Plugins/Experimental/EditorDataStorageFeatures/Source/TedsOutliner/Private/TedsOutlinerMode.cpp
+++ b/Engine/Plugins/Experimental/EditorDataStorageFeatures/Source/TedsOutliner/Private/TedsOutlinerMode.cpp
@@ -19,6 +19,7 @@
 #include "Elements/Framework/TypedElementRegistry.h"
 #include "LevelEditorViewport.h"
 #include "SceneOutlinerFilters.h"
+#include "TedsOutlinerModeSettings.h"
 
 #include UE_INLINE_GENERATED_CPP_BY_NAME(TedsOutlinerMode)
 
@@ -47,7 +48,22 @@ namespace UE::Editor::Outliner
 
 FTedsOutlinerMode::FTedsOutlinerMode(const FTedsOutlinerParams& InParams)
 	: ISceneOutlinerMode(InParams.SceneOutliner)
-{	
+{
+	UTedsOutlinerConfig::Initialize();
+	UTedsOutlinerConfig::Get()->LoadEditorConfig();
+
+	// Get a MutableConfig here to force create a config for the current outliner if it doesn't exist
+	const FTedsOutlinerModeConfig* SavedSettings = GetMutableConfig();
+
+	// Create a local struct to use the default values if this outliner doesn't want to save configs
+	FTedsOutlinerModeConfig LocalSettings;
+
+	// If this outliner doesn't want to save config (OutlinerIdentifier is empty, use the defaults)
+	if (SavedSettings)
+	{
+		LocalSettings = *SavedSettings;
+	}
+		
 	TedsOutlinerImpl = MakeShared<FTedsOutlinerImpl>(InParams, this, false /* bHybridMode */);
 	TedsOutlinerImpl->Init();
 	TedsOutlinerImpl->OnSelectionChanged().AddRaw(this, &FTedsOutlinerMode::OnSelectionChanged);
@@ -59,17 +75,34 @@ FTedsOutlinerMode::FTedsOutlinerMode(const FTedsOutlinerParams& InParams)
 		if (ShowFilter)
 		{
 			ShowFilter->SetSceneOutlinerImpl(TedsOutlinerImpl);
+			const FName ShowFilterName = FName(*ShowFilter->GetName());
+			bool bSettingActive = ShowFilter->IsActiveByDefault();
+			if (const bool* bSettingActivePtr = LocalSettings.ShowSettingsActive.Find(ShowFilterName))
+			{
+				if (const bool bSavedSettingActive = *bSettingActivePtr; bSettingActive != bSavedSettingActive)
+				{
+					bSettingActive = bSavedSettingActive;
+					ShowFilter->ActiveStateChanged(bSavedSettingActive);
+				}
+			}
+			
 			FSceneOutlinerFilterInfo SceneOutlinerFilterInfo(ShowFilter->GetDisplayName(),
-				ShowFilter->GetToolTipText(), ShowFilter->IsActive(),
+				ShowFilter->GetToolTipText(), bSettingActive,
 				FCreateSceneOutlinerFilter::CreateLambda([]()
 				{
 					return MakeShared<TSceneOutlinerFilter<const ISceneOutlinerTreeItem>>(FSceneOutlinerFilter::EDefaultBehaviour::Pass);
 				}));
-			SceneOutlinerFilterInfo.OnToggle().AddLambda([ShowFilter](bool bIsActive)
+
+			SceneOutlinerFilterInfo.OnToggle().AddLambda([this, ShowFilter, ShowFilterName](bool bIsActive)
+			{
+				ShowFilter->ActiveStateChanged(bIsActive);
+				if (FTedsOutlinerModeConfig* Settings = GetMutableConfig())
 				{
-					ShowFilter->ActiveStateChanged(bIsActive);
-				});
-			FilterInfoMap.Add(FName(ShowFilter->GetName()), SceneOutlinerFilterInfo);
+					Settings->ShowSettingsActive.FindOrAdd(ShowFilterName) = bIsActive;
+					SaveConfig();
+				}
+			});
+			FilterInfoMap.Add(ShowFilterName, SceneOutlinerFilterInfo);
 		}
 	}
 	OnInitializeViewMenuExtender = InParams.OnInitializeViewMenuExtender;
@@ -354,6 +387,36 @@ bool FTedsOutlinerMode::ShowFilterOptions() const
 	// to give an empty ShowOptionsFilters in the params if you do not want to show any settings.
 	return bShowViewButton;
 }
+
+struct FTedsOutlinerModeConfig* FTedsOutlinerMode::GetMutableConfig()
+{
+	FName OutlinerIdentifier = SceneOutliner->GetOutlinerIdentifier();
+
+	if (OutlinerIdentifier.IsNone())
+	{
+		return nullptr;
+	}
+
+	return &UTedsOutlinerConfig::Get()->TedsOutliners.FindOrAdd(OutlinerIdentifier);
+}
+
+
+const FTedsOutlinerModeConfig* FTedsOutlinerMode::GetConstConfig() const
+{
+	FName OutlinerIdentifier = SceneOutliner->GetOutlinerIdentifier();
+
+	if (OutlinerIdentifier.IsNone())
+	{
+		return nullptr;
+	}
+
+	return UTedsOutlinerConfig::Get()->TedsOutliners.Find(OutlinerIdentifier);
+}
+
+void FTedsOutlinerMode::SaveConfig()
+{
+	UTedsOutlinerConfig::Get()->SaveEditorConfig();
+}
 	
 } // namespace UE::Editor::Outliner
 
diff --git a/Engine/Plugins/Experimental/EditorDataStorageFeatures/Source/TedsOutliner/Private/TedsOutlinerModeSettings.cpp b/Engine/Plugins/Experimental/EditorDataStorageFeatures/Source/TedsOutliner/Private/TedsOutlinerModeSettings.cpp
new file mode 100644
--- /dev/null
+++ b/Engine/Plugins/Experimental/EditorDataStorageFeatures/Source/TedsOutliner/Private/TedsOutlinerModeSettings.cpp
@@ -0,0 +1,16 @@
+// Copyright Epic Games, Inc. All Rights Reserved.
+
+#include "TedsOutlinerModeSettings.h"
+
+#include UE_INLINE_GENERATED_CPP_BY_NAME(TedsOutlinerModeSettings)
+
+TObjectPtr<UTedsOutlinerConfig> UTedsOutlinerConfig::Instance = nullptr;
+
+void UTedsOutlinerConfig::Initialize()
+{
+	if(!Instance)
+	{
+		Instance = NewObject<UTedsOutlinerConfig>(); 
+		Instance->AddToRoot();
+	}
+}
diff --git a/Engine/Plugins/Experimental/EditorDataStorageFeatures/Source/TedsOutliner/Private/TedsOutlinerModeSettings.h b/Engine/Plugins/Experimental/EditorDataStorageFeatures/Source/TedsOutliner/Private/TedsOutlinerModeSettings.h
new file mode 100644
--- /dev/null
+++ b/Engine/Plugins/Experimental/EditorDataStorageFeatures/Source/TedsOutliner/Private/TedsOutlinerModeSettings.h
@@ -0,0 +1,42 @@
+// Copyright Epic Games, Inc. All Rights Reserved.
+
+#pragma once
+
+#include "Containers/Map.h"
+#include "EditorConfigBase.h"
+#include "UObject/ObjectPtr.h"
+#include "UObject/UObjectGlobals.h"
+
+#include "TedsOutlinerModeSettings.generated.h"
+
+class FName;
+class UObject;
+
+USTRUCT()
+struct FTedsOutlinerModeConfig
+{
+	GENERATED_BODY()
+
+public:
+	/** Map containing the state of the given show settings if changed from their default */
+	UPROPERTY()
+	TMap<FName, bool> ShowSettingsActive = {};
+};
+
+UCLASS(EditorConfig="TedsOutlinerMode")
+class UTedsOutlinerConfig : public UEditorConfigBase
+{
+	GENERATED_BODY()
+	
+public:
+
+	static void Initialize();
+	static UTedsOutlinerConfig* Get() { return Instance; }
+
+	UPROPERTY(meta=(EditorConfig))
+	TMap<FName, FTedsOutlinerModeConfig> TedsOutliners;
+
+private:
+
+	static TObjectPtr<UTedsOutlinerConfig> Instance;
+};
\ No newline at end of file
diff --git a/Engine/Plugins/Experimental/EditorDataStorageFeatures/Source/TedsOutliner/Private/Widgets/SceneOutlinerWidget.cpp b/Engine/Plugins/Experimental/EditorDataStorageFeatures/Source/TedsOutliner/Private/Widgets/SceneOutlinerWidget.cpp
--- a/Engine/Plugins/Experimental/EditorDataStorageFeatures/Source/TedsOutliner/Private/Widgets/SceneOutlinerWidget.cpp
+++ b/Engine/Plugins/Experimental/EditorDataStorageFeatures/Source/TedsOutliner/Private/Widgets/SceneOutlinerWidget.cpp
@@ -315,7 +315,9 @@ void FTedsSceneOutlinerWidgetConstructor::GetShowOptionsFilters(UE::Editor::Data
 			Select()
 			.Where()
 			.None<FActorComponentTypeTag>()
-			.Compile())
+			.Compile()), 
+		true,
+		true
 		));
 
 	// Hide Unloaded Actors
diff --git a/Engine/Plugins/Experimental/EditorDataStorageFeatures/Source/TedsOutliner/Public/TedsOutlinerFilter.h b/Engine/Plugins/Experimental/EditorDataStorageFeatures/Source/TedsOutliner/Public/TedsOutlinerFilter.h
--- a/Engine/Plugins/Experimental/EditorDataStorageFeatures/Source/TedsOutliner/Public/TedsOutlinerFilter.h
+++ b/Engine/Plugins/Experimental/EditorDataStorageFeatures/Source/TedsOutliner/Public/TedsOutlinerFilter.h
@@ -30,9 +30,10 @@ namespace UE::Editor::Outliner
 		 * @param InCategory Group to place the filter in the Add Filter Menu
 		 * @param InFilterQuery QueryHandle of a registered query to filter for (must be the same syntax as the Outliner Query)
 		 * @param bInteractiveFilter Whether this filter is visible in the Add Filter Menu and toggleable or is hidden and enabled by default
+		 * @param bActiveByDefault Whether this filter should be active by default and turned on when the SceneOutliner is Init
 		 */
 		TEDSOUTLINER_API FTedsOutlinerFilter(const FName& InFilterName, const FText& InFilterDisplayName, const FText& InFilterToolTip, const FName& InFilterIconName, 
-			const TSharedPtr<FFilterCategory>& InCategory, const QueryHandle& InFilterQuery, const bool bInteractiveFilter = true);
+			const TSharedPtr<FFilterCategory>& InCategory, const QueryHandle& InFilterQuery, const bool bInteractiveFilter = true, const bool bActiveByDefault = false);
 
 		/**
 		 * Constructor for a TEDS Outliner Filter
@@ -40,9 +41,10 @@ namespace UE::Editor::Outliner
 		 * @param InFilterDisplayName Display Name to show on the Add Filter Menu and toggle button
 		 * @param InFilterQuery QueryHandle of a registered query to filter for (must be the same syntax as the Outliner Query)
 		 * @param bInteractiveFilter Whether this filter is visible in the Add Filter Menu and toggleable or is hidden and enabled by default
+		 * @param bActiveByDefault Whether this filter should be active by default and turned on when the SceneOutliner is Init
 		 */
 		TEDSOUTLINER_API FTedsOutlinerFilter(const FName& InFilterName, const FText& InFilterDisplayName, const QueryHandle& InFilterQuery,
-			const bool bInteractiveFilter = true);
+			const bool bInteractiveFilter = true, const bool bActiveByDefault = false);
 
 		/**
 		 * Constructor for a TEDS Outliner Filter
@@ -53,9 +55,10 @@ namespace UE::Editor::Outliner
 		 * @param InCategory Group to place the filter in the Add Filter Menu
 		 * @param InFilterQuery QueryFunction to filter for, Batching should be used for improved performance if possible
 		 * @param bInteractiveFilter Whether this filter is visible in the Add Filter Menu and toggleable or is hidden and enabled by default
+		 * @param bActiveByDefault Whether this filter should be active by default and turned on when the SceneOutliner is Init
 		 */
 		TEDSOUTLINER_API FTedsOutlinerFilter(const FName& InFilterName, const FText& InFilterDisplayName, const FText& InFilterToolTip, const FName& InFilterIconName,
-			const TSharedPtr<FFilterCategory>& InCategory, const TConstQueryFunction<bool>& InFilterQuery, const bool bInteractiveFilter = true);
+			const TSharedPtr<FFilterCategory>& InCategory, const TConstQueryFunction<bool>& InFilterQuery, const bool bInteractiveFilter = true, const bool bActiveByDefault = false);
 
 		/**
 		 * Constructor for a TEDS Outliner Filter
@@ -63,9 +66,10 @@ namespace UE::Editor::Outliner
 		 * @param InFilterDisplayName Display Name to show on the Add Filter Menu and toggle button
 		 * @param InFilterQuery QueryFunction to filter for, Batching should be used for improved performance if possible
 		 * @param bInteractiveFilter Whether this filter is visible in the Add Filter Menu and toggleable or is hidden and enabled by default
+		 * @param bActiveByDefault Whether this filter should be active by default and turned on when the SceneOutliner is Init
 		 */
 		TEDSOUTLINER_API FTedsOutlinerFilter(const FName& InFilterName, const FText& InFilterDisplayName, const TConstQueryFunction<bool>& InFilterQuery,
-			const bool bInteractiveFilter = true);
+			const bool bInteractiveFilter = true, const bool bActiveByDefault = false);
 		
 		/**
 		 * Constructor for a TEDS Outliner Filter
@@ -78,12 +82,13 @@ namespace UE::Editor::Outliner
 		 * @param InCategory Group to place the filter in the Add Filter Menu
 		 * @param InFilterQuery Function callback that checks the value of a column(s) and returns a bool of whether that row passes the filter
 		 * @param bInteractiveFilter Whether this filter is visible in the Add Filter Menu and toggleable or is hidden and enabled by default
+		 * @param bActiveByDefault Whether this filter should be active by default and turned on when the SceneOutliner is Init
 		 */
 		template< typename FilterFunction UE_REQUIRES(
 			!std::is_same_v<std::remove_cvref_t<FilterFunction>, QueryHandle> &&
 			!std::is_same_v<std::remove_cvref_t<FilterFunction>, Queries::TConstQueryFunction<bool>>) >
 		FTedsOutlinerFilter(const FName& InFilterName, const FText& InFilterDisplayName, const FText& InFilterToolTip, const FName& InFilterIconName,
-			const TSharedPtr<FFilterCategory>& InCategory, FilterFunction&& InFilterQuery, const bool bInteractiveFilter = true)
+			const TSharedPtr<FFilterCategory>& InCategory, FilterFunction&& InFilterQuery, const bool bInteractiveFilter = true, const bool bActiveByDefault = false)
 			: FTedsOutlinerFilter(
 				InFilterName, 
 				InFilterDisplayName, 
@@ -91,7 +96,8 @@ namespace UE::Editor::Outliner
 				InFilterIconName, 
 				InCategory,
 				Queries::BuildConstQueryFunction<bool>(Forward<FilterFunction>(InFilterQuery)),
-				bInteractiveFilter)
+				bInteractiveFilter,
+				bActiveByDefault)
 		{}
 
 		/**
@@ -102,12 +108,13 @@ namespace UE::Editor::Outliner
 		 * @param InFilterDisplayName Display Name to show on the Add Filter Menu and toggle button
 		 * @param InFilterQuery Function callback that checks the value of a column(s) and returns a bool of whether that row passes the filter
 		 * @param bInteractiveFilter Whether this filter is visible in the Add Filter Menu and toggleable or is hidden and enabled by default
+		 * @param bActiveByDefault Whether this filter should be active by default and turned on when the SceneOutliner is Init
 		 */
 		template< typename FilterFunction UE_REQUIRES(
 			!std::is_same_v<std::remove_cvref_t<FilterFunction>, QueryHandle> &&
 			!std::is_same_v<std::remove_cvref_t<FilterFunction>, Queries::TConstQueryFunction<bool>>) >
 		FTedsOutlinerFilter(const FName& InFilterName, const FText& InFilterDisplayName, FilterFunction&& InFilterQuery,
-			const bool bInteractiveFilter = true)
+			const bool bInteractiveFilter = true, const bool bActiveByDefault = false)
 			: FTedsOutlinerFilter(
 				InFilterName, 
 				InFilterDisplayName, 
@@ -115,7 +122,8 @@ namespace UE::Editor::Outliner
 				FName(), 
 				nullptr,
 				Queries::BuildConstQueryFunction<bool>(Forward<FilterFunction>(InFilterQuery)),
-				bInteractiveFilter)
+				bInteractiveFilter,
+				bActiveByDefault)
 		{}
 
 		/**
@@ -123,8 +131,10 @@ namespace UE::Editor::Outliner
 		 * @param InClass UClass to filter for on the TypeInfo Column
 		 * @param InCategory Group to place the filter in the Add Filter Menu
 		 * @param bInteractiveFilter Whether this filter is visible in the Add Filter Menu and toggleable or is hidden and enabled by default
+		 * @param bActiveByDefault Whether this filter should be active by default and turned on when the SceneOutliner is Init
 		 */
-		TEDSOUTLINER_API FTedsOutlinerFilter(const UClass* InClass, const TSharedPtr<FFilterCategory>& InCategory = nullptr, const bool bInteractiveFilter = true);
+		TEDSOUTLINER_API FTedsOutlinerFilter(const UClass* InClass, const TSharedPtr<FFilterCategory>& InCategory = nullptr,
+			const bool bInteractiveFilter = true, const bool bActiveByDefault = false);
 
 		bool IsInteractiveFilter() const { return bInteractiveFilter; }
 
@@ -139,7 +149,7 @@ namespace UE::Editor::Outliner
 		void SetSceneOutlinerImpl(const TSharedPtr<FTedsOutlinerImpl>& InTedsOutlinerImpl);
 		
 		const bool bInteractiveFilter;
-
+		
 		TWeakPtr<FTedsOutlinerImpl> TedsOutlinerImpl;
 	};
 } // namespace UE::Editor::Outliner
diff --git a/Engine/Plugins/Experimental/EditorDataStorageFeatures/Source/TedsOutliner/Public/TedsOutlinerMode.h b/Engine/Plugins/Experimental/EditorDataStorageFeatures/Source/TedsOutliner/Public/TedsOutlinerMode.h
--- a/Engine/Plugins/Experimental/EditorDataStorageFeatures/Source/TedsOutliner/Public/TedsOutlinerMode.h
+++ b/Engine/Plugins/Experimental/EditorDataStorageFeatures/Source/TedsOutliner/Public/TedsOutlinerMode.h
@@ -12,6 +12,8 @@
 
 #define UE_API TEDSOUTLINER_API
 
+struct FTedsOutlinerModeConfig;
+
 namespace UE::Editor::Outliner
 {
 /*
@@ -49,6 +51,24 @@ protected:
 	// Called by TedsOutlinerImpl when the selection in TEDS changes
 	UE_API void OnSelectionChanged();
 
+private:
+	/**
+	 * Get a mutable version of the TedsOutliner config for setting values.
+	 * @returns		The config for this TedsOutliner.
+	 * @note		If OutlinerIdentifier is not set for this outliner, it is not possible to store settings.
+	 */
+	struct FTedsOutlinerModeConfig* GetMutableConfig();
+
+	/**
+	 * Get a const version of the TedsOutliner config for getting values.
+	 * @returns		The config for this TedsOutliner.
+	 * @note		If OutlinerIdentifier is not set for this outliner, it is not possible to store settings.
+	 */
+	const FTedsOutlinerModeConfig* GetConstConfig() const;
+
+	/** Save the config for this TedsOutliner */
+	void SaveConfig();
+
 protected:
 	bool bShowViewButton = false;
 	// Delegate for extending the View Menu with options (World picker, frame on selection, etc.)
diff --git a/Engine/Plugins/Experimental/EditorDataStorageFeatures/Source/TedsOutliner/TedsOutliner.Build.cs b/Engine/Plugins/Experimental/EditorDataStorageFeatures/Source/TedsOutliner/TedsOutliner.Build.cs
--- a/Engine/Plugins/Experimental/EditorDataStorageFeatures/Source/TedsOutliner/TedsOutliner.Build.cs
+++ b/Engine/Plugins/Experimental/EditorDataStorageFeatures/Source/TedsOutliner/TedsOutliner.Build.cs
@@ -21,6 +21,7 @@ public class TedsOutliner : ModuleRules
 				{
 					"Core",
 					"CoreUObject",
+					"EditorConfig",
 					"EditorFramework",
 					"Engine",
 					"SceneOutliner",
diff --git a/Engine/Plugins/Experimental/EditorDataStorageFeatures/Source/TedsTableViewer/Private/TedsFilter.cpp b/Engine/Plugins/Experimental/EditorDataStorageFeatures/Source/TedsTableViewer/Private/TedsFilter.cpp
--- a/Engine/Plugins/Experimental/EditorDataStorageFeatures/Source/TedsTableViewer/Private/TedsFilter.cpp
+++ b/Engine/Plugins/Experimental/EditorDataStorageFeatures/Source/TedsTableViewer/Private/TedsFilter.cpp
@@ -11,6 +11,10 @@ namespace UE::Editor::DataStorage
 	void FTedsFilter::SetTedsFilterBar(const TSharedPtr<STedsFilterBar>& InTedsFilterBar)
 	{
 		TedsFilterBar = InTedsFilterBar;
+		if(IsActiveByDefault())
+		{
+			ActiveStateChanged(true);
+		}
 	}
 
 	void FTedsFilter::ActiveStateChanged(bool bActive)
diff --git a/Engine/Plugins/Experimental/EditorDataStorageFeatures/Source/TedsTableViewer/Public/TedsFilter.inl b/Engine/Plugins/Experimental/EditorDataStorageFeatures/Source/TedsTableViewer/Public/TedsFilter.inl
--- a/Engine/Plugins/Experimental/EditorDataStorageFeatures/Source/TedsTableViewer/Public/TedsFilter.inl
+++ b/Engine/Plugins/Experimental/EditorDataStorageFeatures/Source/TedsTableViewer/Public/TedsFilter.inl
@@ -43,17 +43,19 @@ namespace UE::Editor::DataStorage
 		 * @param InFilterIconName Name of an icon in the App Style Set to display in the Add Filter Menu
 		 * @param InCategory Group to place the filter in the Add Filter Menu
 		 * @param InFilterQuery QueryHandle of a registered query to filter for (must be the same syntax as any other Query used in the viewer)
+		 * @param bActiveByDefault Whether this filter should be active by default and can be implemented per-inherited filter type
 		 */
 		FTedsFilterBase(const FName& InFilterName, const FText& InFilterDisplayName, const FText& InFilterToolTip, const FName& InFilterIconName, 
-			const TSharedPtr<FFilterCategory>& InCategory, const QueryHandle& InFilterQuery);
+			const TSharedPtr<FFilterCategory>& InCategory, const QueryHandle& InFilterQuery, const bool bActiveByDefault = false);
 
 		/**
 		 * Constructor for a TEDS Filter
 		 * @param InFilterName FName Identifier of the Filter
 		 * @param InFilterDisplayName Display Name to show on the Add Filter Menu and toggle button
 		 * @param InFilterQuery QueryHandle of a registered query to filter for (must be the same syntax as any other Query used in the viewer)
+		 * @param bActiveByDefault Whether this filter should be active by default and can be implemented per-inherited filter type
 		 */
-		FTedsFilterBase(const FName& InFilterName, const FText& InFilterDisplayName, const QueryHandle& InFilterQuery);
+		FTedsFilterBase(const FName& InFilterName, const FText& InFilterDisplayName, const QueryHandle& InFilterQuery, const bool bActiveByDefault = false);
 
 		/**
 		 * Constructor for a TEDS Filter
@@ -63,17 +65,20 @@ namespace UE::Editor::DataStorage
 		 * @param InFilterIconName Name of an icon in the App Style Set to display in the Add Filter Menu
 		 * @param InCategory Group to place the filter in the Add Filter Menu
 		 * @param InFilterQuery QueryFunction to filter for, Batching should be used for improved performance if possible
+		 * @param bActiveByDefault Whether this filter should be active by default and can be implemented per-inherited filter type
 		 */
 		FTedsFilterBase(const FName& InFilterName, const FText& InFilterDisplayName, const FText& InFilterToolTip, const FName& InFilterIconName,
-			const TSharedPtr<FFilterCategory>& InCategory, const Queries::TConstQueryFunction<bool>& InFilterQuery);
+			const TSharedPtr<FFilterCategory>& InCategory, const Queries::TConstQueryFunction<bool>& InFilterQuery, const bool bActiveByDefault = false);
 
 		/**
 		 * Constructor for a TEDS Filter
 		 * @param InFilterName FName Identifier of the Filter
 		 * @param InFilterDisplayName Display Name to show on the Add Filter Menu and toggle button
 		 * @param InFilterQuery QueryFunction to filter for, Batching should be used for improved performance if possible
+		 * @param bActiveByDefault Whether this filter should be active by default and can be implemented per-inherited filter type
 		 */
-		FTedsFilterBase(const FName& InFilterName, const FText& InFilterDisplayName, const Queries::TConstQueryFunction<bool>& InFilterQuery);
+		FTedsFilterBase(const FName& InFilterName, const FText& InFilterDisplayName, const Queries::TConstQueryFunction<bool>& InFilterQuery,
+			const bool bActiveByDefault = false);
 		
 		/**
 		 * Constructor for a TEDS Filter
@@ -85,19 +90,21 @@ namespace UE::Editor::DataStorage
 		 * @param InFilterIconName Name of an icon in the App Style Set to display in the Add Filter Menu
 		 * @param InCategory Group to place the filter in the Add Filter Menu
 		 * @param InFilterQuery Function callback that checks the value of a column(s) and returns a bool of whether that row passes the filter
+		 * @param bActiveByDefault Whether this filter should be active by default and can be implemented per-inherited filter type
 		 */
 		template< typename FilterFunction UE_REQUIRES(
 			!std::is_same_v<std::remove_cvref_t<FilterFunction>, QueryHandle> &&
 			!std::is_same_v<std::remove_cvref_t<FilterFunction>, Queries::TConstQueryFunction<bool>>) >
 		FTedsFilterBase(const FName& InFilterName, const FText& InFilterDisplayName, const FText& InFilterToolTip, const FName& InFilterIconName,
-			const TSharedPtr<FFilterCategory>& InCategory, FilterFunction&& InFilterQuery)
+			const TSharedPtr<FFilterCategory>& InCategory, FilterFunction&& InFilterQuery, const bool bActiveByDefault = false)
 			: FTedsFilterBase(
 				InFilterName, 
 				InFilterDisplayName, 
 				InFilterToolTip, 
 				InFilterIconName, 
 				InCategory,
-				Queries::BuildConstQueryFunction<bool>(Forward<FilterFunction>(InFilterQuery)))
+				Queries::BuildConstQueryFunction<bool>(Forward<FilterFunction>(InFilterQuery)),
+				bActiveByDefault)
 		{}
 
 		/**
@@ -107,32 +114,37 @@ namespace UE::Editor::DataStorage
 		 * @param InFilterName FName Identifier of the Filter
 		 * @param InFilterDisplayName Display Name to show on the Add Filter Menu and toggle button
 		 * @param InFilterQuery Function callback that checks the value of a column(s) and returns a bool of whether that row passes the filter
+		 * @param bActiveByDefault Whether this filter should be active by default and can be implemented per-inherited filter type
 		 */
 		template< typename FilterFunction UE_REQUIRES(
 			!std::is_same_v<std::remove_cvref_t<FilterFunction>, QueryHandle> &&
 			!std::is_same_v<std::remove_cvref_t<FilterFunction>, Queries::TConstQueryFunction<bool>>) >
-		FTedsFilterBase(const FName& InFilterName, const FText& InFilterDisplayName, FilterFunction&& InFilterQuery)
+		FTedsFilterBase(const FName& InFilterName, const FText& InFilterDisplayName, FilterFunction&& InFilterQuery, const bool bActiveByDefault = false)
 			: FTedsFilterBase(
 				InFilterName, 
 				InFilterDisplayName, 
 				FText::FromString(InFilterName.ToString()), 
 				FName(), 
 				nullptr,
-				Queries::BuildConstQueryFunction<bool>(Forward<FilterFunction>(InFilterQuery)))
+				Queries::BuildConstQueryFunction<bool>(Forward<FilterFunction>(InFilterQuery)),
+				bActiveByDefault)
 		{}
 
 		/**
 		 * Constructor for a TEDS Class Filter
 		 * @param InClass UClass to filter for on the TypeInfo Column
 		 * @param InCategory Group to place the filter in the Add Filter Menu
+		 * @param bActiveByDefault Whether this filter should be active by default and can be implemented per-inherited filter type
 		 */
-		FTedsFilterBase(const UClass* InClass, const TSharedPtr<FFilterCategory>& InCategory = nullptr);
+		FTedsFilterBase(const UClass* InClass, const TSharedPtr<FFilterCategory>& InCategory = nullptr, const bool bActiveByDefault = false);
 
 		virtual ~FTedsFilterBase() override = default;
 		
 		void SetCategory(const TSharedPtr<FFilterCategory>& InCategory);
         
 		bool IsClassFilter() const { return bIsClassFilter; }
+
+		bool IsActiveByDefault() const { return bActiveByDefault; }
 		
 		FName GetIdentifier() const { return FilterName; }
 
@@ -170,56 +182,62 @@ namespace UE::Editor::DataStorage
 		FName FilterIconName;
 
 		const bool bIsClassFilter;
+		const bool bActiveByDefault;
 
 		TVariant<QueryHandle, Queries::TConstQueryFunction<bool>> FilterQuery;
 	};
 
 	template<typename FilterType>
 	FTedsFilterBase<FilterType>::FTedsFilterBase(const FName& InFilterName, const FText& InFilterDisplayName, const FText& InFilterToolTip, const FName& InFilterIconName, 
-		const TSharedPtr<FFilterCategory>& InCategory, const QueryHandle& InFilterQuery)
+		const TSharedPtr<FFilterCategory>& InCategory, const QueryHandle& InFilterQuery, const bool bActiveByDefault)
 		: FFilterBase<FilterType>(InCategory)
 		, FilterName(InFilterName)
 		, FilterDisplayName(InFilterDisplayName)
 		, FilterToolTip(InFilterToolTip)
 		, FilterIconName(InFilterIconName)
 		, bIsClassFilter(false)
+		, bActiveByDefault(bActiveByDefault)
 	{
 		FilterQuery.Set<QueryHandle>(CheckValidFilterQueryHandle(InFilterQuery) ? InFilterQuery : InvalidQueryHandle);
 	}
 
 	template<typename FilterType>
-	FTedsFilterBase<FilterType>::FTedsFilterBase(const FName& InFilterName, const FText& InFilterDisplayName, const QueryHandle& InFilterQuery)
+	FTedsFilterBase<FilterType>::FTedsFilterBase(const FName& InFilterName, const FText& InFilterDisplayName, const QueryHandle& InFilterQuery, const bool bActiveByDefault)
 		: FTedsFilterBase(InFilterName, InFilterDisplayName, FText::FromString(InFilterName.ToString()), FName(), 
-			nullptr, InFilterQuery)
+			nullptr, InFilterQuery, bActiveByDefault)
 	{}
 
 	template<typename FilterType>
 	FTedsFilterBase<FilterType>::FTedsFilterBase(const FName& InFilterName, const FText& InFilterDisplayName, const FText& InFilterToolTip, 
-		const FName& InFilterIconName, const TSharedPtr<FFilterCategory>& InCategory, const Queries::TConstQueryFunction<bool>& InFilterQuery)
+		const FName& InFilterIconName, const TSharedPtr<FFilterCategory>& InCategory, const Queries::TConstQueryFunction<bool>& InFilterQuery,
+		const bool bActiveByDefault)
 		: FFilterBase<FilterType>(InCategory)
 		, FilterName(InFilterName)
 		, FilterDisplayName(InFilterDisplayName)
 		, FilterToolTip(InFilterToolTip)
 		, FilterIconName(InFilterIconName)
 		, bIsClassFilter(false)
+		, bActiveByDefault(bActiveByDefault)
 	{
 		FilterQuery.Set<Queries::TConstQueryFunction<bool>>(InFilterQuery);
 	}
 
 	template<typename FilterType>
-	FTedsFilterBase<FilterType>::FTedsFilterBase(const FName& InFilterName, const FText& InFilterDisplayName, const Queries::TConstQueryFunction<bool>& InFilterQuery)
+	FTedsFilterBase<FilterType>::FTedsFilterBase(const FName& InFilterName, const FText& InFilterDisplayName, const Queries::TConstQueryFunction<bool>& InFilterQuery,
+		const bool bActiveByDefault)
 		: FTedsFilterBase(InFilterName, InFilterDisplayName, FText::FromString(InFilterName.ToString()), FName(), 
-			nullptr, InFilterQuery)
+			nullptr, InFilterQuery, bActiveByDefault)
 	{}
 
 	template<typename FilterType>
-	FTedsFilterBase<FilterType>::FTedsFilterBase(const UClass* InClass, const TSharedPtr<FFilterCategory>& InCategory)
+	FTedsFilterBase<FilterType>::FTedsFilterBase(const UClass* InClass, const TSharedPtr<FFilterCategory>& InCategory, const bool bActiveByDefault)
 		: FFilterBase<FilterType>(InCategory)
 		, FilterName(InClass->GetFName())
 		, FilterDisplayName(InClass->GetDisplayNameText())
 		, FilterToolTip(FText::Format(LOCTEXT("FilterClassTooltip", "Filter by {0}"), InClass->GetDisplayNameText()))
 		, FilterIconName(FSlateIconFinder::FindIconForClass(InClass).GetStyleName())
 		, bIsClassFilter(true)
+		, bActiveByDefault(bActiveByDefault)
 	{
 		using namespace UE::Editor::DataStorage::Queries;
 		FilterQuery.Set<TConstQueryFunction<bool>>(
