[TEDS] Optimize FTedsAssetData::Initialize() by preventing multiple allocations when the Asset Registry contains a large number of assets.  Use an atomic counter to be able to reserve the needed TArray

#Synced-CL 52496820

diff --git a/Engine/Plugins/Experimental/EditorDataStorageFeatures/Source/TedsAssetData/Private/TedsAssetData.cpp b/Engine/Plugins/Experimental/EditorDataStorageFeatures/Source/TedsAssetData/Private/TedsAssetData.cpp
--- a/Engine/Plugins/Experimental/EditorDataStorageFeatures/Source/TedsAssetData/Private/TedsAssetData.cpp
+++ b/Engine/Plugins/Experimental/EditorDataStorageFeatures/Source/TedsAssetData/Private/TedsAssetData.cpp
@@ -321,27 +321,48 @@ void FTedsAssetData::Initialize()
 	// Init with the data existing at moment in asset registry
 	UE::AssetRegistry::FFiltering::InitializeShouldSkipAsset();
 	TQueue<FAssetData, EQueueMode::Mpsc> CachedAssets;
-	AssetRegistry.EnumerateAllAssets([&CachedAssets](const FAssetData& AssetData)
-		{
-			if (!UE::AssetRegistry::FFiltering::ShouldSkipAsset(AssetData.AssetClassPath, AssetData.PackageFlags) &&
-				ContentBrowserDataUtils::IsPrimaryAsset(AssetData))
+	std::atomic<int32> CachedAssetsCount = 0;
+	{
+		TRACE_CPUPROFILER_EVENT_SCOPE(FTedsAssetData::Initialize_EnumerateAllAssets);
+		AssetRegistry.EnumerateAllAssets([&CachedAssets, &CachedAssetsCount](const FAssetData& AssetData)
 			{
-				CachedAssets.Enqueue(AssetData);
-			}
-			return true;
-		}, UE::AssetRegistry::EEnumerateAssetsFlags::Parallel);
+				if (!UE::AssetRegistry::FFiltering::ShouldSkipAsset(AssetData.AssetClassPath, AssetData.PackageFlags) &&
+					ContentBrowserDataUtils::IsPrimaryAsset(AssetData))
+				{
+					CachedAssets.Enqueue(AssetData);
+					++CachedAssetsCount;
+				}
+				return true;
+			}, UE::AssetRegistry::EEnumerateAssetsFlags::Parallel);
+	}
+
 
-	while (CachedAssets.Dequeue(AddedAssets.AddDefaulted_GetRef()))
 	{
+		TRACE_CPUPROFILER_EVENT_SCOPE(FTedsAssetData::Initialize_DrainQueue);
+		AddedAssets.Reserve(AddedAssets.Num() + CachedAssetsCount);
+		FAssetData AssetDataTemp;
+		while (CachedAssets.Dequeue(AssetDataTemp))
+		{
+			AddedAssets.Add(MoveTemp(AssetDataTemp));
+		}
 	}
 
+
 	TChunkedArray<FName> CachedPaths;
-	AssetRegistry.EnumerateAllCachedPaths([&CachedPaths](FName Name)
-		{
-			CachedPaths.AddElement(Name);
-			return true;
-		});
-	CachedPaths.MoveToLinearArray(AddedPaths);
+	{
+		TRACE_CPUPROFILER_EVENT_SCOPE(FTedsAssetData::Initialize_EnumerateAllCachedPaths);
+		AssetRegistry.EnumerateAllCachedPaths([&CachedPaths](FName Name)
+			{
+				CachedPaths.AddElement(Name);
+				return true;
+			});
+	}
+
+	{
+		TRACE_CPUPROFILER_EVENT_SCOPE(FTedsAssetData::Initialize_CachedPaths.MoveToLinearArray);
+		CachedPaths.MoveToLinearArray(AddedPaths);
+	}
+	
 
 	bIsInitialized = true;
 }
