[TEDS Outliner] Add Selection Framing Settings to the Outliner

#Synced-CL 52541567

diff --git a/Engine/Plugins/Experimental/EditorDataStorageFeatures/Source/TedsOutliner/Private/TedsOutlinerMode.cpp b/Engine/Plugins/Experimental/EditorDataStorageFeatures/Source/TedsOutliner/Private/TedsOutlinerMode.cpp
--- a/Engine/Plugins/Experimental/EditorDataStorageFeatures/Source/TedsOutliner/Private/TedsOutlinerMode.cpp
+++ b/Engine/Plugins/Experimental/EditorDataStorageFeatures/Source/TedsOutliner/Private/TedsOutlinerMode.cpp
@@ -68,6 +68,9 @@ FTedsOutlinerMode::FTedsOutlinerMode(const FTedsOutlinerParams& InParams)
 	TedsOutlinerImpl->Init();
 	TedsOutlinerImpl->OnSelectionChanged().AddRaw(this, &FTedsOutlinerMode::OnSelectionChanged);
 
+	bAlwaysFrameSelection = LocalSettings.bAlwaysFrameSelection;
+	bCollapseOutlinerTreeOnNewSelection = LocalSettings.bCollapseOutlinerTreeOnNewSelection;
+
 	bShowViewButton = InParams.bShowViewButton;
 	// Convert ShowOptionsFilters to FilterInfoMap to appear in the Settings "Show" options menu
 	for (const TSharedPtr<FTedsOutlinerFilter>& ShowFilter : InParams.ShowOptionsFilters)
@@ -162,6 +165,22 @@ void FTedsOutlinerMode::OnSelectionChanged()
 		return false;
 	});
 
+	if (!SceneOutliner->GetSelectedItems().IsEmpty())
+	{
+		if (const FSceneOutlinerTreeItemPtr SelectedItem = SceneOutliner->GetSelectedItems().Last(); SelectedItem.IsValid())
+		{
+			if (bCollapseOutlinerTreeOnNewSelection)
+			{
+				SceneOutliner->CollapseAll();
+			}
+	
+			if (bAlwaysFrameSelection || bCollapseOutlinerTreeOnNewSelection)
+			{
+				SceneOutliner->ScrollItemIntoView(SelectedItem);
+			}
+		}
+	}
+
 	if (FTedsOutlinerSelectionChangeColumn* SelectionChangeColumn = Storage->GetColumn<FTedsOutlinerSelectionChangeColumn>(TedsOutlinerImpl->GetOutlinerRowHandle()))
 	{
 		SelectionChangeColumn->OnSelectionChanged.Broadcast();
@@ -375,6 +394,39 @@ bool FTedsOutlinerMode::ShowViewButton() const
 
 void FTedsOutlinerMode::InitializeViewMenuExtender(TSharedPtr<FExtender> Extender)
 {
+	Extender->AddMenuExtension(SceneOutliner::ExtensionHooks::Show, EExtensionHook::After, nullptr, FMenuExtensionDelegate::CreateLambda([this](FMenuBuilder& MenuBuilder)
+	{
+		MenuBuilder.BeginSection("OutlinerSelectionOptions", LOCTEXT("OptionsHeading", "Options"));
+
+		MenuBuilder.AddMenuEntry(
+			LOCTEXT("AlwaysFrameSelectionLabel", "Always Frame Selection"),
+			LOCTEXT("AlwaysFrameSelectionTooltip", "When enabled, selecting an object in the Viewport also scrolls to that item in the Outliner."),
+			FSlateIcon(),
+			FUIAction(
+				FExecuteAction::CreateRaw(this, &FTedsOutlinerMode::OnToggleAlwaysFrameSelection),
+				FCanExecuteAction(),
+				FIsActionChecked::CreateRaw(this, &FTedsOutlinerMode::ShouldAlwaysFrameSelection)
+			),
+			NAME_None,
+			EUserInterfaceActionType::ToggleButton
+		);
+
+		MenuBuilder.AddMenuEntry(
+			LOCTEXT("CollapseOutlinerTreeOnNewSelectionLabel", "Collapse Hierarchy on Viewport Selection"),
+			LOCTEXT("CollapseOutlinerTreeOnNewSelectionTooltip", "When enabled, all items except the one that was just selected are collapsed."),
+			FSlateIcon(),
+			FUIAction(
+				FExecuteAction::CreateRaw(this, &FTedsOutlinerMode::OnToggleCollapseOutlinerTreeOnNewSelection),
+				FCanExecuteAction(),
+				FIsActionChecked::CreateRaw(this, &FTedsOutlinerMode::CollapseOutlinerTreeOnNewSelection)
+			),
+			NAME_None,
+			EUserInterfaceActionType::ToggleButton
+		);
+
+		MenuBuilder.EndSection();
+	}));
+    		
 	if (OnInitializeViewMenuExtender.IsBound())
 	{
 		OnInitializeViewMenuExtender.Execute(Extender);
@@ -417,6 +469,37 @@ void FTedsOutlinerMode::SaveConfig()
 {
 	UTedsOutlinerConfig::Get()->SaveEditorConfig();
 }
+
+void FTedsOutlinerMode::OnToggleAlwaysFrameSelection()
+{
+	bAlwaysFrameSelection = !bAlwaysFrameSelection;
+		
+	if(FTedsOutlinerModeConfig* Settings = GetMutableConfig())
+	{
+		Settings->bAlwaysFrameSelection = bAlwaysFrameSelection;
+		SaveConfig();
+	}
+}
+
+bool FTedsOutlinerMode::ShouldAlwaysFrameSelection() const
+{
+	return bAlwaysFrameSelection;
+}
+
+void FTedsOutlinerMode::OnToggleCollapseOutlinerTreeOnNewSelection()
+{
+	if (FTedsOutlinerModeConfig* Settings = GetMutableConfig())
+	{
+		Settings->bCollapseOutlinerTreeOnNewSelection = !Settings->bCollapseOutlinerTreeOnNewSelection;
+		bCollapseOutlinerTreeOnNewSelection = Settings->bCollapseOutlinerTreeOnNewSelection;
+		SaveConfig();
+	}
+}
+
+bool FTedsOutlinerMode::CollapseOutlinerTreeOnNewSelection() const
+{
+	return bCollapseOutlinerTreeOnNewSelection;
+}
 	
 } // namespace UE::Editor::Outliner
 
diff --git a/Engine/Plugins/Experimental/EditorDataStorageFeatures/Source/TedsOutliner/Private/TedsOutlinerModeSettings.h b/Engine/Plugins/Experimental/EditorDataStorageFeatures/Source/TedsOutliner/Private/TedsOutlinerModeSettings.h
--- a/Engine/Plugins/Experimental/EditorDataStorageFeatures/Source/TedsOutliner/Private/TedsOutlinerModeSettings.h
+++ b/Engine/Plugins/Experimental/EditorDataStorageFeatures/Source/TedsOutliner/Private/TedsOutlinerModeSettings.h
@@ -21,6 +21,14 @@ public:
 	/** Map containing the state of the given show settings if changed from their default */
 	UPROPERTY()
 	TMap<FName, bool> ShowSettingsActive = {};
+	
+	/** True when the Scene Outliner updates when an actor is selected in the viewport */
+	UPROPERTY()
+	bool bAlwaysFrameSelection = true;
+
+	/** True if we want to collapse Outliner tree on new selection, except for the item that was just selected */
+	UPROPERTY()
+	bool bCollapseOutlinerTreeOnNewSelection = false;
 };
 
 UCLASS(EditorConfig="TedsOutlinerMode")
diff --git a/Engine/Plugins/Experimental/EditorDataStorageFeatures/Source/TedsOutliner/Public/TedsOutlinerMode.h b/Engine/Plugins/Experimental/EditorDataStorageFeatures/Source/TedsOutliner/Public/TedsOutlinerMode.h
--- a/Engine/Plugins/Experimental/EditorDataStorageFeatures/Source/TedsOutliner/Public/TedsOutlinerMode.h
+++ b/Engine/Plugins/Experimental/EditorDataStorageFeatures/Source/TedsOutliner/Public/TedsOutlinerMode.h
@@ -51,6 +51,13 @@ protected:
 	// Called by TedsOutlinerImpl when the selection in TEDS changes
 	UE_API void OnSelectionChanged();
 
+	/** Functions to expose default SceneOutliner Options to the UI */
+	void OnToggleAlwaysFrameSelection();
+	bool ShouldAlwaysFrameSelection() const;
+
+	void OnToggleCollapseOutlinerTreeOnNewSelection();
+	bool CollapseOutlinerTreeOnNewSelection() const;
+
 private:
 	/**
 	 * Get a mutable version of the TedsOutliner config for setting values.
@@ -71,6 +78,12 @@ private:
 
 protected:
 	bool bShowViewButton = false;
+	
+	/** Should the outliner scroll to the item on selection */
+	bool bAlwaysFrameSelection;
+	/** Should Outliner Tree get collapsed on selection, except for the item that was just selected */
+	bool bCollapseOutlinerTreeOnNewSelection;
+	
 	// Delegate for extending the View Menu with options (World picker, frame on selection, etc.)
 	FOnExtendViewMenu OnInitializeViewMenuExtender;
 	// The actual model for the TEDS Outliner
