SRowDetails: Fix a bug causing multiple copies of the same column's widget constructor to show up in the results.
Each "item" now keeps track of every column it matched against instead of only the first one, and some duplicate checks were added to fix the bug.

#Synced-CL 52511197

diff --git a/Engine/Plugins/Experimental/EditorDataStorageFeatures/Source/TedsTableViewer/Private/Widgets/SRowDetails.cpp b/Engine/Plugins/Experimental/EditorDataStorageFeatures/Source/TedsTableViewer/Private/Widgets/SRowDetails.cpp
--- a/Engine/Plugins/Experimental/EditorDataStorageFeatures/Source/TedsTableViewer/Private/Widgets/SRowDetails.cpp
+++ b/Engine/Plugins/Experimental/EditorDataStorageFeatures/Source/TedsTableViewer/Private/Widgets/SRowDetails.cpp
@@ -2,6 +2,7 @@
 
 #include "Widgets/SRowDetails.h"
 
+#include "TedsTableViewerUtils.h"
 #include "DataStorage/Features.h"
 #include "Elements/Columns/TypedElementMiscColumns.h"
 #include "Elements/Columns/TypedElementTypeInfoColumns.h"
@@ -106,20 +107,29 @@ namespace UE::Editor::DataStorage
 		FComboMetaDataView<FGenericMetaDataView, FColumnsMetaDataView> MetaDataView =
 			FComboMetaDataView(FGenericMetaDataView(DataMetaData ? *DataMetaData : EmptyMetaData))
 			.Next(FColumnsMetaDataView(ColumnMetaData));
+		
+		int32 IndexOffset = 0;
 
 		// Create the widget constructors for the columns
 		DataStorageUi->CreateWidgetConstructors(PurposeRow, IUiProvider::EMatchApproach::LongestMatch, ColumnsCopy,
-			MetaDataView, [this, Row, &NewColumnMap](
+			MetaDataView, [this, Row, &NewColumnMap, &IndexOffset](
 			IUiProvider::FWidgetConstructorPtr WidgetConstructor, TConstArrayView<TWeakObjectPtr<const UScriptStruct>> MatchedColumns)
 			{
+				if (!bShowUnpopulatedColumns && !DataStorage->HasColumns(Row, MatchedColumns))
+				{
+					return true;
+				}
+					
+				TArray<TWeakObjectPtr<const UScriptStruct>> MatchedColumnsCopy(MatchedColumns);
+				FName NameId = TableViewerUtils::FindLongestMatchingName(MatchedColumns, IndexOffset);
+				TSharedRef<FRowDetailsItem> Item = MakeShared<FRowDetailsItem>(NameId, MoveTemp(MatchedColumnsCopy), WidgetConstructor, Row);
+					
 				for (const TWeakObjectPtr<const UScriptStruct>& Column : MatchedColumns)
 				{
-					if (!bShowUnpopulatedColumns && !DataStorage->HasColumns(Row, {Column}))
-					{
-						continue;
-					}
-					NewColumnMap.Emplace(Column, MakeShared<FRowDetailsItem>(Column, WidgetConstructor, Row));
+					NewColumnMap.Emplace(Column, Item);
 				}
+					
+				++IndexOffset;
 				return true;
 			});
 
@@ -130,10 +140,12 @@ namespace UE::Editor::DataStorage
 				for (TWeakObjectPtr<const UScriptStruct> Column : ColumnsCopy)
 				{
 					DataStorageUi->CreateWidgetConstructors(DataStorageUi->FindPurpose(DataStorageUi->GetDefaultWidgetPurposeID()),
-						MetaDataView, [this, Column, Row, &NewColumnMap](
-							IUiProvider::FWidgetConstructorPtr Constructor, TConstArrayView<TWeakObjectPtr<const UScriptStruct>> Columns)
+						MetaDataView, [this, Column, Row, &NewColumnMap, &IndexOffset](
+							IUiProvider::FWidgetConstructorPtr Constructor, TConstArrayView<TWeakObjectPtr<const UScriptStruct>> MatchedColumns)
 						{
-							NewColumnMap.Emplace(Column, MakeShared<FRowDetailsItem>(Column, MoveTemp(Constructor), Row));
+							FName NameId = TableViewerUtils::FindLongestMatchingName({Column}, IndexOffset);
+							NewColumnMap.Emplace(Column, MakeShared<FRowDetailsItem>(NameId, TArray{Column}, MoveTemp(Constructor), Row));
+							++IndexOffset;
 							return true;
 						});
 				}
@@ -144,8 +156,12 @@ namespace UE::Editor::DataStorage
 		{
 			if (const TSharedRef<FRowDetailsItem>* FoundColumn = NewColumnMap.Find(ColumnType))
 			{
-				Items.Add(*FoundColumn);
-				NewColumnMap.Remove(ColumnType);
+				// If the column already exists, a widget matched it and a previously encountered column together and was already added
+				// so we can safely ignore it here
+				if(!HasItem((*FoundColumn)->Name))
+				{
+					Items.Add(*FoundColumn);
+				}
 			}
 		}
 
@@ -173,7 +189,16 @@ namespace UE::Editor::DataStorage
 
 		return Row;
 	}
-	
+
+	bool SRowDetails::HasItem(const FName& ItemName) const
+	{
+		return Items.FindByPredicate([ItemName]
+			(const RowDetailsItemPtr& InItem)
+		{
+			return InItem->Name == ItemName;
+		}) != nullptr;
+	}
+
 	//
 	// SRowDetailsRow
 	//
@@ -204,12 +229,14 @@ namespace UE::Editor::DataStorage
 					.Row = Item->Row
 				});
 
-			if (Item->ColumnType.IsValid() &&
+			if (Item->ColumnTypes.Num() &&
 				Item->WidgetConstructor->GetAdditionalColumnsList().Contains(FTypedElementScriptStructTypeInfoColumn::StaticStruct()))
 			{
+				// FTypedElementScriptStructTypeInfoColumn can only store one column type but we can match against multiple columns, for now
+				// we just store the first one
 				DataStorage->AddColumn(Item->WidgetRow, FTypedElementScriptStructTypeInfoColumn
 					{
-						.TypeInfo = Item->ColumnType
+						.TypeInfo = Item->ColumnTypes[0]
 					});
 			}
 		}
@@ -222,7 +249,7 @@ namespace UE::Editor::DataStorage
 				.Next(FColumnsMetaDataView(RowDetails->GetColumnMetaData()));
 
 			STextBlock::FArguments Args = STextBlock::FArguments();
-			if (const FString* const* FontStyle = MetaDataView.FindForColumnOrGeneric(Item->ColumnType, "Font").TryGetExact<const FString*>())
+			if (const FString* const* FontStyle = MetaDataView.FindForColumnsOrGeneric(Item->ColumnTypes, "Font").TryGetExact<const FString*>())
 			{
 				Args._Font = FAppStyle::Get().GetFontStyle(FName(**FontStyle));
 			}
@@ -247,9 +274,10 @@ namespace UE::Editor::DataStorage
 	}
 
 	// FRowDetailsItem
-	FRowDetailsItem::FRowDetailsItem(const TWeakObjectPtr<const UScriptStruct>& InColumnType,
+	FRowDetailsItem::FRowDetailsItem(const FName& InItemName, const TArray<TWeakObjectPtr<const UScriptStruct>>& InColumnTypes,
 		IUiProvider::FWidgetConstructorPtr InWidgetConstructor, RowHandle InRow)
-		: ColumnType(InColumnType)
+		: Name(InItemName)
+		, ColumnTypes(InColumnTypes)
 		, WidgetConstructor(MoveTemp(InWidgetConstructor))
 		, Row(InRow)
 	{
diff --git a/Engine/Plugins/Experimental/EditorDataStorageFeatures/Source/TedsTableViewer/Public/Widgets/SRowDetails.h b/Engine/Plugins/Experimental/EditorDataStorageFeatures/Source/TedsTableViewer/Public/Widgets/SRowDetails.h
--- a/Engine/Plugins/Experimental/EditorDataStorageFeatures/Source/TedsTableViewer/Public/Widgets/SRowDetails.h
+++ b/Engine/Plugins/Experimental/EditorDataStorageFeatures/Source/TedsTableViewer/Public/Widgets/SRowDetails.h
@@ -19,8 +19,11 @@ namespace UE::Editor::DataStorage
 	// A row in the SRowDetails widget that represents a column on the TEDS row we are viewing
 	struct FRowDetailsItem
 	{
-		// The column we this row is displaying data for
-		TWeakObjectPtr<const UScriptStruct> ColumnType;
+		// The logical name of this item (composited from the matched columns or grabbed from the widget constructor)
+		FName Name;
+
+		// The columns this row is displaying data for
+		TArray<TWeakObjectPtr<const UScriptStruct>> ColumnTypes;
 
 		// Widget for the column
 		IUiProvider::FWidgetConstructorPtr WidgetConstructor;
@@ -28,7 +31,7 @@ namespace UE::Editor::DataStorage
 		RowHandle Row = InvalidRowHandle;
 		RowHandle WidgetRow = InvalidRowHandle;
 
-		FRowDetailsItem(const TWeakObjectPtr<const UScriptStruct>& InColumnType, IUiProvider::FWidgetConstructorPtr InWidgetConstructor,
+		FRowDetailsItem(const FName& InItemName, const TArray<TWeakObjectPtr<const UScriptStruct>>& InColumnTypes, IUiProvider::FWidgetConstructorPtr InWidgetConstructor,
 			RowHandle InRow);
 	};
 	
@@ -91,6 +94,7 @@ namespace UE::Editor::DataStorage
 	private:
 		
 		UE_API TSharedRef<ITableRow> CreateRow(RowDetailsItemPtr InItem, const TSharedRef<STableViewBase>& OwnerTable);
+		bool HasItem(const FName& ItemName) const;
 
 		TArray<TWeakObjectPtr<const UScriptStruct>> RequestedTedsColumns;
 		TSharedPtr<SListView<RowDetailsItemPtr>> ListView;
