[TEDS Outliner] Add the Folder Type priority in the Label Sorting, Add the empty folder filter, and update existing show filters to run correctly with WP

#Synced-CL 52623505

diff --git a/Engine/Plugins/Experimental/EditorDataStorageFeatures/Source/TedsOutliner/Private/Widgets/OutlinerLabelWidget.cpp b/Engine/Plugins/Experimental/EditorDataStorageFeatures/Source/TedsOutliner/Private/Widgets/OutlinerLabelWidget.cpp
--- a/Engine/Plugins/Experimental/EditorDataStorageFeatures/Source/TedsOutliner/Private/Widgets/OutlinerLabelWidget.cpp
+++ b/Engine/Plugins/Experimental/EditorDataStorageFeatures/Source/TedsOutliner/Private/Widgets/OutlinerLabelWidget.cpp
@@ -2,6 +2,7 @@
 
 #include "Widgets/OutlinerLabelWidget.h"
 
+#include "ActorFolders/ActorFolderColumns.h"
 #include "Columns/SlateHeaderColumns.h"
 #include "Elements/Columns/TypedElementCompatibilityColumns.h"
 #include "Elements/Columns/TypedElementTypeInfoColumns.h"
@@ -69,8 +70,11 @@ namespace UE::Editor::DataStorage
 			return 1;
 		}
 		
-		// TODO: When Folders are added in TEDS, they should return 2
-
+		if(Storage.HasColumns<FFolderCompatibilityColumn>(Row))
+		{
+			return 2;
+		}
+		
 		return 3;
 	}
 }
diff --git a/Engine/Plugins/Experimental/EditorDataStorageFeatures/Source/TedsOutliner/Private/Widgets/SceneOutlinerWidget.cpp b/Engine/Plugins/Experimental/EditorDataStorageFeatures/Source/TedsOutliner/Private/Widgets/SceneOutlinerWidget.cpp
--- a/Engine/Plugins/Experimental/EditorDataStorageFeatures/Source/TedsOutliner/Private/Widgets/SceneOutlinerWidget.cpp
+++ b/Engine/Plugins/Experimental/EditorDataStorageFeatures/Source/TedsOutliner/Private/Widgets/SceneOutlinerWidget.cpp
@@ -21,6 +21,7 @@
 #include "Elements/Framework/TypedElementSelectionSet.h"
 #include "Factories/TedsEditorHierarchyFactory.h"
 #include "LevelInstance/LevelInstanceInterface.h"
+#include "LevelInstance/LevelInstanceSubsystem.h"
 #include "WorldTreeItem.h"
 
 #define LOCTEXT_NAMESPACE "SceneOutlinerWidget"
@@ -181,22 +182,58 @@ void FTedsSceneOutlinerWidgetConstructor::GetFilters(UE::Editor::DataStorage::IC
 void FTedsSceneOutlinerWidgetConstructor::GetShowOptionsFilters(UE::Editor::DataStorage::ICoreProvider* DataStorage, TMap<FName, TSharedPtr<UE::Editor::Outliner::FTedsOutlinerFilter>>& Filters) const
 {
 	using namespace UE::Editor::Outliner;
+
+	auto IsLevelALevelInstance = [](const FTypedElementLevelColumn& LevelColumn, const FTypedElementWorldColumn& WorldColumn)
+	{
+		if (const ULevel* Level = LevelColumn.Level.Get())
+		{
+			if (const UWorld* World = WorldColumn.World.Get())
+			{
+				if (const ULevelInstanceSubsystem* LevelInstanceSubsystem = World->GetSubsystem<ULevelInstanceSubsystem>())
+				{
+					const ILevelInstanceInterface* LevelInstance = LevelInstanceSubsystem->GetOwningLevelInstance(Level);
+					return LevelInstance != nullptr;
+				}
+			}
+		}
+		
+		return false;
+	};
 	
 	// Hide Temporary Actors
 	const FName HideTemporaryName = TEXT("HideTemporary");
 	Filters.Add(HideTemporaryName, MakeShared<FTedsOutlinerFilter>(
 		HideTemporaryName,
-		LOCTEXT("ToggleHideTemporaryActors", "Hide Temporary Actors"),
-			LOCTEXT("ToggleHideTemporaryActorsToolTip", "When enabled, hides temporary/run-time actors."),
+		LOCTEXT("ToggleHideTemporaryObjects", "Hide Temporary Objects"),
+			LOCTEXT("ToggleHideTemporaryObjectsToolTip", "When enabled, hides temporary/run-time objects."),
 		NAME_None,
 		nullptr,
-		[](TConstQueryContext<SingleRowInfo>, const FTypedElementUObjectColumn& ObjectColumn)
+		[DataStorage, IsLevelALevelInstance](TConstQueryContext<SingleRowInfo> Context, const FTypedElementWorldColumn& WorldColumn)
 	 	{
-			if (const UObject* Object = ObjectColumn.Object.Get())
+			if (DataStorage)
 			{
-				return !Object->HasAnyFlags(RF_Transient);
+				if (const FTypedElementUObjectColumn* ObjectColumn = DataStorage->GetColumn<FTypedElementUObjectColumn>(Context.GetCurrentRow()))
+				{
+					if (const UObject* Object = ObjectColumn->Object.Get())
+					{
+						if (Object->HasAnyFlags(RF_Transient))
+						{
+							if (Context.CurrentRowHasColumns<FPieObjectTag>())
+							{
+								return true;
+							}
+
+							if (const FTypedElementLevelColumn* LevelColumn = DataStorage->GetColumn<FTypedElementLevelColumn>(Context.GetCurrentRow()))
+							{
+								return IsLevelALevelInstance(*LevelColumn, WorldColumn);
+							}
+							
+							return false;
+						}
+					}
+				}
 			}
-			return false;
+			return true;
 	 	}));
 
 	// Hide Level Instance Content
@@ -207,32 +244,13 @@ void FTedsSceneOutlinerWidgetConstructor::GetShowOptionsFilters(UE::Editor::Data
 		LOCTEXT("ToggleHideLevelInstancesToolTip", "When enabled, hides all level instance content."),
 		NAME_None,
 		nullptr,
-		[DataStorage](TConstQueryContext<SingleRowInfo> Context, const FTypedElementUObjectColumn& ObjectColumn)
+		[DataStorage, IsLevelALevelInstance](TConstQueryContext<SingleRowInfo> Context, const FTypedElementWorldColumn& WorldColumn)
 		 {
-			if (UObject* Object = ObjectColumn.Object.Get())
-			{
-				// Check if it is a level instance first since the Level Column holds its outer
-				if (const ILevelInstanceInterface* LevelInstance = Cast<ILevelInstanceInterface>(Object))
-				{
-					return false;
-				}
-				if (Context.CurrentRowHasColumns<FEditorDataStorageLevelTag>())
-				{
-					if (const ULevel* Level = Cast<ULevel>(Object))
-					{
-						return !Level->IsInstancedLevel();
-					}
-				}
-			}
-
 			if (DataStorage && Context.CurrentRowHasColumns<FTypedElementLevelColumn>())
 			{
 				if (const FTypedElementLevelColumn* LevelColumn = DataStorage->GetColumn<FTypedElementLevelColumn>(Context.GetCurrentRow()))
 				{
-					if (const ULevel* Level = LevelColumn->Level.Get())
-					{
-						return !Level->IsInstancedLevel();
-					}
+					return !IsLevelALevelInstance(*LevelColumn, WorldColumn);
 				}
 			}
 				
@@ -262,29 +280,23 @@ void FTedsSceneOutlinerWidgetConstructor::GetShowOptionsFilters(UE::Editor::Data
 		LOCTEXT("ToggleShowOnlyCurrentLevelToolTip", "When enabled, only shows objects that are in the Current Level."),
 		NAME_None,
 		nullptr,
-		[DataStorage](TConstQueryContext<SingleRowInfo> Context, const FTypedElementUObjectColumn& ObjectColumn)
+		[DataStorage](TConstQueryContext<SingleRowInfo> Context)
 		 {
-			if (UObject* Object = ObjectColumn.Object.Get())
+			if (DataStorage)
 			{
-				// Check if it is a level instance first since the Level Column holds its outer
-				if (const ILevelInstanceInterface* LevelInstance = Cast<ILevelInstanceInterface>(Object))
-				{
-					if (const ULevel* Level = LevelInstance->GetLoadedLevel())
-					{
-						return Level->IsCurrentLevel();
-					}
-				}
 				if (Context.CurrentRowHasColumns<FEditorDataStorageLevelTag>())
 				{
-					if (const ULevel* Level = Cast<ULevel>(Object))
+					if (const FTypedElementUObjectColumn* ObjectColumn = DataStorage->GetColumn<FTypedElementUObjectColumn>(Context.GetCurrentRow()))
 					{
-						return Level->IsCurrentLevel();
+						if (UObject* Object = ObjectColumn->Object.Get())
+						{
+							if (const ULevel* Level = Cast<ULevel>(Object))
+							{
+								return Level->IsCurrentLevel();
+							}
+						}
 					}
 				}
-			}
-
-			if (DataStorage && Context.CurrentRowHasColumns<FTypedElementLevelColumn>())
-			{
 				if (const FTypedElementLevelColumn* LevelColumn = DataStorage->GetColumn<FTypedElementLevelColumn>(Context.GetCurrentRow()))
 				{
 					if (const ULevel* Level = LevelColumn->Level.Get())
@@ -295,7 +307,7 @@ void FTedsSceneOutlinerWidgetConstructor::GetShowOptionsFilters(UE::Editor::Data
 			}
 
 			return false;
-		 }));
+		}));
 
 	// Only in any Current Data Layers
 	// TODO When in TEDS
@@ -324,7 +336,20 @@ void FTedsSceneOutlinerWidgetConstructor::GetShowOptionsFilters(UE::Editor::Data
 	// TODO When in TEDS
 
 	// Hide Empty Folders
-	// TODO When in TEDS
+	const FName HideEmptyFoldersFilterName = TEXT("HideEmptyFoldersFilter");
+	const FHierarchyHandle Hierarchy = DataStorage->FindHierarchyByName(UTedsEditorHierarchyFactory::EditorObjectHierarchyName);
+	const UScriptStruct* FolderParentTag = DataStorage->GetParentTagType(Hierarchy);
+	
+	Filters.Add(HideEmptyFoldersFilterName, MakeShared<FTedsOutlinerFilter>(
+		HideEmptyFoldersFilterName,
+		LOCTEXT("ToggleHideEmptyFolders", "Hide Empty Folders"),
+		LOCTEXT("ToggleHideEmptyFoldersToolTip", "When enabled, hides all empty folders."),
+		NAME_None,
+		nullptr,
+		[FolderParentTag](TConstQueryContext<SingleRowInfo> Context)
+		 {
+			return !Context.CurrentRowHasColumns<FFolderCompatibilityColumn>() || Context.CurrentRowHasColumns({FolderParentTag});
+		 }));
 }
 
 void FTedsSceneOutlinerWidgetConstructor::GetInitializeViewMenuExtender(UE::Editor::Outliner::FTedsOutlinerParams& OutlinerParams)
