[TEDS Outliner] Allow the outliner to show the Level/Package Short Name (if in a external package) and replace the header label if external packages are enabled

#Synced-CL 52678418

diff --git a/Engine/Plugins/Experimental/EditorDataStorageFeatures/Source/TedsOutliner/Private/TedsOutlinerImpl.cpp b/Engine/Plugins/Experimental/EditorDataStorageFeatures/Source/TedsOutliner/Private/TedsOutlinerImpl.cpp
--- a/Engine/Plugins/Experimental/EditorDataStorageFeatures/Source/TedsOutliner/Private/TedsOutlinerImpl.cpp
+++ b/Engine/Plugins/Experimental/EditorDataStorageFeatures/Source/TedsOutliner/Private/TedsOutlinerImpl.cpp
@@ -274,23 +274,6 @@ FTedsOutlinerImpl::FTedsOutlinerImpl(const FTedsOutlinerParams& InParams, IScene
 	Storage = GetMutableDataStorageFeature<DataStorage::ICoreProvider>(StorageFeatureName);
 	StorageUi = GetMutableDataStorageFeature<DataStorage::IUiProvider>(UiFeatureName);
 	StorageCompatibility = GetMutableDataStorageFeature<DataStorage::ICompatibilityProvider>(CompatibilityFeatureName);
-
-	// Only do this on init to match parity since the legacy outliner only checks for this when opened
-	if (GEditor) // No need to check if ->GetEditorWorldContext is valid as it should always return a valid UWorld reference
-	{
-		if(const UWorld* World = GEditor->GetEditorWorldContext().World())
-		{
-			if (World->PersistentLevel && World->PersistentLevel->IsUsingExternalActors())
-			{
-				// TODO: When the Package Short Name data is available, have it reassign to that column instead of removing it
-				Helpers::RemoveTedsColumnMappingToOutlinerColumn(Storage, FLevelColumn::StaticStruct());
-			}
-			else
-			{
-				Helpers::AddTedsColumnMappingToOutlinerColumn(Storage, FLevelColumn::StaticStruct(), FSceneOutlinerBuiltInColumnTypes::Level());
-			}
-		}
-	}
 	
 	bUsingQueryConditionsSyntax = InitialQueryDescription.Conditions && !InitialQueryDescription.Conditions->IsEmpty();
 }
@@ -721,8 +704,11 @@ void FTedsOutlinerImpl::RegenerateColumns(const Outliner::FTedsOutlinerColumnDes
 									.OnSortHandler = FTedsTableViewerColumn::FOnSort::CreateSP(this, &FTedsOutlinerImpl::OnSort),
 									.IsSortingHandler = FTedsTableViewerColumn::FIsSorting::CreateSP(this, &FTedsOutlinerImpl::IsSorting)
 								};
-	
-								FTedsOutlinerUiColumnInitParams Params(ColumnDescription.GetGenericMetaData(), *Storage, *StorageUi, *StorageCompatibility);
+
+								FMetaData TedsOutlinerUiColumnMetaData = ColumnDescription.GetGenericMetaData();
+								TedsOutlinerUiColumnMetaData.AddImmutableData(TEXT("OutlinerIdentifier"), SceneOutliner->GetOutlinerIdentifier().ToString());
+									
+								FTedsOutlinerUiColumnInitParams Params(TedsOutlinerUiColumnMetaData, *Storage, *StorageUi, *StorageCompatibility);
 								Params.NameId = LabelColumnName;
 								Params.ColumnTypes = TArray<TWeakObjectPtr<const UScriptStruct>>(ColumnTypes.GetData(), ColumnTypes.Num());
 								Params.HeaderWidgetConstructor = MoveTemp(HeaderConstructor);
@@ -795,7 +781,7 @@ void FTedsOutlinerImpl::RegenerateColumns(const Outliner::FTedsOutlinerColumnDes
 			SceneOutliner->RemoveColumn(FallbackColumnName);
 
 			FName NameId = TableViewerUtils::FindLongestMatchingName(ColumnTypes, IndexOffset);
-			const FText DisplayName = CellConstructor.Get()->CreateWidgetDisplayNameText(Storage);
+			const FText DisplayName = CellConstructor.Get()->CreateWidgetDisplayNameText(Storage, GetOutlinerRowHandle());
 			AddedColumns.Add(NameId);
 			SceneOutliner->AddColumn(NameId,
 				FSceneOutlinerColumnInfo(
@@ -813,7 +799,10 @@ void FTedsOutlinerImpl::RegenerateColumns(const Outliner::FTedsOutlinerColumnDes
 								.IsSortingHandler = DataStorage::FTedsTableViewerColumn::FIsSorting::CreateSP(this, &FTedsOutlinerImpl::IsSorting)
 							};
 
-							FTedsOutlinerUiColumnInitParams Params(ColumnDescription.GetGenericMetaData(), *Storage, *StorageUi, *StorageCompatibility);
+							FMetaData TedsOutlinerUiColumnMetaData = ColumnDescription.GetGenericMetaData();
+							TedsOutlinerUiColumnMetaData.AddImmutableData(TEXT("OutlinerIdentifier"), SceneOutliner->GetOutlinerIdentifier().ToString());
+								
+							FTedsOutlinerUiColumnInitParams Params(TedsOutlinerUiColumnMetaData, *Storage, *StorageUi, *StorageCompatibility);
 							Params.NameId = NameId;
 							Params.ColumnTypes = TArray<TWeakObjectPtr<const UScriptStruct>>(ColumnTypes.GetData(), ColumnTypes.Num());
 							Params.HeaderWidgetConstructor = MoveTemp(HeaderConstructor);
diff --git a/Engine/Plugins/Experimental/EditorDataStorageFeatures/Source/TedsOutliner/Private/TedsOutlinerModule.cpp b/Engine/Plugins/Experimental/EditorDataStorageFeatures/Source/TedsOutliner/Private/TedsOutlinerModule.cpp
--- a/Engine/Plugins/Experimental/EditorDataStorageFeatures/Source/TedsOutliner/Private/TedsOutlinerModule.cpp
+++ b/Engine/Plugins/Experimental/EditorDataStorageFeatures/Source/TedsOutliner/Private/TedsOutlinerModule.cpp
@@ -303,6 +303,11 @@ void FTedsSceneOutlinerWorldFilter::BuildWorldFilter()
 
 			if (const TSharedPtr<FTedsOutlinerImpl> TedsOutlinerImplPin = TedsOutlinerImpl.Pin())
 			{
+				if (ICoreProvider* Storage = TedsOutlinerImplPin->GetStorage())
+				{
+					Storage->AddColumn(TedsOutlinerImplPin->GetOutlinerRowHandle(), FTypedElementWorldColumn{ .World = RepresentingWorld });
+				}
+				
 				// Perform a FullRefresh here since the Outliner does this when Filters change but we have to do it manually 
 				// since this filter isn't managed by the FilterBar.
 				TedsOutlinerImplPin->FullRefresh();
@@ -342,6 +347,11 @@ TWeakObjectPtr<UWorld> FTedsSceneOutlinerWorldFilter::GetUserChosenWorld() const
 {
 	return UserChosenWorld;
 }
+
+	TWeakObjectPtr<UWorld> FTedsSceneOutlinerWorldFilter::GetRepresentingWorld() const
+{
+	return RepresentingWorld;
+}
 	
 } // namespace UE::Editor::Outliner
 
diff --git a/Engine/Plugins/Experimental/EditorDataStorageFeatures/Source/TedsOutliner/Private/Widgets/OutlinerLevelWidget.cpp b/Engine/Plugins/Experimental/EditorDataStorageFeatures/Source/TedsOutliner/Private/Widgets/OutlinerLevelWidget.cpp
--- a/Engine/Plugins/Experimental/EditorDataStorageFeatures/Source/TedsOutliner/Private/Widgets/OutlinerLevelWidget.cpp
+++ b/Engine/Plugins/Experimental/EditorDataStorageFeatures/Source/TedsOutliner/Private/Widgets/OutlinerLevelWidget.cpp
@@ -3,38 +3,72 @@
 #include "Widgets/OutlinerLevelWidget.h"
 
 #include "SceneOutlinerHelpers.h"
+#include "Columns/TedsOutlinerColumns.h"
 #include "Elements/Columns/TypedElementCompatibilityColumns.h"
 #include "Elements/Columns/TypedElementMiscColumns.h"
+#include "Elements/Columns/TypedElementPackageColumns.h"
 #include "Elements/Framework/TypedElementAttributeBinding.h"
 #include "Elements/Interfaces/TypedElementDataStorageUiInterface.h"
 #include "Engine/Level.h"
+#include "Misc/PackageName.h"
 #include "Widgets/Text/STextBlock.h"
 
 #define LOCTEXT_NAMESPACE "OutlinerLevelWidget"
 
 namespace UE::Editor::DataStorage
 {
+	// Helper to extract the Short name from the package row
+	FString GetLevelName(const ICoreProvider& Storage, RowHandle Row)
+	{
+		if (const FTypedElementPackageReference* PackageRefColumn = Storage.GetColumn<FTypedElementPackageReference>(Row))
+		{
+			if (const FTypedElementPackagePathColumn* PackagePathColumn = Storage.GetColumn<FTypedElementPackagePathColumn>(PackageRefColumn->Row))
+			{
+				return FPackageName::GetShortName(PackagePathColumn->Path);
+			}
+		}
+		return FString();
+	}
+	
 	FText FLevelWidgetSorter::GetShortName() const
 	{
 		return LOCTEXT("LevelWidgetSorter", "Level Sorter");
 	}
+    
+	FColumnSorterInterface::ESortType FLevelWidgetSorter::GetSortType() const
+	{
+		return ESortType::HybridSort;
+	}
 
-	int32 FLevelWidgetSorter::Compare(const FLevelColumn& Left, const FLevelColumn& Right) const
+	int32 FLevelWidgetSorter::Compare(const ICoreProvider& Storage, RowHandle Left, RowHandle Right) const
 	{
-		return Left.Level->GetOuter()->GetName().Compare(Right.Level->GetOuter()->GetName());
+		return GetLevelName(Storage, Left).Compare(GetLevelName(Storage, Right));
 	}
 
-	FPrefixInfo FLevelWidgetSorter::CalculatePrefix(const FLevelColumn& Column, uint32 ByteIndex) const
+	FPrefixInfo FLevelWidgetSorter::CalculatePrefix(const ICoreProvider& Storage, RowHandle Row, uint32 ByteIndex) const
 	{
-		return CreateSortPrefix(ByteIndex, TSortStringView(FSortCaseSensitive{}, Column.Level->GetOuter()->GetName()));
+		return CreateSortPrefix(ByteIndex, TSortStringView(FSortCaseSensitive{}, GetLevelName(Storage, Row)));
 	}
 	
-	void FLevelWidgetSearcher::GetSearchableString(FString& SearchableString, const ICoreProvider& Storage, const FLevelColumn& Column) const
+	void FLevelWidgetSearcher::GetSearchableString(FString& SearchableString, const ICoreProvider& Storage, RowHandle Row) const
+	{
+		SearchableString = GetLevelName(Storage, Row);
+	}
+
+	FText GetOutlinerLevelDisplayText(ICoreProvider* DataStorage, RowHandle Row)
 	{
-		if (const TWeakObjectPtr<ULevel> Level = Column.Level; Level.IsValid())
+		// The Outliner row will be passed in as the Row, see if there is a world assigned to it
+		if (const FTypedElementWorldColumn* WorldColumn = DataStorage->GetColumn<FTypedElementWorldColumn>(Row))
 		{
-			SearchableString = Level->GetOuter()->GetName();
+			if (const TWeakObjectPtr<UWorld> World = WorldColumn->World; World.IsValid())
+			{
+				if (World->PersistentLevel && World->PersistentLevel->IsUsingExternalActors())
+				{
+					return LOCTEXT("OutlinerPackageShortNameColumnDisplayName", "Package Short Name");
+				}
+			}
 		}
+		return LOCTEXT("OutlinerLevelColumnDisplayName", "Level");
 	}
 }
 
@@ -46,9 +80,36 @@ void UOutlinerLevelWidgetFactory::RegisterWidgetConstructors(UE::Editor::DataSto
 	DataStorageUi.RegisterWidgetFactory<FOutlinerLevelWidgetConstructor>(
 		DataStorageUi.FindPurpose(UE::Editor::DataStorage::IUiProvider::FPurposeInfo("SceneOutliner", "Cell", NAME_None).GeneratePurposeID()),
 		TColumn<FLevelColumn>());
-	DataStorageUi.RegisterWidgetFactory<FOutlinerLevelWidgetConstructor>(
-			DataStorageUi.FindPurpose(DataStorageUi.GetGeneralWidgetPurposeID()),
-			TColumn<FLevelColumn>());
+
+	DataStorageUi.RegisterWidgetFactory<FOutlinerLevelHeaderWidgetConstructor>(
+		DataStorageUi.FindPurpose(UE::Editor::DataStorage::IUiProvider::FPurposeInfo("SceneOutliner", "Header", NAME_None).GeneratePurposeID()),
+		TColumn<FLevelColumn>());
+}
+
+FOutlinerLevelHeaderWidgetConstructor::FOutlinerLevelHeaderWidgetConstructor()
+	: FSimpleWidgetConstructor(StaticStruct())
+{
+}
+
+TSharedPtr<SWidget> FOutlinerLevelHeaderWidgetConstructor::CreateWidget(
+	UE::Editor::DataStorage::ICoreProvider* DataStorage,
+	UE::Editor::DataStorage::IUiProvider* DataStorageUi,
+	UE::Editor::DataStorage::RowHandle TargetRow,
+	UE::Editor::DataStorage::RowHandle WidgetRow,
+	const UE::Editor::DataStorage::FMetaDataView& Arguments)
+{
+	using namespace UE::Editor::DataStorage;
+	
+	FName TabIdentifier;
+	const FMetaDataEntryView OutlinerIdMeta = Arguments.FindGeneric("OutlinerIdentifier");
+	if(const FString* const* IdMetaData = OutlinerIdMeta.TryGetExact<const FString*>())
+	{
+		TabIdentifier = FName(**IdMetaData);
+	}
+	const RowHandle OutlinerRow = DataStorage->LookupMappedRow(UE::Editor::Outliner::MappingDomain, FMapKey(TabIdentifier));
+
+	return SNew(STextBlock)
+		.Text(GetOutlinerLevelDisplayText(DataStorage, OutlinerRow));
 }
 
 FOutlinerLevelWidgetConstructor::FOutlinerLevelWidgetConstructor()
@@ -80,20 +141,30 @@ TSharedPtr<SWidget> FOutlinerLevelWidgetConstructor::CreateWidget(
 		[
 			SNew(STextBlock)
 				.ColorAndOpacity(FSlateColor::UseSubduedForeground())
-				.Text(Binder.BindData(&FLevelColumn::Level, [](TWeakObjectPtr<ULevel> InLevel)
+				.Text(Binder.BindData(&FTypedElementPackageReference::Row, [TargetRow, DataStorage](RowHandle PackageRow)
 				{
-					if (InLevel.IsValid())
+					if (const FTypedElementPackagePathColumn* PackagePathColumn = DataStorage->GetColumn<FTypedElementPackagePathColumn>(PackageRow))
 					{
-						return FText::FromString(InLevel->GetOuter()->GetName());
+						return FText::FromString(FPackageName::GetShortName(PackagePathColumn->Path));
 					}
-					else
+					if (const FLevelColumn* LevelColumn = DataStorage->GetColumn<FLevelColumn>(TargetRow))
 					{
-						return LOCTEXT("OutlinerLevelUnknown", "<unknown>");
+						if (const TWeakObjectPtr<ULevel> Level = LevelColumn->Level; Level.IsValid())
+						{
+							return FText::FromString(Level->GetOuter()->GetName());
+						}
 					}
+
+					return LOCTEXT("OutlinerLevelUnknown", "<unknown>");
 				}))
 		];
 }
 
+FText FOutlinerLevelWidgetConstructor::CreateWidgetDisplayNameText(UE::Editor::DataStorage::ICoreProvider* DataStorage, UE::Editor::DataStorage::RowHandle Row) const
+{
+	return GetOutlinerLevelDisplayText(DataStorage, Row);
+}
+
 TArray<TSharedPtr<const UE::Editor::DataStorage::FColumnSorterInterface>> FOutlinerLevelWidgetConstructor::ConstructColumnSorters(
 	UE::Editor::DataStorage::ICoreProvider* DataStorage,
 	UE::Editor::DataStorage::IUiProvider* DataStorageUi,
diff --git a/Engine/Plugins/Experimental/EditorDataStorageFeatures/Source/TedsOutliner/Public/TedsOutlinerModule.h b/Engine/Plugins/Experimental/EditorDataStorageFeatures/Source/TedsOutliner/Public/TedsOutlinerModule.h
--- a/Engine/Plugins/Experimental/EditorDataStorageFeatures/Source/TedsOutliner/Public/TedsOutlinerModule.h
+++ b/Engine/Plugins/Experimental/EditorDataStorageFeatures/Source/TedsOutliner/Public/TedsOutlinerModule.h
@@ -84,6 +84,8 @@ public:
 	void QueryFunction(TConstQueryContext<RowBatchInfo> Context, TResult<bool>& Result, TConstBatch<FTypedElementWorldColumn> WorldColumns);
 
 	void SetUserChosenWorld(const TWeakObjectPtr<UWorld>& InWorld);
+	
+	TWeakObjectPtr<UWorld> GetRepresentingWorld() const;
 
 	TWeakObjectPtr<UWorld> GetUserChosenWorld() const;
 	
diff --git a/Engine/Plugins/Experimental/EditorDataStorageFeatures/Source/TedsOutliner/Public/Widgets/OutlinerLevelWidget.h b/Engine/Plugins/Experimental/EditorDataStorageFeatures/Source/TedsOutliner/Public/Widgets/OutlinerLevelWidget.h
--- a/Engine/Plugins/Experimental/EditorDataStorageFeatures/Source/TedsOutliner/Public/Widgets/OutlinerLevelWidget.h
+++ b/Engine/Plugins/Experimental/EditorDataStorageFeatures/Source/TedsOutliner/Public/Widgets/OutlinerLevelWidget.h
@@ -13,19 +13,19 @@
 
 namespace UE::Editor::DataStorage
 {
-	class FLevelWidgetSorter final : public TColumnSorterInterface<FColumnSorterInterface::ESortType::HybridSort, FLevelColumn>
+	class FLevelWidgetSorter final : public FColumnSorterInterface
 	{
 	public:
+		virtual int32 Compare(const ICoreProvider& Storage, RowHandle Left, RowHandle Right) const override;
+		virtual FPrefixInfo CalculatePrefix(const ICoreProvider& Storage, RowHandle Row, uint32 ByteIndex) const override;
+		virtual ESortType GetSortType() const override;
 		virtual FText GetShortName() const override;
-	protected:
-		virtual int32 Compare(const FLevelColumn& Left, const FLevelColumn& Right) const override;
-		virtual FPrefixInfo CalculatePrefix(const FLevelColumn& Column, uint32 ByteIndex) const override;
 	};
 	
-	class FLevelWidgetSearcher final : public TColumnSearcherInterface<FLevelColumn>
+	class FLevelWidgetSearcher final : public FColumnSearcherInterface
     {
     protected:
-    	virtual void GetSearchableString(FString& SearchableString, const ICoreProvider& Storage, const FLevelColumn& Column) const override;
+    	virtual void GetSearchableString(FString& SearchableString, const ICoreProvider& Storage, RowHandle Row) const override;
     };
 }
 
@@ -41,6 +41,23 @@ public:
 		UE::Editor::DataStorage::IUiProvider& DataStorageUi) const override;
 };
 
+USTRUCT()
+struct FOutlinerLevelHeaderWidgetConstructor : public FSimpleWidgetConstructor
+{
+	GENERATED_BODY()
+
+public:
+	UE_API FOutlinerLevelHeaderWidgetConstructor();
+	~FOutlinerLevelHeaderWidgetConstructor() override = default;
+
+	UE_API virtual TSharedPtr<SWidget> CreateWidget(
+		UE::Editor::DataStorage::ICoreProvider* DataStorage,
+		UE::Editor::DataStorage::IUiProvider* DataStorageUi,
+		UE::Editor::DataStorage::RowHandle TargetRow,
+		UE::Editor::DataStorage::RowHandle WidgetRow,
+		const UE::Editor::DataStorage::FMetaDataView& Arguments) override;
+};
+
 USTRUCT()
 struct FOutlinerLevelWidgetConstructor : public FSimpleWidgetConstructor
 {
@@ -56,6 +73,10 @@ public:
 		UE::Editor::DataStorage::RowHandle TargetRow,
 		UE::Editor::DataStorage::RowHandle WidgetRow,
 		const UE::Editor::DataStorage::FMetaDataView& Arguments) override;
+
+	UE_API virtual FText CreateWidgetDisplayNameText(
+		UE::Editor::DataStorage::ICoreProvider* DataStorage,
+		UE::Editor::DataStorage::RowHandle Row = UE::Editor::DataStorage::InvalidRowHandle) const override;
 		
 	UE_API virtual TArray<TSharedPtr<const UE::Editor::DataStorage::FColumnSorterInterface>> ConstructColumnSorters(
 		UE::Editor::DataStorage::ICoreProvider* DataStorage,
